Add the ObservableArray container to slate.

diff --git a/Engine/Source/Runtime/SlateCore/Private/Test/ObservableArrayTest.cpp b/Engine/Source/Runtime/SlateCore/Private/Test/ObservableArrayTest.cpp
new file mode 100644
--- /dev/null
+++ b/Engine/Source/Runtime/SlateCore/Private/Test/ObservableArrayTest.cpp
@@ -0,0 +1,281 @@
+// Copyright Epic Games, Inc. All Rights Reserved.
+
+#include "CoreMinimal.h"
+#include "Containers/ObservableArray.h"
+
+#include "Misc/AutomationTest.h"
+
+#if WITH_DEV_AUTOMATION_TESTS
+
+namespace UE::Test
+{
+	
+IMPLEMENT_SIMPLE_AUTOMATION_TEST(FObservableArrayTest, "Slate.ObservableArray", EAutomationTestFlags::EditorContext | EAutomationTestFlags::EngineFilter)
+
+/**
+ *
+ */
+bool FObservableArrayTest::RunTest(const FString& Parameters)
+{
+	::UE::Slate::Containers::EObservableArrayChangedAction ExpectedAction = ::UE::Slate::Containers::EObservableArrayChangedAction::Add;
+	int32 ExpectedCounter = 0;
+	int32 Counter = 0;
+	TArray<int32> ExpectedArray;
+	FObservableArrayTest* Self = this;
+	auto TestExpectedAction = [Self , &ExpectedAction, &Counter, &ExpectedArray](::UE::Slate::Containers::TObservableArray<int32>::ObservableArrayChangedArgsType Args)
+	{
+		Self->AddErrorIfFalse(Args.GetAction() == ExpectedAction, TEXT("The notification occurs with the wrong action"));
+		if (Args.GetAction() == Slate::Containers::EObservableArrayChangedAction::Add)
+		{
+			Self->AddErrorIfFalse(ExpectedArray == Args.GetItems(), TEXT("The notification occurs with the GetItems"));
+		}
+		else if (Args.GetAction() == Slate::Containers::EObservableArrayChangedAction::Remove)
+		{
+			Self->AddErrorIfFalse(ExpectedArray == Args.GetItems(), TEXT("The notification occurs with the wrong GetItmes"));
+		}
+		++Counter;
+	};
+
+	{
+		{
+			TArray<int32> ArrayValues;
+			::UE::Slate::Containers::TObservableArray<int32> ObservableValues;
+			ObservableValues.OnArrayChanged().AddLambda(TestExpectedAction);
+			{
+				ExpectedAction = ::UE::Slate::Containers::EObservableArrayChangedAction::Add;
+
+				{
+					ExpectedArray.Reset();
+					ExpectedArray.Add(2);
+
+					ArrayValues.Add(2);
+					ObservableValues.Add(2);
+
+					AddErrorIfFalse(ArrayValues == ObservableValues, TEXT("ObservableValues == ArrayValues"));
+					++ExpectedCounter;
+					AddErrorIfFalse(Counter == ExpectedCounter, "The delegate was not executed");
+				}
+				{
+					ExpectedArray.Reset();
+					ExpectedArray.Add(3);
+
+					ArrayValues.Emplace(3);
+					ObservableValues.Emplace(3);
+
+					AddErrorIfFalse(ArrayValues == ObservableValues, TEXT("ObservableValues == ArrayValues"));
+					++ExpectedCounter;
+					AddErrorIfFalse(Counter == ExpectedCounter, "The delegate was not executed");
+				}
+				{
+					ExpectedArray.Reset();
+					ExpectedArray.Add(4);
+
+					ArrayValues.EmplaceAt(2, 4);
+					ObservableValues.EmplaceAt(2, 4);
+
+					AddErrorIfFalse(ArrayValues == ObservableValues, TEXT("ObservableValues == ArrayValues"));
+					++ExpectedCounter;
+					AddErrorIfFalse(Counter == ExpectedCounter, "The delegate was not executed");
+				}
+				{
+					ExpectedArray.Reset();
+					ExpectedArray.Add(1);
+
+					ArrayValues.EmplaceAt(0, 1);
+					ObservableValues.EmplaceAt(0, 1);
+
+					AddErrorIfFalse(ArrayValues == ObservableValues, TEXT("ObservableValues == ArrayValues"));
+					++ExpectedCounter;
+					AddErrorIfFalse(Counter == ExpectedCounter, "The delegate was not executed");
+				}
+				{
+					ExpectedArray.Reset();
+					ExpectedArray.Add(5);
+
+					ArrayValues.Add(5);
+					ObservableValues.Add(5);
+
+					AddErrorIfFalse(ArrayValues == ObservableValues, TEXT("ObservableValues == ArrayValues"));
+					++ExpectedCounter;
+					AddErrorIfFalse(Counter == ExpectedCounter, "The delegate was not executed");
+				}
+				{
+					ExpectedArray.Reset();
+					ExpectedArray.Append({ 6, 7, 8, 9, 10, 11, 12 });
+
+					ArrayValues.Append(ExpectedArray);
+					ObservableValues.Append(ExpectedArray);
+
+					AddErrorIfFalse(ArrayValues == ObservableValues, TEXT("ObservableValues == ArrayValues"));
+					++ExpectedCounter;
+					AddErrorIfFalse(Counter == ExpectedCounter, "The delegate was not executed");
+				}
+				{
+					ExpectedArray.Reset();
+					ExpectedArray.Add(9);
+
+					ArrayValues.Append(ExpectedArray);
+					ObservableValues.Append(ExpectedArray);
+
+					AddErrorIfFalse(ArrayValues == ObservableValues, TEXT("ObservableValues == ArrayValues"));
+					++ExpectedCounter;
+					AddErrorIfFalse(Counter == ExpectedCounter, "The delegate was not executed");
+				}
+			}
+			{
+				ExpectedAction = ::UE::Slate::Containers::EObservableArrayChangedAction::Swap;
+				ExpectedArray.Reset();
+
+				ArrayValues.Swap(1, 2);
+				ObservableValues.Swap(1, 2);
+				AddErrorIfFalse(ArrayValues == ObservableValues, TEXT("ObservableValues == ArrayValues"));
+				++ExpectedCounter;
+				AddErrorIfFalse(Counter == ExpectedCounter, "The delegate was not executed");
+			}
+			{
+				ExpectedAction = ::UE::Slate::Containers::EObservableArrayChangedAction::Remove;
+
+				{
+					ExpectedArray.Reset();
+					ExpectedArray.Add(2);
+
+					ArrayValues.RemoveSingle(2);
+					ObservableValues.RemoveSingle(2);
+
+					AddErrorIfFalse(ArrayValues == ObservableValues, TEXT("ObservableValues == ArrayValues"));
+					++ExpectedCounter;
+					AddErrorIfFalse(Counter == ExpectedCounter, "The delegate was not executed");
+				}
+				{
+					bool bContains = ObservableValues.Contains(2);
+					if (bContains)
+					{
+						ExpectedArray.Reset();
+						ExpectedArray.Add(2);
+					}
+					ArrayValues.RemoveSingle(2);
+					bool bWasRemoved = ObservableValues.RemoveSingle(2) > 0;
+					if (bWasRemoved)
+					{
+						++ExpectedCounter;
+					}
+					AddErrorIfFalse(bWasRemoved == bContains, TEXT("ObservableValues.Contains == ObservableValues.RemoveSingle"));
+					AddErrorIfFalse(ArrayValues == ObservableValues, TEXT("ObservableValues == ArrayValues"));
+					AddErrorIfFalse(Counter == ExpectedCounter, "The delegate was not executed");
+				}
+				{
+					ExpectedArray.Reset();
+					ExpectedArray.Add(3);
+
+					ArrayValues.RemoveSingleSwap(3);
+					ObservableValues.RemoveSingleSwap(3);
+
+					AddErrorIfFalse(ArrayValues == ObservableValues, TEXT("ObservableValues == ArrayValues"));
+					++ExpectedCounter;
+					AddErrorIfFalse(Counter == ExpectedCounter, "The delegate was not executed");
+				}
+				{
+					ExpectedArray.Reset();
+					ExpectedArray.Add(ObservableValues[2]);
+
+					ArrayValues.RemoveAt(2);
+					ObservableValues.RemoveAt(2);
+
+					AddErrorIfFalse(ArrayValues == ObservableValues, TEXT("ObservableValues == ArrayValues"));
+					++ExpectedCounter;
+					AddErrorIfFalse(Counter == ExpectedCounter, "The delegate was not executed");
+				}
+				{
+					ExpectedArray.Reset();
+					int32 ToRemoveIndex = 2;
+					int32 NumberToRemove = 3;
+					for (int32 Index = ToRemoveIndex; Index < NumberToRemove + ToRemoveIndex; ++Index)
+					{
+						if (ObservableValues.IsValidIndex(Index))
+						{
+							ExpectedArray.Add(ObservableValues[Index]);
+						}
+					}
+
+					ArrayValues.RemoveAt(ToRemoveIndex, NumberToRemove);
+					ObservableValues.RemoveAt(ToRemoveIndex, NumberToRemove);
+
+					AddErrorIfFalse(ArrayValues == ObservableValues, TEXT("ObservableValues == ArrayValues"));
+					++ExpectedCounter;
+					AddErrorIfFalse(Counter == ExpectedCounter, "The delegate was not executed");
+				}
+				{
+					ExpectedArray.Reset();
+					ExpectedArray.Add(ObservableValues[1]);
+
+					ArrayValues.RemoveAtSwap(1);
+					ObservableValues.RemoveAtSwap(1);
+
+					AddErrorIfFalse(ArrayValues == ObservableValues, TEXT("ObservableValues == ArrayValues"));
+					++ExpectedCounter;
+					AddErrorIfFalse(Counter == ExpectedCounter, "The delegate was not executed");
+				}
+				{
+					ExpectedArray.Reset();
+					ExpectedArray.Add(ObservableValues[0]);
+
+					ArrayValues.RemoveAtSwap(0);
+					ObservableValues.RemoveAtSwap(0);
+
+					AddErrorIfFalse(ArrayValues == ObservableValues, TEXT("ObservableValues == ArrayValues"));
+					++ExpectedCounter;
+					AddErrorIfFalse(Counter == ExpectedCounter, "The delegate was not executed");
+				}
+				{
+					ExpectedArray.Reset();
+					int32 ToRemoveIndex = 1;
+					int32 NumberToRemove = 4;
+					for (int32 Index = ToRemoveIndex; Index < ToRemoveIndex + NumberToRemove; ++Index)
+					{
+						if (ObservableValues.IsValidIndex(Index))
+						{
+							ExpectedArray.Add(ObservableValues[Index]);
+						}
+					}
+
+					ArrayValues.RemoveAtSwap(ToRemoveIndex, NumberToRemove);
+					ObservableValues.RemoveAtSwap(ToRemoveIndex, NumberToRemove);
+
+					AddErrorIfFalse(ArrayValues == ObservableValues, TEXT("ObservableValues == ArrayValues"));
+					++ExpectedCounter;
+					AddErrorIfFalse(Counter == ExpectedCounter, "The delegate was not executed");
+				}
+			}
+			{
+				ExpectedAction = ::UE::Slate::Containers::EObservableArrayChangedAction::Reset;
+
+				ArrayValues.Reset();
+				ObservableValues.Reset();
+				AddErrorIfFalse(ArrayValues == ObservableValues, TEXT("ObservableValues == ArrayValues"));
+				++ExpectedCounter;
+				AddErrorIfFalse(Counter == ExpectedCounter, "The delegate was not executed");
+			}
+			{
+				ExpectedAction = ::UE::Slate::Containers::EObservableArrayChangedAction::Add;
+
+				{
+					ExpectedArray.Reset();
+					ExpectedArray.Add(2);
+
+					ArrayValues.Add(2);
+					ObservableValues.Add(2);
+
+					AddErrorIfFalse(ArrayValues == ObservableValues, TEXT("ObservableValues == ArrayValues"));
+					++ExpectedCounter;
+					AddErrorIfFalse(Counter == ExpectedCounter, "The delegate was not executed");
+				}
+			}
+		}
+	}
+
+	return true;
+}
+} //namespace
+
+#endif //WITH_DEV_AUTOMATION_TESTS
+
diff --git a/Engine/Source/Runtime/SlateCore/Public/Containers/ObservableArray.h b/Engine/Source/Runtime/SlateCore/Public/Containers/ObservableArray.h
new file mode 100644
--- /dev/null
+++ b/Engine/Source/Runtime/SlateCore/Public/Containers/ObservableArray.h
@@ -0,0 +1,532 @@
+// Copyright Epic Games, Inc. All Rights Reserved.
+
+#pragma once
+
+#include "Containers/Array.h"
+#include "Containers/ArrayView.h"
+#include "Delegates/Delegate.h"
+#include "Delegates/DelegateCombinations.h"
+#include "HAL/UnrealMemory.h"
+#include "Misc/AssertionMacros.h"
+#include "Misc/MemStack.h"
+
+namespace UE::Slate::Containers
+{
+
+template <typename InElementType>
+struct TObservableArray;
+
+/** Type of action of the FScriptObservableArrayChangedArgs */
+enum class EObservableArrayChangedAction : uint8
+{
+	/** The array was reset. */
+	Reset,
+	/** Elements were added. */
+	Add,
+	/** Elements were removed. */
+	Remove,
+	/** Elements were removed and the same amount of elements moved from the end of the array to removed location. */
+	RemoveSwap,
+	/** 2 elements swapped location with each other. */
+	Swap,
+};
+
+/**
+ *
+ */
+template<typename InElementType>
+struct TObservableArrayChangedArgs
+{
+private:
+	using ArrayViewType = TArrayView<InElementType>;
+	using SizeType = typename ArrayViewType::SizeType;
+	using ElementType = typename ArrayViewType::ElementType;
+	friend TObservableArray<InElementType>;
+	
+private:
+	static TObservableArrayChangedArgs MakeResetAction()
+	{
+		TObservableArrayChangedArgs Result;
+		Result.Action = EObservableArrayChangedAction::Reset;
+		return Result;
+	}
+
+	static TObservableArrayChangedArgs MakeAddAction(const ArrayViewType InAddedItems, int32 InNewIndex)
+	{
+		TObservableArrayChangedArgs Result;
+		Result.Array = InAddedItems;
+		Result.StartIndex = InNewIndex;
+		Result.Action = EObservableArrayChangedAction::Add;
+		check(Result.Array.Num() > 0 && Result.StartIndex >= 0);
+		return Result;
+	}
+
+	static TObservableArrayChangedArgs MakeRemoveAction(const ArrayViewType InRemoveItems, int32 InRemoveStartedIndex)
+	{
+		TObservableArrayChangedArgs Result;
+		Result.Array = InRemoveItems;
+		Result.StartIndex = InRemoveStartedIndex;
+		Result.Action = EObservableArrayChangedAction::Remove;
+		check(Result.Array.Num() > 0 && Result.StartIndex >= 0);
+		return Result;
+	}
+	
+	static TObservableArrayChangedArgs MakeRemoveSwapAction(const ArrayViewType InRemoveItems, int32 InRemoveStartedIndex, int32 InPreviousMovedItemLocation)
+	{
+		TObservableArrayChangedArgs Result;
+		Result.Array = InRemoveItems;
+		Result.StartIndex = InRemoveStartedIndex;
+		Result.MoveIndex = InPreviousMovedItemLocation;
+		Result.Action = EObservableArrayChangedAction::RemoveSwap;
+		// The move index can be invalid if the array is now empty.
+		check(Result.Array.Num() > 0 && (Result.MoveIndex > 0 || Result.MoveIndex == INDEX_NONE) && Result.StartIndex >= 0);
+		return Result;
+	}
+
+	static TObservableArrayChangedArgs MakeSwapAction(int32 InFirstIndex, int32 InSecondIndex)
+	{
+		TObservableArrayChangedArgs Result;
+		Result.StartIndex = InFirstIndex;
+		Result.MoveIndex = InSecondIndex;
+		Result.Action = EObservableArrayChangedAction::Swap;
+		check(Result.Array.Num() == 0 && Result.MoveIndex >= 0 && Result.StartIndex >= 0 && Result.MoveIndex != Result.StartIndex);
+		return Result;
+	}
+
+public:
+	/** @return The action that caused the event. */
+	EObservableArrayChangedAction GetAction() const
+	{
+		return Action;
+	}
+
+	/**
+	 * Valid for the Add, Remove, RemoveSwap action.
+	 * Use GetItems.Num() to know how many elements were added/removed.
+	 * Add: The array index where we added the elements.
+	 * Remove: The old array index before it removed the elements. The index is not valid anymore.
+	 * RemoveSwap: The old array index before it removed the elements. The index is valid if the array is not empty.
+	 * @return The index the action started.
+	 */
+	SizeType GetActionIndex() const
+	{
+		return StartIndex;
+	}
+
+	struct FRemoveSwapIndex
+	{
+		FRemoveSwapIndex(SizeType InRemoveIndex, SizeType InPreviousMovedElmenentIndex)
+			: RemoveIndex(InRemoveIndex)
+			, PreviousMovedElmenentIndex(InPreviousMovedElmenentIndex)
+		{
+			
+		}
+		/**
+		 * The removed elements index.
+		 * The removed index is still valid if the array is not empty.
+		 * The moved elements are now at that location (if any).
+		 */
+		SizeType RemoveIndex;
+		/**
+		 * The previous location of the elements (if any) before they moved to the new location.
+		 * The index is not valid anymore.
+		 * Set to INDEX_NONE, if no element was moved.
+		 */
+		SizeType PreviousMovedElmenentIndex;
+	};
+
+	/**
+	 * Valid for the RemoveSwap action. 
+	 * @return The index of the removed elements and the index of the moved elements.
+	 */
+	FRemoveSwapIndex GetRemovedSwapIndexes() const
+	{
+		return (Action == EObservableArrayChangedAction::RemoveSwap) ? FRemoveSwapIndex(StartIndex, MoveIndex) : FRemoveSwapIndex(INDEX_NONE, INDEX_NONE);
+	}
+
+	struct FSwapIndex
+	{
+		FSwapIndex(SizeType InFirstIndex, SizeType InSecondIndex)
+			: FirstIndex(InFirstIndex)
+			, SecondIndex(InSecondIndex)
+		{
+			
+		}
+		SizeType FirstIndex;
+		SizeType SecondIndex;
+	};
+	
+	/**
+	 * Valid for the Swap action.
+	 * @return The indexes of the 2 swapped elements. */
+	FSwapIndex GetSwapIndex() const
+	{
+		return (Action == EObservableArrayChangedAction::Swap) ? FSwapIndex(MoveIndex, StartIndex): FSwapIndex(INDEX_NONE, INDEX_NONE);
+	}
+
+	/** @return The items added to the array or removed from the array. Valid for the Add, Remove and RemoveSwap action. */
+	ArrayViewType GetItems() const
+	{
+		return Array;
+	}
+	
+private:
+	ArrayViewType Array;
+	SizeType StartIndex = INDEX_NONE;
+	SizeType MoveIndex = INDEX_NONE;
+	EObservableArrayChangedAction Action = EObservableArrayChangedAction::Reset;
+};
+
+
+/**
+ *
+ */
+template <typename InElementType>
+struct TObservableArray
+{
+public:
+	using ArrayType = TArray<InElementType>;
+	using SizeType = typename ArrayType::SizeType;
+	using ElementType = typename ArrayType::ElementType;
+	using AllocatorType = typename ArrayType::AllocatorType;
+	using ObservableArrayChangedArgsType = TObservableArrayChangedArgs<InElementType>;
+	
+	DECLARE_MULTICAST_DELEGATE_OneParam(FArrayChangedDelegate, ObservableArrayChangedArgsType);
+	
+public:
+	TObservableArray() = default;
+	explicit TObservableArray(const ElementType* Ptr, SizeType Count)
+		: Array(Ptr, Count)
+	{
+		
+	}
+	
+	TObservableArray(std::initializer_list<ElementType> InitList)
+		: Array(InitList)
+	{
+		
+	}
+	
+	template<typename InOtherAllocatorType>
+	explicit TObservableArray(const TArray<ElementType, InOtherAllocatorType>& Other)
+		: Array(Other)
+	{
+		
+	}
+	
+	template <typename InOtherAllocatorType>
+	explicit TObservableArray(TArray<ElementType, InOtherAllocatorType>&& Other)
+		: Array(MoveTemp(Other))
+	{
+		
+	}
+
+	// Non-copyable for now, but this could be made copyable in future if needed.
+	TObservableArray(const TObservableArray&) = delete;
+	TObservableArray& operator=(const TObservableArray&) = delete;
+	TObservableArray(TObservableArray&& Other) = delete;
+	TObservableArray& operator=(TObservableArray&& Other) = delete;
+
+	~TObservableArray()
+	{
+	}
+	
+public:
+	FArrayChangedDelegate& OnArrayChanged()
+	{
+		return ArrayChangedDelegate;
+	}
+
+public:
+	int32 Add(const ElementType& Item)
+	{
+		int32 NewIndex = Array.Add(Item);
+		ArrayChangedDelegate.Broadcast(ObservableArrayChangedArgsType::MakeAddAction({ GetData() + NewIndex, 1 }, NewIndex));
+		return NewIndex;
+	}
+
+	int32 Add(ElementType&& Item)
+	{
+		int32 NewIndex = Array.Add(MoveTempIfPossible(Item));
+		ArrayChangedDelegate.Broadcast(ObservableArrayChangedArgsType::MakeAddAction({ GetData() + NewIndex, 1 }, NewIndex));
+		return NewIndex;
+	}
+
+	template <typename... InArgTypes>
+	int32 Emplace(InArgTypes&&... Args)
+	{
+		int32 NewIndex = Array.Emplace(Forward<InArgTypes>(Args)...);
+		ArrayChangedDelegate.Broadcast(ObservableArrayChangedArgsType::MakeAddAction({ GetData() + NewIndex, 1 }, NewIndex));
+		return NewIndex;
+	}
+
+	template <typename... InArgTypes>
+	void EmplaceAt(SizeType Index, InArgTypes&&... Args)
+	{
+		Array.EmplaceAt(Index, Forward<InArgTypes>(Args)...);
+		ArrayChangedDelegate.Broadcast(ObservableArrayChangedArgsType::MakeAddAction({ GetData() + Index, 1 }, Index));
+	}
+
+	template <typename OtherElementType>
+	void Append(const TObservableArray<OtherElementType>& Source)
+	{
+		Append(Source.Array);
+	}
+
+	template <typename OtherElementType, typename OtherAllocatorType>
+	void Append(const TArray<OtherElementType, OtherAllocatorType>& Source)
+	{
+		SizeType PreviousNum = Array.Num();
+		if (PreviousNum > 0)
+		{
+			Array.Append(Source);
+			SizeType NewNum = Array.Num();
+			ArrayChangedDelegate.Broadcast(ObservableArrayChangedArgsType::MakeAddAction({ GetData() + PreviousNum, NewNum - PreviousNum }, PreviousNum));
+		}
+	}
+
+	template <typename OtherElementType, typename OtherAllocator>
+	void Append(TArray<OtherElementType, OtherAllocator>&& Source)
+	{
+		SizeType PreviousNum = Array.Num();
+		if (PreviousNum > 0)
+		{
+			Array.Append(MoveTempIfPossible(Source));
+			SizeType NewNum = Array.Num();
+			ArrayChangedDelegate.Broadcast(ObservableArrayChangedArgsType::MakeAddAction({ GetData() + PreviousNum, NewNum - PreviousNum }, PreviousNum));
+		}
+	}
+
+	SizeType RemoveSingle(const ElementType& Item, bool bAllowShrinking = true)
+	{
+		SizeType Index = Array.Find(Item);
+		if (Index == INDEX_NONE)
+		{
+			return 0;
+		}
+		RemoveAt(Index, 1, bAllowShrinking);
+		return 1;
+	}
+
+	SizeType RemoveSingleSwap(const ElementType& Item, bool bAllowShrinking = true)
+	{
+		SizeType Index = Array.Find(Item);
+		if (Index == INDEX_NONE)
+		{
+			return 0;
+		}
+		RemoveAtSwap(Index, 1, bAllowShrinking);
+		return 1;
+	}
+
+	void RemoveAt(SizeType Index)
+	{
+		RemoveAt(Index, 1, true);
+	}
+
+	template <typename CountType>
+	void RemoveAt(SizeType Index, CountType NumToRemove = 1, bool bAllowShrinking = true)
+	{
+		static_assert(!TAreTypesEqual<CountType, bool>::Value, "TObservableArray::RemoveAt: unexpected bool passed as the Count argument");
+		check((NumToRemove > 0) & (Index >= 0) & (Index + NumToRemove <= Num()));
+		if (NumToRemove > 0)
+		{		
+			if (NumToRemove == 1)
+			{
+				ElementType RemovedElement = MoveTemp(Array[Index]);
+				Array.RemoveAt(Index, NumToRemove, bAllowShrinking);
+				ArrayChangedDelegate.Broadcast(ObservableArrayChangedArgsType::MakeRemoveAction({ &RemovedElement, 1 }, Index));
+			}
+			else
+			{
+				// Copy the items to a temporary array to call the delegate
+				FMemMark Mark(FMemStack::Get());
+				TArray<ElementType, TMemStackAllocator<>> RemovedElements;
+				RemovedElements.Reserve(NumToRemove);
+				for (int32 RemoveIndex = 0; RemoveIndex < NumToRemove; ++RemoveIndex)
+				{
+					RemovedElements.Add(MoveTemp(Array[RemoveIndex + Index]));
+				}
+				
+				Array.RemoveAt(Index, NumToRemove, bAllowShrinking);
+				ArrayChangedDelegate.Broadcast(ObservableArrayChangedArgsType::MakeRemoveAction({ (ElementType*)RemovedElements.GetData(), NumToRemove }, Index));
+			}
+		}
+	}
+
+	void RemoveAtSwap(SizeType Index)
+	{
+		RemoveAtSwap(Index, 1, true);
+	}
+	
+	template <typename CountType>
+	void RemoveAtSwap(SizeType Index, CountType NumToRemove = 1, bool bAllowShrinking = true)
+	{
+		static_assert(!TAreTypesEqual<CountType, bool>::Value, "TObservableArray::RemoveAtSwap: unexpected bool passed as the Count argument");
+		check((NumToRemove > 0) & (Index >= 0) & (Index + NumToRemove <= Num()));
+		if (NumToRemove > 0)
+		{
+			if (NumToRemove == 1)
+			{
+				ElementType RemovedElement = MoveTemp(Array[Index]);
+				
+				Array.RemoveAtSwap(Index, NumToRemove, bAllowShrinking);
+				ArrayChangedDelegate.Broadcast(ObservableArrayChangedArgsType::MakeRemoveAction({ &RemovedElement, 1 }, Index));
+			}
+			else
+			{
+				// Copy the items to a temporary array to call the delegate
+				FMemMark Mark(FMemStack::Get());
+				TArray<ElementType, TMemStackAllocator<>> RemovedElements;
+				RemovedElements.Reserve(NumToRemove);
+				for (int32 RemoveIndex = 0; RemoveIndex < NumToRemove; ++RemoveIndex)
+				{
+					RemovedElements.Add(MoveTemp(Array[RemoveIndex + Index]));
+				}
+				
+				Array.RemoveAtSwap(Index, NumToRemove, bAllowShrinking);
+				ArrayChangedDelegate.Broadcast(ObservableArrayChangedArgsType::MakeRemoveAction({ RemovedElements.GetData(), NumToRemove }, Index));
+			}
+		}
+	}
+
+	void Swap(SizeType FirstIndexToSwap, SizeType SecondIndexToSwap)
+	{
+		check(FirstIndexToSwap >= 0 && SecondIndexToSwap >= 0);
+		if (FirstIndexToSwap != SecondIndexToSwap)
+		{
+			Array.SwapMemory(FirstIndexToSwap, SecondIndexToSwap);
+			ArrayChangedDelegate.Broadcast(ObservableArrayChangedArgsType::MakeSwapAction(FirstIndexToSwap, SecondIndexToSwap));
+		}
+	}
+
+	void Reset(SizeType NewSize = 0)
+	{
+		if (Num())
+		{
+			Array.Reset(NewSize);
+			ArrayChangedDelegate.Broadcast(ObservableArrayChangedArgsType::MakeResetAction());
+		}
+	}
+
+	void Reserve(SizeType Number)
+	{
+		Array.Reserve(Number);
+	}
+
+	bool IsEmpty() const
+	{
+		return Array.IsEmpty();
+	}
+
+	int32 Num() const
+	{
+		return Array.Num();
+	}
+
+	bool IsValidIndex(SizeType Index) const
+	{
+		return Array.IsValidIndex(Index);
+	}
+
+	ElementType& operator[](SizeType Index)
+	{
+		return GetData()[Index];
+	}
+
+	const ElementType& operator[](SizeType Index) const
+	{
+		return GetData()[Index];
+	}
+
+	template <typename ComparisonType>
+	bool Contains(const ComparisonType& Item) const
+	{
+		return Array.Contains(Item);
+	}
+
+	template <typename InPredicate>
+	SizeType ContainsByPredicate(InPredicate Pred) const
+	{
+		return Array.ContainsByPredicate(Pred);
+	}
+
+	SizeType Find(const ElementType& Item) const
+	{
+		return Array.Find(Item);
+	}
+
+	template <typename InPredicate>
+	ElementType* FindByPredicate(InPredicate Pred)
+	{
+		return Array.FindByPredicate(Pred);
+	}
+
+	template <typename InPredicate>
+	const ElementType* FindByPredicate(InPredicate Pred) const
+	{
+		return Array.FindByPredicate(Pred);
+	}
+
+	template <typename InPredicate>
+	SizeType IndexByPredicate(InPredicate Pred) const
+	{
+		return Array.IndexByPredicate(Pred);
+	}
+
+private:
+	ElementType* GetData()
+	{
+		return Array.GetData();
+	}
+
+	const ElementType* GetData() const
+	{
+		return Array.GetData();
+	}
+
+public:
+	template <typename InOtherAllocatorType>
+	bool operator==(TArray<ElementType, InOtherAllocatorType>& OtherArray) const
+	{
+		return Array.operator==(OtherArray);
+	}
+	
+	bool operator==(const TObservableArray& OtherArray) const
+	{
+		return Array.operator==(OtherArray.Array);
+	}
+
+	template <typename InOtherAllocatorType>
+	friend bool operator==(TArray<ElementType, InOtherAllocatorType>& OtherArray, const TObservableArray& Self)
+	{
+		return Self.Array.operator==(OtherArray);
+	}
+
+public:
+	ElementType* begin(TObservableArray& Arr)
+	{
+		return Arr.Array.begin();
+	}
+
+	const ElementType* begin(const TObservableArray& Arr) const
+	{
+		return Arr.Array.begin();
+	}
+
+	ElementType* end(TObservableArray& Arr)
+	{
+		return Arr.Array.end();
+	}
+
+	const ElementType* end(const TObservableArray& Arr) const
+	{
+		return Arr.Array.end();
+	}
+	
+private:
+	ArrayType Array;
+	FArrayChangedDelegate ArrayChangedDelegate;
+};
+ 
+ } //namespace
\ No newline at end of file
