Apply Disabled Effects CVar from SGameLayerManager to PaintInvalidationRoot to control said effects on Invalidation

diff --git a/Engine/Source/Runtime/Engine/Private/Slate/SGameLayerManager.cpp b/Engine/Source/Runtime/Engine/Private/Slate/SGameLayerManager.cpp
--- a/Engine/Source/Runtime/Engine/Private/Slate/SGameLayerManager.cpp
+++ b/Engine/Source/Runtime/Engine/Private/Slate/SGameLayerManager.cpp
@@ -31,6 +31,14 @@ static void HandlePerUserHitTestingToggled(IConsoleVariable* CVar)
 	FSlateApplication::Get().InvalidateAllWidgets(false);
 }
 
+static void HandleApplyDisabledEffectOnWidgetsToggled(IConsoleVariable* CVar)
+{
+	if (FSlateApplication::IsInitialized())
+	{
+		FSlateApplication::Get().InvalidateAllWidgets(false);
+	}
+}
+
 static int32 bEnablePerUserHitTesting = true;
 
 static FAutoConsoleVariableRef CVarEnablePerUserHitTesting(
@@ -45,7 +53,9 @@ static int32 bApplyDisabledEffectOnWidgets = 1;
 static FAutoConsoleVariableRef CVarApplyDisabledEffectOnWidgets(
 	TEXT("ApplyDisabledEffectOnWidgets"),
 	bApplyDisabledEffectOnWidgets,
-	TEXT("If true, disabled game-layer widgets will have alpha multiplied by 0.45."));
+	TEXT("If true, disabled game-layer widgets will have alpha multiplied by 0.45."),
+	FConsoleVariableDelegate::CreateStatic(&HandleApplyDisabledEffectOnWidgetsToggled)
+);
 
 #if UE_SLATE_WITH_GAMELAYER_CANVAS_VISIBILITY_COMMANDS
 namespace CanvasVisibility
diff --git a/Engine/Source/Runtime/SlateCore/Private/FastUpdate/WidgetProxy.cpp b/Engine/Source/Runtime/SlateCore/Private/FastUpdate/WidgetProxy.cpp
--- a/Engine/Source/Runtime/SlateCore/Private/FastUpdate/WidgetProxy.cpp
+++ b/Engine/Source/Runtime/SlateCore/Private/FastUpdate/WidgetProxy.cpp
@@ -227,7 +227,7 @@ FWidgetProxy::FUpdateResult FWidgetProxy::Repaint(const FPaintArgs& PaintArgs, F
 	const int32 PrevUserIndex = PaintArgs.GetHittestGrid().GetUserIndex();
 
 	PaintArgs.GetHittestGrid().SetUserIndex(MyState.IncomingUserIndex);
-	GSlateFlowDirection = MyState.IncomingFlowDirection;
+	TGuardValue<EFlowDirection> FlowGuard(GSlateFlowDirection, MyState.IncomingFlowDirection);
 	
 	FPaintArgs UpdatedArgs = PaintArgs.WithNewParent(MyState.PaintParent.Pin().Get());
 	UpdatedArgs.SetInheritedHittestability(MyState.bInheritedHittestability);
@@ -257,7 +257,8 @@ FWidgetProxy::FUpdateResult FWidgetProxy::Repaint(const FPaintArgs& PaintArgs, F
 		// clip index should be what it was before.  if this assert fails something internal inside the above paint call did not pop clip properly
 		check(StartingClipIndex == OutDrawElements.GetClippingIndex());
 	}
-
+	
+	OutDrawElements.SetIsInGameLayer(MyState.bIsInGameLayer);
 	return FUpdateResult{ PrevLayerId, NewLayerId };
 }
 
diff --git a/Engine/Source/Runtime/SlateCore/Private/Widgets/SWidget.cpp b/Engine/Source/Runtime/SlateCore/Private/Widgets/SWidget.cpp
--- a/Engine/Source/Runtime/SlateCore/Private/Widgets/SWidget.cpp
+++ b/Engine/Source/Runtime/SlateCore/Private/Widgets/SWidget.cpp
@@ -1478,6 +1478,7 @@ int32 SWidget::Paint(const FPaintArgs& Args, const FGeometry& AllottedGeometry,
 	PersistentState.IncomingUserIndex = (int8)IncomingUserIndex;
 
 	PersistentState.IncomingFlowDirection = GSlateFlowDirection;
+	PersistentState.bIsInGameLayer = OutDrawElements.GetIsInGameLayer();
 
 	FPaintArgs UpdatedArgs = Args.WithNewParent(this);
 	UpdatedArgs.SetInheritedHittestability(bOutgoingHittestability);
diff --git a/Engine/Source/Runtime/SlateCore/Public/FastUpdate/WidgetProxy.h b/Engine/Source/Runtime/SlateCore/Public/FastUpdate/WidgetProxy.h
--- a/Engine/Source/Runtime/SlateCore/Public/FastUpdate/WidgetProxy.h
+++ b/Engine/Source/Runtime/SlateCore/Public/FastUpdate/WidgetProxy.h
@@ -246,6 +246,7 @@ struct FSlateWidgetPersistentState
 		, bParentEnabled(true)
 		, bInheritedHittestability(false)
 		, bDeferredPainting(false)
+		, bIsInGameLayer(false)
 	{}
 
 	TWeakPtr<SWidget> PaintParent;
@@ -263,6 +264,7 @@ struct FSlateWidgetPersistentState
 	uint8 bParentEnabled : 1;
 	uint8 bInheritedHittestability : 1;
 	uint8 bDeferredPainting : 1;
+	uint8 bIsInGameLayer : 1;
 
 	static const FSlateWidgetPersistentState NoState;
 };
