Fix SViewport::OnPaint asserting when inserting a custom hit test path, when the viewport is temporarily not hittestable.
This fixes a crash when dragging viewport tabs that have an overlay widget, specifically, when using HUDs in virtual cameras.

diff --git a/Engine/Source/Runtime/Slate/Private/Widgets/SViewport.cpp b/Engine/Source/Runtime/Slate/Private/Widgets/SViewport.cpp
--- a/Engine/Source/Runtime/Slate/Private/Widgets/SViewport.cpp
+++ b/Engine/Source/Runtime/Slate/Private/Widgets/SViewport.cpp
@@ -205,9 +205,9 @@ int32 SViewport::OnPaint( const FPaintArgs& Args, const FGeometry& AllottedGeome
 	}
 
 	// If there are any custom hit testable widgets in the 3D world we need to register their custom hit test path here.
-	if ( CustomHitTestPath.IsValid() )
+	if ( CustomHitTestPath.IsValid() && Args.GetHittestGrid().ContainsWidget(this))
 	{
-		Args.InsertCustomHitTestPath(this, CustomHitTestPath.ToSharedRef());
+		Args.GetHittestGrid().InsertCustomHitTestPath(this, CustomHitTestPath.ToSharedRef());
 	}
 
 	return Layer;
diff --git a/Engine/Source/Runtime/SlateCore/Private/Input/HittestGrid.cpp b/Engine/Source/Runtime/SlateCore/Private/Input/HittestGrid.cpp
--- a/Engine/Source/Runtime/SlateCore/Private/Input/HittestGrid.cpp
+++ b/Engine/Source/Runtime/SlateCore/Private/Input/HittestGrid.cpp
@@ -223,8 +223,8 @@ TArray<FWidgetAndPointer> FHittestGrid::GetBubblePath(FVector2D DesktopSpaceCoor
 				{
 					if (BestHitWidgetData.CustomPath.IsValid())
 					{
-						const TArray<FWidgetAndPointer> BubblePathExtension = BestHitWidgetData.CustomPath.Pin()->GetBubblePathAndVirtualCursors(FirstHitWidget->GetTickSpaceGeometry(), DesktopSpaceCoordinate, bIgnoreEnabledStatus);
-						Path.Append(BubblePathExtension);
+						TArray<FWidgetAndPointer> BubblePathExtension = BestHitWidgetData.CustomPath.Pin()->GetBubblePathAndVirtualCursors(FirstHitWidget->GetTickSpaceGeometry(), DesktopSpaceCoordinate, bIgnoreEnabledStatus);
+						Path.Append(MoveTemp(BubblePathExtension));
 					}
 				}
 	
@@ -792,6 +792,11 @@ void FHittestGrid::UpdateWidget(const SWidget* InWidget, FSlateInvalidationWidge
 	}
 }
 
+bool FHittestGrid::ContainsWidget(const SWidget* InWidget) const
+{
+	return WidgetMap.Contains(InWidget);
+}
+
 void FHittestGrid::InsertCustomHitTestPath(const TSharedRef<SWidget> InWidget, TSharedRef<ICustomHitTestPath> CustomHitTestPath)
 {
 	InsertCustomHitTestPath(&InWidget.Get(), CustomHitTestPath);
@@ -799,9 +804,12 @@ void FHittestGrid::InsertCustomHitTestPath(const TSharedRef<SWidget> InWidget, T
 
 void FHittestGrid::InsertCustomHitTestPath(const SWidget* InWidget, const TSharedRef<ICustomHitTestPath>& CustomHitTestPath)
 {
-	int32 WidgetIndex = WidgetMap.FindChecked(InWidget);
-	FWidgetData& WidgetData = WidgetArray[WidgetIndex];
-	WidgetData.CustomPath = CustomHitTestPath;
+	int32* WidgetIndex = WidgetMap.Find(InWidget);
+	if (ensureMsgf(WidgetIndex, TEXT("The widget is not in the hittest grid.")))
+	{
+		FWidgetData& WidgetData = WidgetArray[*WidgetIndex];
+		WidgetData.CustomPath = CustomHitTestPath;
+	}
 }
 
 FHittestGrid::FIndexAndDistance FHittestGrid::GetHitIndexFromCellIndex(const FGridTestingParams& Params) const
diff --git a/Engine/Source/Runtime/SlateCore/Public/Input/HittestGrid.h b/Engine/Source/Runtime/SlateCore/Public/Input/HittestGrid.h
--- a/Engine/Source/Runtime/SlateCore/Public/Input/HittestGrid.h
+++ b/Engine/Source/Runtime/SlateCore/Public/Input/HittestGrid.h
@@ -110,6 +110,9 @@ public:
 	/** Update the widget SecondarySort without removing it and readding it again. */
 	void UpdateWidget(const SWidget* InWidget, FSlateInvalidationWidgetSortOrder InSecondarySort);
 
+	/** Check if SWidget is contained within the HitTest Grid */
+	bool ContainsWidget(const SWidget* InWidget) const;
+
 	/** Append an already existing grid that occupy the same space. */
 	UE_DEPRECATED(4.26, "Deprecated. Use the FHittestGrid::AddGrid method instead")
 	void AppendGrid(FHittestGrid& OtherGrid) {}
