Prep for element generation for possible bulk batch creation per element type.
Migrate FSlateRenderBatch index poputation to a simple static array return in boxes, borders, and aliased lines.

diff --git a/Engine/Source/Developer/SlateReflector/Private/VisualTreeCapture.cpp b/Engine/Source/Developer/SlateReflector/Private/VisualTreeCapture.cpp
--- a/Engine/Source/Developer/SlateReflector/Private/VisualTreeCapture.cpp
+++ b/Engine/Source/Developer/SlateReflector/Private/VisualTreeCapture.cpp
@@ -58,17 +58,20 @@ void FVisualEntry::Resolve(const FSlateWindowElementList& ElementList)
 		return;
 	}
 
-	const FSlateDrawElement& Element = ElementList.GetUncachedDrawElements()[ElementIndex];
-	const FSlateRenderTransform& Transform = Element.GetRenderTransform();
-	const FVector2D LocalSize = Element.GetLocalSize();
+	if (const FSlateDrawElementContainer* Container = ElementList.GetUncachedDrawElements().Find(EElementType::ET_NonMapped))
+	{
+		const FSlateDrawElement& Element = Container->Elements[ElementIndex];
+		const FSlateRenderTransform& Transform = Element.GetRenderTransform();
+		const FVector2D LocalSize = Element.GetLocalSize();
 
-	TopLeft = Transform.TransformPoint(FVector2D(0, 0));
-	TopRight = Transform.TransformPoint(FVector2D(LocalSize.X, 0));
-	BottomLeft = Transform.TransformPoint(FVector2D(0, LocalSize.Y));
-	BottomRight = Transform.TransformPoint(LocalSize);
+		TopLeft = Transform.TransformPoint(FVector2D(0, 0));
+		TopRight = Transform.TransformPoint(FVector2D(LocalSize.X, 0));
+		BottomLeft = Transform.TransformPoint(FVector2D(0, LocalSize.Y));
+		BottomRight = Transform.TransformPoint(LocalSize);
 
-	LayerId = Element.GetLayer();
-	ClippingIndex = Element.GetPrecachedClippingIndex();
+		LayerId = Element.GetLayer();
+		ClippingIndex = Element.GetPrecachedClippingIndex();
+	}
 }
 
 bool FVisualEntry::IsPointInside(const FVector2D& Point) const
@@ -178,23 +181,26 @@ void FVisualTreeCapture::AddInvalidationRootCachedEntries(TSharedRef<FVisualTree
 	const TArray<TSharedPtr<FSlateCachedElementList>>& CachedElements = Data.GetCachedElementLists();
 	for (const TSharedPtr<FSlateCachedElementList>& CachedElement : CachedElements)
 	{
-		const FSlateDrawElementArray& DrawElementArray = CachedElement->DrawElements;
-		const SWidget* Widget = CachedElement->OwningWidget;
-		// todo, should check if parents has the metadata also
-		if (Widget && !Widget->GetMetaData<FInvisibleToWidgetReflectorMetaData>())
+		for (TPair<EElementType, FSlateDrawElementContainer>& DrawElement : CachedElement->DrawElements)
 		{
-			for (const FSlateDrawElement& Element : DrawElementArray)
+			const FSlateDrawElementArray& DrawElementArray = DrawElement.Value.Elements;
+			const SWidget* Widget = CachedElement->OwningWidget;
+			// todo, should check if parents has the metadata also
+			if (Widget && !Widget->GetMetaData<FInvisibleToWidgetReflectorMetaData>())
 			{
-				const int32 EntryIndex = Tree->Entries.Emplace(Widget->AsShared(), Element);
-
-				const FSlateClippingState* ClippingState = Element.GetClippingHandle().GetCachedClipState();
-				if (ClippingState)
+				for (const FSlateDrawElement& Element : DrawElementArray)
 				{
-					int32& ClippingRefIndex = Tree->Entries[EntryIndex].ClippingIndex;
-					ClippingRefIndex = Tree->CachedClippingStates.IndexOfByKey(*ClippingState);
-					if (ClippingRefIndex == INDEX_NONE)
+					const int32 EntryIndex = Tree->Entries.Emplace(Widget->AsShared(), Element);
+
+					const FSlateClippingState* ClippingState = Element.GetClippingHandle().GetCachedClipState();
+					if (ClippingState)
 					{
-						ClippingRefIndex = Tree->CachedClippingStates.Add(*ClippingState);
+						int32& ClippingRefIndex = Tree->Entries[EntryIndex].ClippingIndex;
+						ClippingRefIndex = Tree->CachedClippingStates.IndexOfByKey(*ClippingState);
+						if (ClippingRefIndex == INDEX_NONE)
+						{
+							ClippingRefIndex = Tree->CachedClippingStates.Add(*ClippingState);
+						}
 					}
 				}
 			}
diff --git a/Engine/Source/Runtime/SlateCore/Private/Debugging/SlateDebugging.cpp b/Engine/Source/Runtime/SlateCore/Private/Debugging/SlateDebugging.cpp
--- a/Engine/Source/Runtime/SlateCore/Private/Debugging/SlateDebugging.cpp
+++ b/Engine/Source/Runtime/SlateCore/Private/Debugging/SlateDebugging.cpp
@@ -390,7 +390,11 @@ FSlateDebugging::FEndWidgetPaint FSlateDebugging::EndWidgetPaint;
 
 FSlateDebugging::FPaintDebugElements FSlateDebugging::PaintDebugElements;
 
+PRAGMA_DISABLE_DEPRECATION_WARNINGS
 FSlateDebugging::FDrawElement FSlateDebugging::ElementAdded;
+PRAGMA_ENABLE_DEPRECATION_WARNINGS
+
+FSlateDebugging::FDrawElementType FSlateDebugging::ElementTypeAdded;
 
 FSlateDebugging::FWidgetWarningEvent FSlateDebugging::Warning;
 
diff --git a/Engine/Source/Runtime/SlateCore/Private/Rendering/DrawElements.cpp b/Engine/Source/Runtime/SlateCore/Private/Rendering/DrawElements.cpp
--- a/Engine/Source/Runtime/SlateCore/Private/Rendering/DrawElements.cpp
+++ b/Engine/Source/Runtime/SlateCore/Private/Rendering/DrawElements.cpp
@@ -244,7 +244,7 @@ void FSlateDrawElement::MakeDebugQuad( FSlateWindowElementList& ElementList, uin
 		return;
 	}
 
-	FSlateDrawElement& Element = ElementList.AddUninitialized();
+	FSlateDrawElement& Element = ElementList.AddUninitialized(EElementType::ET_DebugQuad);
 	FSlateBoxPayload& BoxPayload = ElementList.CreatePayload<FSlateBoxPayload>(Element);
 
 	BoxPayload.SetTint(Tint);
@@ -265,7 +265,7 @@ FSlateDrawElement& FSlateDrawElement::MakeBoxInternal(
 	// Cast to Rounded Rect to get the internal parameters 
 	// New payload type - inherit from BoxPayload 
 
-	FSlateDrawElement& Element = ElementList.AddUninitialized();
+	FSlateDrawElement& Element = ElementList.AddUninitialized(ElementType);
 
 	FSlateBoxPayload* BoxPayload;
 	if ( ElementType == EElementType::ET_RoundedBox )
@@ -386,7 +386,7 @@ void FSlateDrawElement::MakeText( FSlateWindowElementList& ElementList, uint32 I
 		return;
 	}
 
-	FSlateDrawElement& Element = ElementList.AddUninitialized();
+	FSlateDrawElement& Element = ElementList.AddUninitialized(EElementType::ET_Text);
 	FSlateTextPayload& DataPayload = ElementList.CreatePayload<FSlateTextPayload>(Element);
 
 	DataPayload.SetTint(InTint);
@@ -417,7 +417,7 @@ void FSlateDrawElement::MakeText( FSlateWindowElementList& ElementList, uint32 I
 		return;
 	}
 
-	FSlateDrawElement& Element = ElementList.AddUninitialized();
+	FSlateDrawElement& Element = ElementList.AddUninitialized(EElementType::ET_Text);
 
 	FSlateTextPayload& DataPayload = ElementList.CreatePayload<FSlateTextPayload>(Element);
 
@@ -491,7 +491,7 @@ void FSlateDrawElement::MakeShapedText(FSlateWindowElementList& ElementList, uin
 	SlateDrawElement::CheckInvalidUMaterial(InShapedGlyphSequence->GetFontOutlineSettings().OutlineMaterial, TEXT("Outline Material"));
 #endif
 
-	FSlateDrawElement& Element = ElementList.AddUninitialized();
+	FSlateDrawElement& Element = ElementList.AddUninitialized(EElementType::ET_ShapedText);
 
 	FSlateShapedTextPayload& DataPayload = ElementList.CreatePayload<FSlateShapedTextPayload>(Element);
 	DataPayload.SetTint(BaseTint);
@@ -510,7 +510,7 @@ void FSlateDrawElement::MakeGradient( FSlateWindowElementList& ElementList, uint
 		return;
 	}
 
-	FSlateDrawElement& Element = ElementList.AddUninitialized();
+	FSlateDrawElement& Element = ElementList.AddUninitialized(EElementType::ET_Gradient);
 
 	FSlateGradientPayload& DataPayload = ElementList.CreatePayload<FSlateGradientPayload>(Element);
 
@@ -532,7 +532,7 @@ void FSlateDrawElement::MakeSpline(FSlateWindowElementList& ElementList, uint32
 	{
 		return;
 	}
-	FSlateDrawElement& Element = ElementList.AddUninitialized();
+	FSlateDrawElement& Element = ElementList.AddUninitialized(EElementType::ET_Spline);
 
 	FSlateSplinePayload& DataPayload = ElementList.CreatePayload<FSlateSplinePayload>(Element);
 
@@ -554,7 +554,7 @@ void FSlateDrawElement::MakeCubicBezierSpline(FSlateWindowElementList& ElementLi
 	{
 		return;
 	}
-	FSlateDrawElement& Element = ElementList.AddUninitialized();
+	FSlateDrawElement& Element = ElementList.AddUninitialized(EElementType::ET_Spline);
 
 	FSlateSplinePayload& DataPayload = ElementList.CreatePayload<FSlateSplinePayload>(Element);
 
@@ -583,7 +583,7 @@ void FSlateDrawElement::MakeDrawSpaceGradientSpline(FSlateWindowElementList& Ele
 		return;
 	}
 
-	FSlateDrawElement& Element = ElementList.AddUninitialized();
+	FSlateDrawElement& Element = ElementList.AddUninitialized(EElementType::ET_Spline);
 
 	FSlateSplinePayload& DataPayload = ElementList.CreatePayload<FSlateSplinePayload>(Element);
 	DataPayload.SetGradientHermiteSpline(InStart, InStartDir, InEnd, InEndDir, InThickness, InGradientStops);
@@ -601,7 +601,7 @@ void FSlateDrawElement::MakeDrawSpaceGradientSpline(FSlateWindowElementList& Ele
 		return;
 	}
 
-	FSlateDrawElement& Element = ElementList.AddUninitialized();
+	FSlateDrawElement& Element = ElementList.AddUninitialized(EElementType::ET_Spline);
 
 	FSlateSplinePayload& DataPayload = ElementList.CreatePayload<FSlateSplinePayload>(Element);
 	DataPayload.SetGradientHermiteSpline(InStart, InStartDir, InEnd, InEndDir, InThickness, InGradientStops);
@@ -634,7 +634,7 @@ void FSlateDrawElement::MakeLines(FSlateWindowElementList& ElementList, uint32 I
 		return;
 	}
 
-	FSlateDrawElement& Element = ElementList.AddUninitialized();
+	FSlateDrawElement& Element = ElementList.AddUninitialized(EElementType::ET_Line);
 
 	FSlateLinePayload& DataPayload = ElementList.CreatePayload<FSlateLinePayload>(Element);
 
@@ -678,7 +678,7 @@ void FSlateDrawElement::MakeLines( FSlateWindowElementList& ElementList, uint32
 		return;
 	}
 
-	FSlateDrawElement& Element = ElementList.AddUninitialized();
+	FSlateDrawElement& Element = ElementList.AddUninitialized(EElementType::ET_Line);
 
 	FSlateLinePayload& DataPayload = ElementList.CreatePayload<FSlateLinePayload>(Element);
 	DataPayload.SetTint(InTint);
@@ -697,7 +697,7 @@ void FSlateDrawElement::MakeViewport( FSlateWindowElementList& ElementList, uint
 		return;
 	}
 
-	FSlateDrawElement& Element = ElementList.AddUninitialized();
+	FSlateDrawElement& Element = ElementList.AddUninitialized(EElementType::ET_Viewport);
 	FSlateViewportPayload& DataPayload = ElementList.CreatePayload<FSlateViewportPayload>(Element);
 
 	DataPayload.SetViewport(Viewport, InTint);
@@ -714,7 +714,7 @@ void FSlateDrawElement::MakeCustom( FSlateWindowElementList& ElementList, uint32
 		return;
 	}
 
-	FSlateDrawElement& Element = ElementList.AddUninitialized();
+	FSlateDrawElement& Element = ElementList.AddUninitialized(EElementType::ET_Custom);
 
 	FSlateCustomDrawerPayload& DataPayload = ElementList.CreatePayload<FSlateCustomDrawerPayload>(Element);
 	DataPayload.SetCustomDrawer(CustomDrawer);
@@ -733,7 +733,7 @@ void FSlateDrawElement::MakeCustomVerts(FSlateWindowElementList& ElementList, ui
 		return;
 	}
 
-	FSlateDrawElement& Element = ElementList.AddUninitialized();
+	FSlateDrawElement& Element = ElementList.AddUninitialized(EElementType::ET_CustomVerts);
 	FSlateCustomVertsPayload& DataPayload = ElementList.CreatePayload<FSlateCustomVertsPayload>(Element);
 
 	const FSlateShaderResourceProxy* RenderingProxy = InRenderResourceHandle.GetResourceProxy();
@@ -753,7 +753,7 @@ void FSlateDrawElement::MakePostProcessPass(FSlateWindowElementList& ElementList
 		return;
 	}
 
-	FSlateDrawElement& Element = ElementList.AddUninitialized();
+	FSlateDrawElement& Element = ElementList.AddUninitialized(EElementType::ET_PostProcessPass);
 
 	FSlatePostProcessPayload& DataPayload = ElementList.CreatePayload<FSlatePostProcessPayload>(Element);
 	DataPayload.DownsampleAmount = DownsampleAmount;
@@ -829,7 +829,8 @@ bool FSlateWindowElementList::GetIsInGameLayer()
 	return bIsInGameLayer;
 }
 
-FSlateDrawElement& FSlateWindowElementList::AddUninitialized()
+
+FSlateDrawElement& FSlateWindowElementList::AddUninitialized(EElementType InElementType)
 {
 	const bool bAllowCache = CachedElementDataListStack.Num() > 0 && WidgetDrawStack.Num() && !WidgetDrawStack.Top().bIsVolatile;
 
@@ -840,15 +841,16 @@ FSlateDrawElement& FSlateWindowElementList::AddUninitialized()
 	}
 	else
 	{
-		FSlateDrawElementArray& Elements = UncachedDrawElements;
-		const int32 InsertIdx = Elements.AddDefaulted();
+		FSlateDrawElementMap& Elements = UncachedDrawElements;
+		FSlateDrawElementContainer& Container = Elements.FindOrAdd(InElementType);
+		const int32 InsertIdx = Container.Elements.AddDefaulted();
 
 #if WITH_SLATE_DEBUGGING
-		FSlateDebugging::ElementAdded.Broadcast(*this, InsertIdx);
+		FSlateDebugging::ElementTypeAdded.Broadcast(*this, InsertIdx, InElementType);
 #endif
 
-		FSlateDrawElement& NewElement = Elements[InsertIdx];
-		return Elements[InsertIdx];
+		FSlateDrawElement& NewElement = Container.Elements[InsertIdx];
+		return NewElement;
 	}
 }
 
@@ -1098,9 +1100,12 @@ void FSlateWindowElementList::ResetElementList()
 
 void FSlateWindowElementList::AddReferencedObjects(FReferenceCollector& Collector)
 {
-	for (FSlateDrawElement& Element : UncachedDrawElements)
+	for (TPair<EElementType, FSlateDrawElementContainer>& DrawElement : UncachedDrawElements)
 	{
-		Element.AddReferencedObjects(Collector);
+		for (FSlateDrawElement& Element : DrawElement.Value.Elements)
+		{
+			Element.AddReferencedObjects(Collector);
+		}
 	}
 }
 
@@ -1163,8 +1168,10 @@ FSlateDrawElement& FSlateCachedElementData::AddCachedElement(FSlateCachedElement
 	check(List->OwningWidget == CurrentWidget);
 	check(CurrentWidget->GetParentWidget().IsValid());
 #endif
-
-	FSlateDrawElement& NewElement = List->DrawElements.AddDefaulted_GetRef();
+	
+	// Add cached widget to NonMapped since we don't have an element type dedicated to it
+	FSlateDrawElementContainer& Container = List->DrawElements.FindOrAdd(EElementType::ET_NonMapped);
+	FSlateDrawElement& NewElement = Container.Elements.AddDefaulted_GetRef();
 	NewElement.SetIsCached(true);
 
 	// Check if slow vs checking a flag on the list to see if it contains new data.
@@ -1287,9 +1294,12 @@ void FSlateCachedElementList::AddCachedClipState(FSlateCachedClipState& ClipStat
 
 void FSlateCachedElementList::AddReferencedObjects(FReferenceCollector& Collector)
 {
-	for (FSlateDrawElement& Element : DrawElements)
+	for (TPair<EElementType, FSlateDrawElementContainer>& DrawElement : DrawElements)
 	{
-		Element.AddReferencedObjects(Collector);
+		for (FSlateDrawElement& Element : DrawElement.Value.Elements)
+		{
+			Element.AddReferencedObjects(Collector);
+		}
 	}
 }
 
diff --git a/Engine/Source/Runtime/SlateCore/Private/Rendering/ElementBatcher.cpp b/Engine/Source/Runtime/SlateCore/Private/Rendering/ElementBatcher.cpp
--- a/Engine/Source/Runtime/SlateCore/Private/Rendering/ElementBatcher.cpp
+++ b/Engine/Source/Runtime/SlateCore/Private/Rendering/ElementBatcher.cpp
@@ -353,98 +353,241 @@ void FSlateElementBatcher::AddElements(FSlateWindowElementList& WindowElementLis
 }
 
 
-void FSlateElementBatcher::AddElementsInternal(const FSlateDrawElementArray& DrawElements, FVector2f ViewportSize)
+void FSlateElementBatcher::AddElementsInternal(const FSlateDrawElementMap& DrawElements, FVector2f ViewportSize)
 {
-	for (const FSlateDrawElement& DrawElement : DrawElements)
+	// @TODO: DarenC = Warning this is ugly but right thing to do.
+	auto HandleUnMapped = [&](const FSlateDrawElementArray& Elements, FVector2f ViewportSize)
 	{
+		for (const FSlateDrawElement& DrawElement : Elements)
+		{
+			// Determine what type of element to add
+			switch (DrawElement.GetElementType())
+			{
+			case EElementType::ET_Box:
+			case EElementType::ET_RoundedBox:
+			{
+				SCOPED_NAMED_EVENT_TEXT("Slate::AddBoxElement", FColor::Magenta);
+				STAT(ElementStat_Boxes++);
+				DrawElement.IsPixelSnapped() ? AddBoxElement<ESlateVertexRounding::Enabled>(DrawElement) : AddBoxElement<ESlateVertexRounding::Disabled>(DrawElement);
+			}
+			break;
+			case EElementType::ET_Border:
+			{
+				SCOPED_NAMED_EVENT_TEXT("Slate::AddBorderElement", FColor::Magenta);
+				STAT(ElementStat_Borders++);
+				DrawElement.IsPixelSnapped() ? AddBorderElement<ESlateVertexRounding::Enabled>(DrawElement) : AddBorderElement<ESlateVertexRounding::Disabled>(DrawElement);
+			}
+			break;
+			case EElementType::ET_Text:
+			{
+				SCOPED_NAMED_EVENT_TEXT("Slate::AddTextElement", FColor::Magenta);
+				STAT(ElementStat_Text++);
+				DrawElement.IsPixelSnapped() ? AddTextElement<ESlateVertexRounding::Enabled>(DrawElement) : AddTextElement<ESlateVertexRounding::Disabled>(DrawElement);
+			}
+			break;
+			case EElementType::ET_ShapedText:
+			{
+				SCOPED_NAMED_EVENT_TEXT("Slate::AddShapedTextElement", FColor::Magenta);
+				STAT(ElementStat_ShapedText++);
+				DrawElement.IsPixelSnapped() ? AddShapedTextElement<ESlateVertexRounding::Enabled>(DrawElement) : AddShapedTextElement<ESlateVertexRounding::Disabled>(DrawElement);
+			}
+			break;
+			case EElementType::ET_Line:
+			{
+				SCOPED_NAMED_EVENT_TEXT("Slate::AddLineElement", FColor::Magenta);
+				STAT(ElementStat_Line++);
+				DrawElement.IsPixelSnapped() ? AddLineElement<ESlateVertexRounding::Enabled>(DrawElement) : AddLineElement<ESlateVertexRounding::Disabled>(DrawElement);
+			}
+			break;
+			case EElementType::ET_DebugQuad:
+			{
+				SCOPED_NAMED_EVENT_TEXT("Slate::AddDebugQuadElement", FColor::Magenta);
+				STAT(ElementStat_Other++);
+				DrawElement.IsPixelSnapped() ? AddDebugQuadElement<ESlateVertexRounding::Enabled>(DrawElement) : AddDebugQuadElement<ESlateVertexRounding::Disabled>(DrawElement);
+			}
+			break;
+			case EElementType::ET_Spline:
+			{
+				SCOPED_NAMED_EVENT_TEXT("Slate::AddSplineElement", FColor::Magenta);
+				// Note that we ignore pixel snapping here; see implementation for more info.
+				STAT(ElementStat_Other++);
+				AddSplineElement(DrawElement);
+			}
+			break;
+			case EElementType::ET_Gradient:
+			{
+				SCOPED_NAMED_EVENT_TEXT("Slate::AddGradientElement", FColor::Magenta);
+				STAT(ElementStat_Other++);
+				DrawElement.IsPixelSnapped() ? AddGradientElement<ESlateVertexRounding::Enabled>(DrawElement) : AddGradientElement<ESlateVertexRounding::Disabled>(DrawElement);
+			}
+			break;
+			case EElementType::ET_Viewport:
+			{
+				SCOPED_NAMED_EVENT_TEXT("Slate::AddViewportElement", FColor::Magenta);
+				STAT(ElementStat_Other++);
+				DrawElement.IsPixelSnapped() ? AddViewportElement<ESlateVertexRounding::Enabled>(DrawElement) : AddViewportElement<ESlateVertexRounding::Disabled>(DrawElement);
+			}
+			break;
+			case EElementType::ET_Custom:
+			{
+				SCOPED_NAMED_EVENT_TEXT("Slate::AddCustomElement", FColor::Magenta);
+				STAT(ElementStat_Other++);
+				AddCustomElement(DrawElement);
+			}
+			break;
+			case EElementType::ET_CustomVerts:
+			{
+				SCOPED_NAMED_EVENT_TEXT("Slate::AddCustomVertsElement", FColor::Magenta);
+				STAT(ElementStat_Other++);
+				AddCustomVerts(DrawElement);
+			}
+			break;
+			case EElementType::ET_PostProcessPass:
+			{
+				SCOPED_NAMED_EVENT_TEXT("Slate::AddPostProcessElement", FColor::Magenta);
+				STAT(ElementStat_Other++);
+				AddPostProcessPass(DrawElement, ViewportSize);
+			}
+			break;
+			default:
+				checkf(0, TEXT("Invalid element type"));
+			}
+		}
+	};
+
+	for (const TPair<EElementType, FSlateDrawElementContainer>& DrawElementMap : DrawElements)
+	{
+		EElementType ElmentType = DrawElementMap.Key;
+		const FSlateDrawElementContainer& Container = DrawElementMap.Value;
 		// Determine what type of element to add
-		switch ( DrawElement.GetElementType() )
+		switch (ElmentType)
 		{
 		case EElementType::ET_Box:
 		case EElementType::ET_RoundedBox:	
 		{
-			SCOPED_NAMED_EVENT_TEXT("Slate::AddBoxElement", FColor::Magenta);
-			STAT(ElementStat_Boxes++);
-			DrawElement.IsPixelSnapped() ? AddBoxElement<ESlateVertexRounding::Enabled>(DrawElement) : AddBoxElement<ESlateVertexRounding::Disabled>(DrawElement);
+			SCOPED_NAMED_EVENT_TEXT("Slate::AddRoundedBoxElement", FColor::Magenta);
+			STAT(ElementStat_Other += Container.Elements.Num());
+			for (const FSlateDrawElement& DrawElement : Container.Elements)
+			{
+				DrawElement.IsPixelSnapped() ? AddBoxElement<ESlateVertexRounding::Enabled>(DrawElement) : AddBoxElement<ESlateVertexRounding::Disabled>(DrawElement);
+			}
 		}
 			break;
 		case EElementType::ET_Border:
 		{
 			SCOPED_NAMED_EVENT_TEXT("Slate::AddBorderElement", FColor::Magenta);
-			STAT(ElementStat_Borders++);
-			DrawElement.IsPixelSnapped() ? AddBorderElement<ESlateVertexRounding::Enabled>(DrawElement) : AddBorderElement<ESlateVertexRounding::Disabled>(DrawElement);
+			STAT(ElementStat_Other += Container.Elements.Num());
+			for (const FSlateDrawElement& DrawElement : Container.Elements)
+			{
+				DrawElement.IsPixelSnapped() ? AddBorderElement<ESlateVertexRounding::Enabled>(DrawElement) : AddBorderElement<ESlateVertexRounding::Disabled>(DrawElement);
+			}
 		}
 			break;
 		case EElementType::ET_Text:
 		{
 			SCOPED_NAMED_EVENT_TEXT("Slate::AddTextElement", FColor::Magenta);
-			STAT(ElementStat_Text++);
-			DrawElement.IsPixelSnapped() ? AddTextElement<ESlateVertexRounding::Enabled>(DrawElement) : AddTextElement<ESlateVertexRounding::Disabled>(DrawElement);
+			STAT(ElementStat_Other += Container.Elements.Num());
+			for (const FSlateDrawElement& DrawElement : Container.Elements)
+			{
+				DrawElement.IsPixelSnapped() ? AddTextElement<ESlateVertexRounding::Enabled>(DrawElement) : AddTextElement<ESlateVertexRounding::Disabled>(DrawElement);
+			}
 		}
 			break;
 		case EElementType::ET_ShapedText:
 		{
 			SCOPED_NAMED_EVENT_TEXT("Slate::AddShapedTextElement", FColor::Magenta);
-			STAT(ElementStat_ShapedText++);
-			DrawElement.IsPixelSnapped() ? AddShapedTextElement<ESlateVertexRounding::Enabled>(DrawElement) : AddShapedTextElement<ESlateVertexRounding::Disabled>(DrawElement);
+			STAT(ElementStat_Other += Container.Elements.Num());
+			for (const FSlateDrawElement& DrawElement : Container.Elements)
+			{
+				DrawElement.IsPixelSnapped() ? AddShapedTextElement<ESlateVertexRounding::Enabled>(DrawElement) : AddShapedTextElement<ESlateVertexRounding::Disabled>(DrawElement);
+			}
 		}
 			break;
 		case EElementType::ET_Line:
 		{
 			SCOPED_NAMED_EVENT_TEXT("Slate::AddLineElement", FColor::Magenta);
-			STAT(ElementStat_Line++);
-			DrawElement.IsPixelSnapped() ? AddLineElement<ESlateVertexRounding::Enabled>(DrawElement) : AddLineElement<ESlateVertexRounding::Disabled>(DrawElement);
+			STAT(ElementStat_Other += Container.Elements.Num());
+			for (const FSlateDrawElement& DrawElement : Container.Elements)
+			{
+				DrawElement.IsPixelSnapped() ? AddLineElement<ESlateVertexRounding::Enabled>(DrawElement) : AddLineElement<ESlateVertexRounding::Disabled>(DrawElement);
+			}
 		}
 			break;
 		case EElementType::ET_DebugQuad:
 		{
 			SCOPED_NAMED_EVENT_TEXT("Slate::AddDebugQuadElement", FColor::Magenta);
-			STAT(ElementStat_Other++);
-			DrawElement.IsPixelSnapped() ? AddDebugQuadElement<ESlateVertexRounding::Enabled>(DrawElement) : AddDebugQuadElement<ESlateVertexRounding::Disabled>(DrawElement);
+			STAT(ElementStat_Other += Container.Elements.Num());
+			for (const FSlateDrawElement& DrawElement : Container.Elements)
+			{
+				DrawElement.IsPixelSnapped() ? AddDebugQuadElement<ESlateVertexRounding::Enabled>(DrawElement) : AddDebugQuadElement<ESlateVertexRounding::Disabled>(DrawElement);
+			}
 		}
 			break;
 		case EElementType::ET_Spline:
 		{
 			SCOPED_NAMED_EVENT_TEXT("Slate::AddSplineElement", FColor::Magenta);
 			// Note that we ignore pixel snapping here; see implementation for more info.
-			STAT(ElementStat_Other++);
-			AddSplineElement(DrawElement);
+			STAT(ElementStat_Other += Container.Elements.Num());
+			for (const FSlateDrawElement& DrawElement : Container.Elements)
+			{
+				AddSplineElement(DrawElement);
+			}
 		}
 			break;
 		case EElementType::ET_Gradient:
 		{
 			SCOPED_NAMED_EVENT_TEXT("Slate::AddGradientElement", FColor::Magenta);
-			STAT(ElementStat_Other++);
-			DrawElement.IsPixelSnapped() ? AddGradientElement<ESlateVertexRounding::Enabled>(DrawElement) : AddGradientElement<ESlateVertexRounding::Disabled>(DrawElement);
+			STAT(ElementStat_Other += Container.Elements.Num());
+			for (const FSlateDrawElement& DrawElement : Container.Elements)
+			{
+				DrawElement.IsPixelSnapped() ? AddGradientElement<ESlateVertexRounding::Enabled>(DrawElement) : AddGradientElement<ESlateVertexRounding::Disabled>(DrawElement);
+			}
 		}
 			break;
 		case EElementType::ET_Viewport:
 		{
 			SCOPED_NAMED_EVENT_TEXT("Slate::AddViewportElement", FColor::Magenta);
-			STAT(ElementStat_Other++);
-			DrawElement.IsPixelSnapped() ? AddViewportElement<ESlateVertexRounding::Enabled>(DrawElement) : AddViewportElement<ESlateVertexRounding::Disabled>(DrawElement);
+			STAT(ElementStat_Other += Container.Elements.Num());
+			for (const FSlateDrawElement& DrawElement : Container.Elements)
+			{
+				DrawElement.IsPixelSnapped() ? AddViewportElement<ESlateVertexRounding::Enabled>(DrawElement) : AddViewportElement<ESlateVertexRounding::Disabled>(DrawElement);
+			}
 		}
 			break;
 		case EElementType::ET_Custom:
 		{
 			SCOPED_NAMED_EVENT_TEXT("Slate::AddCustomElement", FColor::Magenta);
-			STAT(ElementStat_Other++);
-			AddCustomElement(DrawElement);
+			STAT(ElementStat_Other += Container.Elements.Num());
+			for (const FSlateDrawElement& DrawElement : Container.Elements)
+			{
+				AddCustomElement(DrawElement);
+			}
 		}
 			break;
 		case EElementType::ET_CustomVerts:
 		{
 			SCOPED_NAMED_EVENT_TEXT("Slate::AddCustomVertsElement", FColor::Magenta);
-			STAT(ElementStat_Other++);
-			AddCustomVerts(DrawElement);
+			STAT(ElementStat_Other += Container.Elements.Num());
+			for (const FSlateDrawElement& DrawElement : Container.Elements)
+			{
+				AddCustomVerts(DrawElement);
+			}
 		}
 			break;
-	
 		case EElementType::ET_PostProcessPass:
 		{
 			SCOPED_NAMED_EVENT_TEXT("Slate::AddPostProcessElement", FColor::Magenta);
-			STAT(ElementStat_Other++);
-			AddPostProcessPass(DrawElement, ViewportSize);
+			STAT(ElementStat_Other += Container.Elements.Num());
+			for (const FSlateDrawElement& DrawElement : Container.Elements)
+			{
+				AddPostProcessPass(DrawElement, ViewportSize);
+			}
+		}
+			break;
+		case EElementType::ET_NonMapped:
+		{
+			SCOPED_NAMED_EVENT_TEXT("Slate::AddNonMappedElement", FColor::Magenta);
+			HandleUnMapped(Container.Elements, ViewportSize);
 		}
 			break;
 		default:
@@ -759,71 +902,83 @@ void FSlateElementBatcher::AddBoxElement(const FSlateDrawElement& DrawElement)
 		RenderBatch.AddVertex( FSlateVertex::Make<Rounding>( RenderTransform, FVector2f( RightMarginX, EndPos.Y ),		LocalSize, DrawScale, FVector4f(FVector2f( RightMarginU, EndUV.Y ),			Tiling),	Tint, SecondaryColor) ); //14
 		RenderBatch.AddVertex( FSlateVertex::Make<Rounding>( RenderTransform, FVector2f( EndPos.X, EndPos.Y ),			LocalSize, DrawScale, FVector4f(FVector2f( EndUV ),							Tiling),	Tint, SecondaryColor) ); //15
 
-		// Top
-		RenderBatch.AddIndex( IndexStart + 0 );
-		RenderBatch.AddIndex( IndexStart + 1 );
-		RenderBatch.AddIndex( IndexStart + 2 );
-		RenderBatch.AddIndex( IndexStart + 2 );
-		RenderBatch.AddIndex( IndexStart + 1 );
-		RenderBatch.AddIndex( IndexStart + 3 );
-
-		RenderBatch.AddIndex( IndexStart + 2 );
-		RenderBatch.AddIndex( IndexStart + 3 );
-		RenderBatch.AddIndex( IndexStart + 4 );
-		RenderBatch.AddIndex( IndexStart + 4 );
-		RenderBatch.AddIndex( IndexStart + 3 );
-		RenderBatch.AddIndex( IndexStart + 5 );
-
-		RenderBatch.AddIndex( IndexStart + 4 );
-		RenderBatch.AddIndex( IndexStart + 5 );
-		RenderBatch.AddIndex( IndexStart + 6 );
-		RenderBatch.AddIndex( IndexStart + 6 );
-		RenderBatch.AddIndex( IndexStart + 5 );
-		RenderBatch.AddIndex( IndexStart + 7 );
+		static FSlateIndexArray DefaultPaddedBoxIndiciesArray = []()
+		{
+			const int32 IndexStart = 0;
+			FSlateIndexArray DefaultPaddedBoxIndicies = {};
+			DefaultPaddedBoxIndicies.Reserve(6 * 3 * 3);
 
-		// Middle
-		RenderBatch.AddIndex( IndexStart + 1 );
-		RenderBatch.AddIndex( IndexStart + 8 );
-		RenderBatch.AddIndex( IndexStart + 3 );
-		RenderBatch.AddIndex( IndexStart + 3 );
-		RenderBatch.AddIndex( IndexStart + 8 );
-		RenderBatch.AddIndex( IndexStart + 9 );
-
-		RenderBatch.AddIndex( IndexStart + 3 );
-		RenderBatch.AddIndex( IndexStart + 9 );
-		RenderBatch.AddIndex( IndexStart + 5 );
-		RenderBatch.AddIndex( IndexStart + 5 );
-		RenderBatch.AddIndex( IndexStart + 9 );
-		RenderBatch.AddIndex( IndexStart + 10 );
-
-		RenderBatch.AddIndex( IndexStart + 5 );
-		RenderBatch.AddIndex( IndexStart + 10 );
-		RenderBatch.AddIndex( IndexStart + 7 );
-		RenderBatch.AddIndex( IndexStart + 7 );
-		RenderBatch.AddIndex( IndexStart + 10 );
-		RenderBatch.AddIndex( IndexStart + 11 );
+			// Top
+			DefaultPaddedBoxIndicies.Add(IndexStart + 0);
+			DefaultPaddedBoxIndicies.Add(IndexStart + 1);
+			DefaultPaddedBoxIndicies.Add(IndexStart + 2);
+			DefaultPaddedBoxIndicies.Add(IndexStart + 2);
+			DefaultPaddedBoxIndicies.Add(IndexStart + 1);
+			DefaultPaddedBoxIndicies.Add(IndexStart + 3);
+
+			DefaultPaddedBoxIndicies.Add(IndexStart + 2);
+			DefaultPaddedBoxIndicies.Add(IndexStart + 3);
+			DefaultPaddedBoxIndicies.Add(IndexStart + 4);
+			DefaultPaddedBoxIndicies.Add(IndexStart + 4);
+			DefaultPaddedBoxIndicies.Add(IndexStart + 3);
+			DefaultPaddedBoxIndicies.Add(IndexStart + 5);
+
+			DefaultPaddedBoxIndicies.Add(IndexStart + 4);
+			DefaultPaddedBoxIndicies.Add(IndexStart + 5);
+			DefaultPaddedBoxIndicies.Add(IndexStart + 6);
+			DefaultPaddedBoxIndicies.Add(IndexStart + 6);
+			DefaultPaddedBoxIndicies.Add(IndexStart + 5);
+			DefaultPaddedBoxIndicies.Add(IndexStart + 7);
+
+			// Middle
+			DefaultPaddedBoxIndicies.Add(IndexStart + 1);
+			DefaultPaddedBoxIndicies.Add(IndexStart + 8);
+			DefaultPaddedBoxIndicies.Add(IndexStart + 3);
+			DefaultPaddedBoxIndicies.Add(IndexStart + 3);
+			DefaultPaddedBoxIndicies.Add(IndexStart + 8);
+			DefaultPaddedBoxIndicies.Add(IndexStart + 9);
+
+			DefaultPaddedBoxIndicies.Add(IndexStart + 3);
+			DefaultPaddedBoxIndicies.Add(IndexStart + 9);
+			DefaultPaddedBoxIndicies.Add(IndexStart + 5);
+			DefaultPaddedBoxIndicies.Add(IndexStart + 5);
+			DefaultPaddedBoxIndicies.Add(IndexStart + 9);
+			DefaultPaddedBoxIndicies.Add(IndexStart + 10);
+
+			DefaultPaddedBoxIndicies.Add(IndexStart + 5);
+			DefaultPaddedBoxIndicies.Add(IndexStart + 10);
+			DefaultPaddedBoxIndicies.Add(IndexStart + 7);
+			DefaultPaddedBoxIndicies.Add(IndexStart + 7);
+			DefaultPaddedBoxIndicies.Add(IndexStart + 10);
+			DefaultPaddedBoxIndicies.Add(IndexStart + 11);
 
-		// Bottom
-		RenderBatch.AddIndex( IndexStart + 8 );
-		RenderBatch.AddIndex( IndexStart + 12 );
-		RenderBatch.AddIndex( IndexStart + 9 );
-		RenderBatch.AddIndex( IndexStart + 9 );
-		RenderBatch.AddIndex( IndexStart + 12 );
-		RenderBatch.AddIndex( IndexStart + 13 );
-
-		RenderBatch.AddIndex( IndexStart + 9 );
-		RenderBatch.AddIndex( IndexStart + 13 );
-		RenderBatch.AddIndex( IndexStart + 10 );
-		RenderBatch.AddIndex( IndexStart + 10 );
-		RenderBatch.AddIndex( IndexStart + 13 );
-		RenderBatch.AddIndex( IndexStart + 14 );
-
-		RenderBatch.AddIndex( IndexStart + 10 );
-		RenderBatch.AddIndex( IndexStart + 14 );
-		RenderBatch.AddIndex( IndexStart + 11 );
-		RenderBatch.AddIndex( IndexStart + 11 );
-		RenderBatch.AddIndex( IndexStart + 14 );
-		RenderBatch.AddIndex( IndexStart + 15 );
+			// Bottom
+			DefaultPaddedBoxIndicies.Add(IndexStart + 8);
+			DefaultPaddedBoxIndicies.Add(IndexStart + 12);
+			DefaultPaddedBoxIndicies.Add(IndexStart + 9);
+			DefaultPaddedBoxIndicies.Add(IndexStart + 9);
+			DefaultPaddedBoxIndicies.Add(IndexStart + 12);
+			DefaultPaddedBoxIndicies.Add(IndexStart + 13);
+
+			DefaultPaddedBoxIndicies.Add(IndexStart + 9);
+			DefaultPaddedBoxIndicies.Add(IndexStart + 13);
+			DefaultPaddedBoxIndicies.Add(IndexStart + 10);
+			DefaultPaddedBoxIndicies.Add(IndexStart + 10);
+			DefaultPaddedBoxIndicies.Add(IndexStart + 13);
+			DefaultPaddedBoxIndicies.Add(IndexStart + 14);
+
+			DefaultPaddedBoxIndicies.Add(IndexStart + 10);
+			DefaultPaddedBoxIndicies.Add(IndexStart + 14);
+			DefaultPaddedBoxIndicies.Add(IndexStart + 11);
+			DefaultPaddedBoxIndicies.Add(IndexStart + 11);
+			DefaultPaddedBoxIndicies.Add(IndexStart + 14);
+			DefaultPaddedBoxIndicies.Add(IndexStart + 15);
+
+			return DefaultPaddedBoxIndicies;
+		}();
+		RenderBatch.SourceIndices = &DefaultPaddedBoxIndiciesArray;
+		RenderBatch.NumIndices = DefaultPaddedBoxIndiciesArray.Num();
+		RenderBatch.IndexOffset = 0;
 
 		if ( GSlateFeathering && Rounding == ESlateVertexRounding::Disabled )
 		{
@@ -849,39 +1004,63 @@ void FSlateElementBatcher::AddBoxElement(const FSlateDrawElement& DrawElement)
 			RenderBatch.AddVertex(FSlateVertex::Make<Rounding>(RenderTransform, FVector2f(RightMarginX, EndPos.Y) + FVector2f(0.f, 1.f) / DrawScale,		LocalSize, DrawScale, FVector4f(FVector2f(RightMarginU, EndUV.Y), Tiling), FeatherColor, SecondaryColor)); //10
 			RenderBatch.AddVertex(FSlateVertex::Make<Rounding>(RenderTransform, FVector2f(EndPos.X, EndPos.Y) + FVector2f(1.f, 1.f) / DrawScale,			LocalSize, DrawScale, FVector4f(FVector2f(EndUV), Tiling), FeatherColor, SecondaryColor)); //11
 
-			// Top Left
-			IndexQuad(RenderBatch, FeatherStart + 0, FeatherStart + 1, IndexStart + 2, IndexStart + 0);
-			// Top Middle
-			IndexQuad(RenderBatch, FeatherStart + 1, FeatherStart + 2, IndexStart + 4, IndexStart + 2);
-			// Top Right
-			IndexQuad(RenderBatch, FeatherStart + 2, FeatherStart + 3, IndexStart + 6, IndexStart + 4);
-
-			//-----------------------------------------------------------
-
-			// Left Top
-			IndexQuad(RenderBatch, FeatherStart + 0, IndexStart + 0, IndexStart + 1, FeatherStart + 4);
-			// Left Middle
-			IndexQuad(RenderBatch, FeatherStart + 4, IndexStart + 1, IndexStart + 8, FeatherStart + 5);
-			// Left Bottom
-			IndexQuad(RenderBatch, FeatherStart + 5, IndexStart + 8, IndexStart + 12, FeatherStart + 8);
-
-			//-----------------------------------------------------------
-
-			// Right Top
-			IndexQuad(RenderBatch, IndexStart + 6, FeatherStart + 3, FeatherStart + 6, IndexStart + 7);
-			// Right Middle
-			IndexQuad(RenderBatch, IndexStart + 7, FeatherStart + 6, FeatherStart + 7, IndexStart + 11);
-			// Right Bottom
-			IndexQuad(RenderBatch, IndexStart + 11, FeatherStart + 7, FeatherStart + 11, IndexStart + 15);
-
-			//-----------------------------------------------------------
+			static FSlateIndexArray FeatheredPaddedBoxIndiciesArray = [DefaultPaddedBoxIndiciesArray = DefaultPaddedBoxIndiciesArray]()
+			{
+				const int32 IndexStart = 0;
+				const int32 FeatherStart = 16;
+				FSlateIndexArray FeatheredPaddedBoxIndicies = DefaultPaddedBoxIndiciesArray;
+				FeatheredPaddedBoxIndicies.Reserve(FeatheredPaddedBoxIndicies.Num() + 6 * 3 * 4);
 
-			// Bottom Left
-			IndexQuad(RenderBatch, IndexStart + 12, IndexStart + 13, FeatherStart + 9, FeatherStart + 8);
-			// Bottom Middle
-			IndexQuad(RenderBatch, IndexStart + 13, IndexStart + 14, FeatherStart + 10, FeatherStart + 9);
-			// Bottom Right
-			IndexQuad(RenderBatch, IndexStart + 14, IndexStart + 15, FeatherStart + 11, FeatherStart + 10);
+				auto IndexQuadArray = [](FSlateIndexArray& IndexArray, int32 TopLeft, int32 TopRight, int32 BottomRight, int32 BottomLeft)
+				{
+					IndexArray.Add(TopLeft);
+					IndexArray.Add(TopRight);
+					IndexArray.Add(BottomRight);
+
+					IndexArray.Add(BottomRight);
+					IndexArray.Add(BottomLeft);
+					IndexArray.Add(TopLeft);
+				};
+
+				// Top Left
+				IndexQuadArray(FeatheredPaddedBoxIndicies, FeatherStart + 0, FeatherStart + 1, IndexStart + 2, IndexStart + 0);
+				// Top Middle
+				IndexQuadArray(FeatheredPaddedBoxIndicies, FeatherStart + 1, FeatherStart + 2, IndexStart + 4, IndexStart + 2);
+				// Top Right
+				IndexQuadArray(FeatheredPaddedBoxIndicies, FeatherStart + 2, FeatherStart + 3, IndexStart + 6, IndexStart + 4);
+
+				//-----------------------------------------------------------
+
+				// Left Top
+				IndexQuadArray(FeatheredPaddedBoxIndicies, FeatherStart + 0, IndexStart + 0, IndexStart + 1, FeatherStart + 4);
+				// Left Middle
+				IndexQuadArray(FeatheredPaddedBoxIndicies, FeatherStart + 4, IndexStart + 1, IndexStart + 8, FeatherStart + 5);
+				// Left Bottom
+				IndexQuadArray(FeatheredPaddedBoxIndicies, FeatherStart + 5, IndexStart + 8, IndexStart + 12, FeatherStart + 8);
+
+				//-----------------------------------------------------------
+
+				// Right Top
+				IndexQuadArray(FeatheredPaddedBoxIndicies, IndexStart + 6, FeatherStart + 3, FeatherStart + 6, IndexStart + 7);
+				// Right Middle
+				IndexQuadArray(FeatheredPaddedBoxIndicies, IndexStart + 7, FeatherStart + 6, FeatherStart + 7, IndexStart + 11);
+				// Right Bottom
+				IndexQuadArray(FeatheredPaddedBoxIndicies, IndexStart + 11, FeatherStart + 7, FeatherStart + 11, IndexStart + 15);
+
+				//-----------------------------------------------------------
+
+				// Bottom Left
+				IndexQuadArray(FeatheredPaddedBoxIndicies, IndexStart + 12, IndexStart + 13, FeatherStart + 9, FeatherStart + 8);
+				// Bottom Middle
+				IndexQuadArray(FeatheredPaddedBoxIndicies, IndexStart + 13, IndexStart + 14, FeatherStart + 10, FeatherStart + 9);
+				// Bottom Right
+				IndexQuadArray(FeatheredPaddedBoxIndicies, IndexStart + 14, IndexStart + 15, FeatherStart + 11, FeatherStart + 10);
+
+				return FeatheredPaddedBoxIndicies;
+			}();
+			RenderBatch.SourceIndices = &FeatheredPaddedBoxIndiciesArray;
+			RenderBatch.NumIndices = FeatheredPaddedBoxIndiciesArray.Num();
+			RenderBatch.IndexOffset = 0;
 		}
 	}
 	else
@@ -909,18 +1088,25 @@ void FSlateElementBatcher::AddBoxElement(const FSlateDrawElement& DrawElement)
 		RenderBatch.AddVertex(FSlateVertex::Make<Rounding>(RenderTransform, BotLeft, LocalSize, DrawScale, FVector4f(FVector2f(StartUV.X, EndUV.Y), Tiling), Tint, SecondaryColor));
 		RenderBatch.AddVertex(FSlateVertex::Make<Rounding>(RenderTransform, BotRight, LocalSize, DrawScale, FVector4f(FVector2f(EndUV), Tiling), Tint, SecondaryColor));
 
-		RenderBatch.AddIndex( IndexStart + 0 );
-		RenderBatch.AddIndex( IndexStart + 1 );
-		RenderBatch.AddIndex( IndexStart + 2 );
+		static FSlateIndexArray DefaultBoxIndiciesArray = []()
+		{
+			const int32 IndexStart = 0;
+			FSlateIndexArray DefaultBoxIndicies = {};
+			DefaultBoxIndicies.Reserve(6);
+
+			DefaultBoxIndicies.Add(IndexStart + 0);
+			DefaultBoxIndicies.Add(IndexStart + 1);
+			DefaultBoxIndicies.Add(IndexStart + 2);
 
-		RenderBatch.AddIndex( IndexStart + 2 );
-		RenderBatch.AddIndex( IndexStart + 1 );
-		RenderBatch.AddIndex( IndexStart + 3 );
+			DefaultBoxIndicies.Add(IndexStart + 2);
+			DefaultBoxIndicies.Add(IndexStart + 1);
+			DefaultBoxIndicies.Add(IndexStart + 3);
 
-		int32 TopLeftIndex = IndexStart + 0;
-		int32 TopRightIndex = IndexStart + 1;
-		int32 BottomLeftIndex = IndexStart + 2;
-		int32 BottomRightIndex = IndexStart + 3;
+			return DefaultBoxIndicies;
+		}();
+		RenderBatch.SourceIndices = &DefaultBoxIndiciesArray;
+		RenderBatch.NumIndices = DefaultBoxIndiciesArray.Num();
+		RenderBatch.IndexOffset = 0;
 
 		if ( GSlateFeathering && Rounding == ESlateVertexRounding::Disabled )
 		{
@@ -931,48 +1117,68 @@ void FSlateElementBatcher::AddBoxElement(const FSlateDrawElement& DrawElement)
 			RenderBatch.AddVertex(FSlateVertex::Make<Rounding>(RenderTransform, BotLeft + FVector2f(-1.f, 1.f) / DrawScale, LocalSize, DrawScale, FVector4f(FVector2f(StartUV.X, EndUV.Y), Tiling), FeatherColor, SecondaryColor));
 			RenderBatch.AddVertex(FSlateVertex::Make<Rounding>(RenderTransform, BotRight + FVector2f(1.f, 1.f) / DrawScale, LocalSize, DrawScale, FVector4f(FVector2f(EndUV), Tiling), FeatherColor, SecondaryColor));
 
-			// Top-Top
-			RenderBatch.AddIndex(FeatherStart + 0);
-			RenderBatch.AddIndex(FeatherStart + 1);
-			RenderBatch.AddIndex(TopRightIndex);
-
-			// Top-Bottom
-			RenderBatch.AddIndex(FeatherStart + 0);
-			RenderBatch.AddIndex(TopRightIndex);
-			RenderBatch.AddIndex(TopLeftIndex);
-
-			// Left-Top
-			RenderBatch.AddIndex(FeatherStart + 0);
-			RenderBatch.AddIndex(BottomLeftIndex);
-			RenderBatch.AddIndex(FeatherStart + 2);
-
-			// Left-Bottom
-			RenderBatch.AddIndex(FeatherStart + 0);
-			RenderBatch.AddIndex(TopLeftIndex);
-			RenderBatch.AddIndex(BottomLeftIndex);
-
-			// Right-Top
-			RenderBatch.AddIndex(TopRightIndex);
-			RenderBatch.AddIndex(FeatherStart + 1);
-			RenderBatch.AddIndex(FeatherStart + 3);
-
-			// Right-Bottom
-			RenderBatch.AddIndex(TopRightIndex);
-			RenderBatch.AddIndex(FeatherStart + 3);
-			RenderBatch.AddIndex(BottomRightIndex);
-
-			// Bottom-Top
-			RenderBatch.AddIndex(BottomLeftIndex);
-			RenderBatch.AddIndex(BottomRightIndex);
-			RenderBatch.AddIndex(FeatherStart + 3);
-
-			// Bottom-Bottom
-			RenderBatch.AddIndex(FeatherStart + 3);
-			RenderBatch.AddIndex(FeatherStart + 2);
-			RenderBatch.AddIndex(BottomLeftIndex);
+			static FSlateIndexArray FeatheredBoxIndiciesArray = [DefaultBoxIndiciesArray = DefaultBoxIndiciesArray]()
+			{
+				const int32 IndexStart = 0;
+				const int32 FeatherStart = 4;
+
+				const int32 TopLeftIndex = IndexStart + 0;
+				const int32 TopRightIndex = IndexStart + 1;
+				const int32 BottomLeftIndex = IndexStart + 2;
+				const int32 BottomRightIndex = IndexStart + 3;
+
+				FSlateIndexArray FeatheredBoxIndicies = DefaultBoxIndiciesArray;
+				FeatheredBoxIndicies.Reserve(FeatheredBoxIndicies.Num() + 3 * 8);
+
+				// Top-Top
+				FeatheredBoxIndicies.Add(FeatherStart + 0);
+				FeatheredBoxIndicies.Add(FeatherStart + 1);
+				FeatheredBoxIndicies.Add(TopRightIndex);
+
+				// Top-Bottom
+				FeatheredBoxIndicies.Add(FeatherStart + 0);
+				FeatheredBoxIndicies.Add(TopRightIndex);
+				FeatheredBoxIndicies.Add(TopLeftIndex);
+
+				// Left-Top
+				FeatheredBoxIndicies.Add(FeatherStart + 0);
+				FeatheredBoxIndicies.Add(BottomLeftIndex);
+				FeatheredBoxIndicies.Add(FeatherStart + 2);
+
+				// Left-Bottom
+				FeatheredBoxIndicies.Add(FeatherStart + 0);
+				FeatheredBoxIndicies.Add(TopLeftIndex);
+				FeatheredBoxIndicies.Add(BottomLeftIndex);
+
+				// Right-Top
+				FeatheredBoxIndicies.Add(TopRightIndex);
+				FeatheredBoxIndicies.Add(FeatherStart + 1);
+				FeatheredBoxIndicies.Add(FeatherStart + 3);
+
+				// Right-Bottom
+				FeatheredBoxIndicies.Add(TopRightIndex);
+				FeatheredBoxIndicies.Add(FeatherStart + 3);
+				FeatheredBoxIndicies.Add(BottomRightIndex);
+
+				// Bottom-Top
+				FeatheredBoxIndicies.Add(BottomLeftIndex);
+				FeatheredBoxIndicies.Add(BottomRightIndex);
+				FeatheredBoxIndicies.Add(FeatherStart + 3);
+
+				// Bottom-Bottom
+				FeatheredBoxIndicies.Add(FeatherStart + 3);
+				FeatheredBoxIndicies.Add(FeatherStart + 2);
+				FeatheredBoxIndicies.Add(BottomLeftIndex);
+
+				return FeatheredBoxIndicies;
+			}();
+			RenderBatch.SourceIndices = &FeatheredBoxIndiciesArray;
+			RenderBatch.NumIndices = FeatheredBoxIndiciesArray.Num();
+			RenderBatch.IndexOffset = 0;
 		}
 	}
 }
+
 namespace SlateElementBatcher
 {
 #if SLATE_CHECK_UOBJECT_RENDER_RESOURCES
@@ -1904,6 +2110,10 @@ struct FLineBuilder
 		// Build the start cap at the first point
 		MakeStartCap<Rounding>(Position, Direction, Length, Up, PointColor);
 
+		// @TODO: Since the vertex - index relationship is not homogenous whenever 
+		// a miter angle break occurs we cannot pre-allocate indicies here
+		// However if we were to create a new renderbatch on this break then it would be possible
+
 		// Build the intermediate points and their incoming segments
 		const int32 LastPointIndex = Points.Num() - 1;
 		for (int32 Point = 1; Point < LastPointIndex; ++Point)
@@ -2133,11 +2343,42 @@ void FSlateElementBatcher::AddLineElement( const FSlateDrawElement& DrawElement
 			{
 				const FColor PointColor = PackedColors.Num() ? PackedColors[Point] : PackedTint;
 				RenderBatch.AddVertex(FSlateVertex::Make<Rounding>(RenderTransform, Points[Point], FVector2f::ZeroVector, PointColor));
+			}
+
+			// The number of points can technically be unbound, so pick some number that if smaller we use the preallocated index array
+			const int32 MaxPreAllocatedPointIndicies = 4096 + 1024;
 
-				const uint32 IndexStart = RenderBatch.GetNumVertices();
-				RenderBatch.AddIndex(IndexStart - 2);
-				RenderBatch.AddIndex(IndexStart - 1);
+			if (NumPoints < MaxPreAllocatedPointIndicies)
+			{
+				static FSlateIndexArray DefaultLineIndiciesArray = [MaxPreAllocatedPointIndicies = MaxPreAllocatedPointIndicies]()
+				{
+					FSlateIndexArray DefaultLineIndicies = {};
+					DefaultLineIndicies.Reserve(MaxPreAllocatedPointIndicies * 2 - 2);
+
+					for (int32 Point = 1; Point < MaxPreAllocatedPointIndicies; ++Point)
+					{
+						const uint32 IndexStart = 1 + Point;
+						DefaultLineIndicies.Add(IndexStart - 2);
+						DefaultLineIndicies.Add(IndexStart - 1);
+					}
+
+					return DefaultLineIndicies;
+				}();
+
+				RenderBatch.SourceIndices = &DefaultLineIndiciesArray;
+				RenderBatch.NumIndices = NumPoints * 2 - 2;
+				RenderBatch.IndexOffset = 0;
 			}
+			else
+			{
+				for (int32 Point = 1; Point < NumPoints; ++Point)
+				{
+					const uint32 IndexStart = 1 + Point;
+					RenderBatch.AddIndex(IndexStart - 2);
+					RenderBatch.AddIndex(IndexStart - 1);
+				}
+			}
+
 		}
 		else
 		{
@@ -2416,64 +2657,76 @@ void FSlateElementBatcher::AddBorderElement( const FSlateDrawElement& DrawElemen
 	RenderBatch.AddVertex( FSlateVertex::Make<Rounding>( RenderTransform, FVector2f( EndPos.X, BottomMarginY ),		LocalSize, DrawScale, FVector4f( EndUV.X, BottomMarginV, 0.0f, 0.0f),			Tint ) ); //31
 	RenderBatch.AddVertex( FSlateVertex::Make<Rounding>( RenderTransform, FVector2f( EndPos.X, EndPos.Y ),			LocalSize, DrawScale, FVector4f( EndUV.X, EndUV.Y, 0.0f, 0.0f),					Tint ) ); //32
 
-	// Top
-	RenderBatch.AddIndex( IndexStart + 0 );
-	RenderBatch.AddIndex( IndexStart + 1 );
-	RenderBatch.AddIndex( IndexStart + 2 );
-	RenderBatch.AddIndex( IndexStart + 2 );
-	RenderBatch.AddIndex( IndexStart + 1 );
-	RenderBatch.AddIndex( IndexStart + 3 );
+	static FSlateIndexArray DefaultBorderIndiciesArray = []()
+	{
+		const int32 IndexStart = 0;
+		FSlateIndexArray DefaultBorderIndicies = {};
+		DefaultBorderIndicies.Reserve(6 * (3 + 2 + 3));
+
+		// Top
+		DefaultBorderIndicies.Add(IndexStart + 0);
+		DefaultBorderIndicies.Add(IndexStart + 1);
+		DefaultBorderIndicies.Add(IndexStart + 2);
+		DefaultBorderIndicies.Add(IndexStart + 2);
+		DefaultBorderIndicies.Add(IndexStart + 1);
+		DefaultBorderIndicies.Add(IndexStart + 3);
+
+		DefaultBorderIndicies.Add(IndexStart + 4);
+		DefaultBorderIndicies.Add(IndexStart + 5);
+		DefaultBorderIndicies.Add(IndexStart + 6);
+		DefaultBorderIndicies.Add(IndexStart + 6);
+		DefaultBorderIndicies.Add(IndexStart + 5);
+		DefaultBorderIndicies.Add(IndexStart + 7);
+
+		DefaultBorderIndicies.Add(IndexStart + 8);
+		DefaultBorderIndicies.Add(IndexStart + 9);
+		DefaultBorderIndicies.Add(IndexStart + 10);
+		DefaultBorderIndicies.Add(IndexStart + 10);
+		DefaultBorderIndicies.Add(IndexStart + 9);
+		DefaultBorderIndicies.Add(IndexStart + 11);
 
-	RenderBatch.AddIndex( IndexStart + 4 );
-	RenderBatch.AddIndex( IndexStart + 5 );
-	RenderBatch.AddIndex( IndexStart + 6 );
-	RenderBatch.AddIndex( IndexStart + 6 );
-	RenderBatch.AddIndex( IndexStart + 5 );
-	RenderBatch.AddIndex( IndexStart + 7 );
-
-	RenderBatch.AddIndex( IndexStart + 8 );
-	RenderBatch.AddIndex( IndexStart + 9 );
-	RenderBatch.AddIndex( IndexStart + 10 );
-	RenderBatch.AddIndex( IndexStart + 10 );
-	RenderBatch.AddIndex( IndexStart + 9 );
-	RenderBatch.AddIndex( IndexStart + 11 );
-
-	// Middle
-	RenderBatch.AddIndex( IndexStart + 12 );
-	RenderBatch.AddIndex( IndexStart + 13 );
-	RenderBatch.AddIndex( IndexStart + 14 );
-	RenderBatch.AddIndex( IndexStart + 14 );
-	RenderBatch.AddIndex( IndexStart + 13 );
-	RenderBatch.AddIndex( IndexStart + 15 );
-
-	RenderBatch.AddIndex( IndexStart + 16 );
-	RenderBatch.AddIndex( IndexStart + 17 );
-	RenderBatch.AddIndex( IndexStart + 18 );
-	RenderBatch.AddIndex( IndexStart + 18 );
-	RenderBatch.AddIndex( IndexStart + 17 );
-	RenderBatch.AddIndex( IndexStart + 19 );
-
-	// Bottom
-	RenderBatch.AddIndex( IndexStart + 20 );
-	RenderBatch.AddIndex( IndexStart + 21 );
-	RenderBatch.AddIndex( IndexStart + 22 );
-	RenderBatch.AddIndex( IndexStart + 22 );
-	RenderBatch.AddIndex( IndexStart + 21 );
-	RenderBatch.AddIndex( IndexStart + 23 );
-
-	RenderBatch.AddIndex( IndexStart + 24 );
-	RenderBatch.AddIndex( IndexStart + 25 );
-	RenderBatch.AddIndex( IndexStart + 26 );
-	RenderBatch.AddIndex( IndexStart + 26 );
-	RenderBatch.AddIndex( IndexStart + 25 );
-	RenderBatch.AddIndex( IndexStart + 27 );
-
-	RenderBatch.AddIndex( IndexStart + 28 );
-	RenderBatch.AddIndex( IndexStart + 29 );
-	RenderBatch.AddIndex( IndexStart + 30 );
-	RenderBatch.AddIndex( IndexStart + 30 );
-	RenderBatch.AddIndex( IndexStart + 29 );
-	RenderBatch.AddIndex( IndexStart + 31 );
+		// Middle
+		DefaultBorderIndicies.Add(IndexStart + 12);
+		DefaultBorderIndicies.Add(IndexStart + 13);
+		DefaultBorderIndicies.Add(IndexStart + 14);
+		DefaultBorderIndicies.Add(IndexStart + 14);
+		DefaultBorderIndicies.Add(IndexStart + 13);
+		DefaultBorderIndicies.Add(IndexStart + 15);
+
+		DefaultBorderIndicies.Add(IndexStart + 16);
+		DefaultBorderIndicies.Add(IndexStart + 17);
+		DefaultBorderIndicies.Add(IndexStart + 18);
+		DefaultBorderIndicies.Add(IndexStart + 18);
+		DefaultBorderIndicies.Add(IndexStart + 17);
+		DefaultBorderIndicies.Add(IndexStart + 19);
+
+		// Bottom
+		DefaultBorderIndicies.Add(IndexStart + 20);
+		DefaultBorderIndicies.Add(IndexStart + 21);
+		DefaultBorderIndicies.Add(IndexStart + 22);
+		DefaultBorderIndicies.Add(IndexStart + 22);
+		DefaultBorderIndicies.Add(IndexStart + 21);
+		DefaultBorderIndicies.Add(IndexStart + 23);
+
+		DefaultBorderIndicies.Add(IndexStart + 24);
+		DefaultBorderIndicies.Add(IndexStart + 25);
+		DefaultBorderIndicies.Add(IndexStart + 26);
+		DefaultBorderIndicies.Add(IndexStart + 26);
+		DefaultBorderIndicies.Add(IndexStart + 25);
+		DefaultBorderIndicies.Add(IndexStart + 27);
+
+		DefaultBorderIndicies.Add(IndexStart + 28);
+		DefaultBorderIndicies.Add(IndexStart + 29);
+		DefaultBorderIndicies.Add(IndexStart + 30);
+		DefaultBorderIndicies.Add(IndexStart + 30);
+		DefaultBorderIndicies.Add(IndexStart + 29);
+		DefaultBorderIndicies.Add(IndexStart + 31);
+
+		return DefaultBorderIndicies;
+	}();
+	RenderBatch.SourceIndices = &DefaultBorderIndiciesArray;
+	RenderBatch.NumIndices = DefaultBorderIndiciesArray.Num();
+	RenderBatch.IndexOffset = 0;
 }
 
 void FSlateElementBatcher::AddCustomElement( const FSlateDrawElement& DrawElement )
diff --git a/Engine/Source/Runtime/SlateCore/Public/Debugging/SlateDebugging.h b/Engine/Source/Runtime/SlateCore/Public/Debugging/SlateDebugging.h
--- a/Engine/Source/Runtime/SlateCore/Public/Debugging/SlateDebugging.h
+++ b/Engine/Source/Runtime/SlateCore/Public/Debugging/SlateDebugging.h
@@ -36,6 +36,7 @@ class FWeakWidgetPath;
 class FWidgetPath;
 class FNavigationReply;
 class FSlateInvalidationRoot;
+enum class EElementType : uint8;
 
 UENUM()
 enum class ESlateDebuggingInputEvent : uint8
@@ -284,8 +285,18 @@ public:
 	 * Called as soon as the element is added to the element list.
 	 * NOTE: These elements are not valid until the widget finishes painting, or you can resolve them all after the window finishes painting.
 	 */
+	UE_DEPRECATED(5.2, "FSlateDebugging::ElementAdded is deprecated, use FSlateDebugging::ElementTypeAdded instead")
 	DECLARE_MULTICAST_DELEGATE_TwoParams(FDrawElement, const FSlateWindowElementList& /*ElementList*/, int32 /*ElementIndex*/);
+PRAGMA_DISABLE_DEPRECATION_WARNINGS
 	static FDrawElement ElementAdded;
+PRAGMA_ENABLE_DEPRECATION_WARNINGS
+
+	/**
+	 * Called as soon as the element is added to the element list.
+	 * NOTE: These elements are not valid until the widget finishes painting, or you can resolve them all after the window finishes painting.
+	 */
+	DECLARE_MULTICAST_DELEGATE_ThreeParams(FDrawElementType, const FSlateWindowElementList& /*ElementList*/, int32 /*ElementIndex*/, EElementType /*ElementType*/);
+	static FDrawElementType ElementTypeAdded;
 
 public:
 	/**  */
diff --git a/Engine/Source/Runtime/SlateCore/Public/Rendering/DrawElements.h b/Engine/Source/Runtime/SlateCore/Public/Rendering/DrawElements.h
--- a/Engine/Source/Runtime/SlateCore/Public/Rendering/DrawElements.h
+++ b/Engine/Source/Runtime/SlateCore/Public/Rendering/DrawElements.h
@@ -51,6 +51,12 @@ enum class EElementType : uint8
 	ET_CustomVerts,
 	ET_PostProcessPass,
 	ET_RoundedBox,
+	/** 
+	 * We map draw elements by type on add for better cache coherency if possible, 
+	 * this type is used when that grouping is not possible. 
+	 * Grouping is also planned to be used for bulk element type processing.
+	 */
+	ET_NonMapped,
 	/** Total number of draw commands */
 	ET_Count,
 };
@@ -434,7 +440,7 @@ private:
 
 public:
 	/** List of source draw elements to create batches from */
-	FSlateDrawElementArray DrawElements;
+	FSlateDrawElementMap DrawElements;
 
 	TArray<int32> CachedRenderBatchIndices;
 
@@ -552,22 +558,11 @@ public:
 	}
 
 	/** @return Get the draw elements that we want to render into this window */
-	const FSlateDrawElementArray& GetUncachedDrawElements() const
+	const FSlateDrawElementMap& GetUncachedDrawElements() const
 	{
 		return UncachedDrawElements;
 	}
 
-	/**
-	 * Add a draw element to the list
-	 *
-	 * @param InDrawElement  The draw element to add
-	 */
-	UE_DEPRECATED(4.23, "AddItem is deprecated, use AddUninitialized instead")
-	void AddItem(const FSlateDrawElement& InDrawElement)
-	{
-		UncachedDrawElements.Add(InDrawElement);
-	}
-
 	/** @return Get the window size that we will be painting */
 	UE::Slate::FDeprecateVector2D GetWindowSize() const
 	{
@@ -596,7 +591,7 @@ public:
 	/**
 	 * Creates an uninitialized draw element
 	 */
-	SLATECORE_API FSlateDrawElement& AddUninitialized();
+	SLATECORE_API FSlateDrawElement& AddUninitialized(EElementType InElementType = EElementType::ET_NonMapped);
 
 
 	//--------------------------------------------------------------------------
@@ -754,7 +749,7 @@ private:
 	};
 
 	/** The uncached draw elements to be processed */
-	FSlateDrawElementArray UncachedDrawElements;
+	FSlateDrawElementMap UncachedDrawElements;
 
 	TArray<FWidgetDrawElementState, TInlineAllocator<50>> WidgetDrawStack;
 	TArray<FSlateCachedElementData*, TInlineAllocator<4>> CachedElementDataList;
diff --git a/Engine/Source/Runtime/SlateCore/Public/Rendering/ElementBatcher.h b/Engine/Source/Runtime/SlateCore/Public/Rendering/ElementBatcher.h
--- a/Engine/Source/Runtime/SlateCore/Public/Rendering/ElementBatcher.h
+++ b/Engine/Source/Runtime/SlateCore/Public/Rendering/ElementBatcher.h
@@ -276,7 +276,7 @@ public:
 	}
 
 private:
-	void AddElementsInternal(const FSlateDrawElementArray& DrawElements, FVector2f ViewportSize);
+	void AddElementsInternal(const FSlateDrawElementMap& DrawElements, FVector2f ViewportSize);
 	void AddCachedElements(FSlateCachedElementData& CachedElementData, FVector2f ViewportSize);
 
 	/** 
diff --git a/Engine/Source/Runtime/SlateCore/Public/Rendering/RenderingCommon.h b/Engine/Source/Runtime/SlateCore/Public/Rendering/RenderingCommon.h
--- a/Engine/Source/Runtime/SlateCore/Public/Rendering/RenderingCommon.h
+++ b/Engine/Source/Runtime/SlateCore/Public/Rendering/RenderingCommon.h
@@ -477,12 +477,28 @@ typedef TArray<FSlateVertex, FSlateStatTrackingMemoryAllocator<FRenderingBufferS
 typedef TArray<SlateIndex, FSlateStatTrackingMemoryAllocator<FRenderingBufferStatTracker>> FSlateIndexArray;
 typedef TArray<FSlateDrawElement, FSlateStatTrackingMemoryAllocator<FDrawElementStatTracker>> FSlateDrawElementArray;
 
+struct FSlateDrawElementContainer
+{
+	FSlateDrawElementArray Elements;
+};
+
+enum class EElementType : uint8;
+typedef TMap<EElementType, FSlateDrawElementContainer> FSlateDrawElementMap;
+
 #else
 
 typedef TArray<FSlateVertex> FSlateVertexArray;
 typedef TArray<SlateIndex> FSlateIndexArray;
 typedef TArray<FSlateDrawElement> FSlateDrawElementArray;
 
+struct FSlateDrawElementContainer
+{
+	FSlateDrawElementArray Elements;
+};
+
+enum class EElementType : uint8;
+typedef TMap<EElementType, FSlateDrawElementContainer> FSlateDrawElementMap;
+
 #endif
 
 /**
