Add the scrollbox padding to it style.

diff --git a/Engine/Source/Runtime/Slate/Private/Widgets/Layout/SScrollBox.cpp b/Engine/Source/Runtime/Slate/Private/Widgets/Layout/SScrollBox.cpp
--- a/Engine/Source/Runtime/Slate/Private/Widgets/Layout/SScrollBox.cpp
+++ b/Engine/Source/Runtime/Slate/Private/Widgets/Layout/SScrollBox.cpp
@@ -208,7 +208,7 @@ void SScrollBox::ConstructVerticalLayout()
 			SNew(SOverlay)
 
 			+ SOverlay::Slot()
-			.Padding(FMargin(0.0f, 0.0f, 0.0f, 1.0f))
+			.Padding(Style->VerticalScrolledContentPadding)
 			[
 				// Scroll panel that presents the scrolled content
 				ScrollPanel.ToSharedRef()
@@ -264,7 +264,7 @@ void SScrollBox::ConstructHorizontalLayout()
 			SNew(SOverlay)
 
 			+ SOverlay::Slot()
-			.Padding(FMargin(0.0f, 0.0f, 1.0f, 0.0f))
+			.Padding(Style->HorizontalScrolledContentPadding)
 			[
 				// Scroll panel that presents the scrolled content
 				ScrollPanel.ToSharedRef()
diff --git a/Engine/Source/Runtime/SlateCore/Public/Styling/SlateTypes.h b/Engine/Source/Runtime/SlateCore/Public/Styling/SlateTypes.h
--- a/Engine/Source/Runtime/SlateCore/Public/Styling/SlateTypes.h
+++ b/Engine/Source/Runtime/SlateCore/Public/Styling/SlateTypes.h
@@ -1975,6 +1975,24 @@ struct SLATECORE_API FScrollBoxStyle : public FSlateWidgetStyle
 		return *this;
 	}
 
+	/** Padding scroll panel that presents the scrolled content */
+	UPROPERTY(EditAnywhere, BlueprintReadWrite, Category = Appearance)
+	FMargin HorizontalScrolledContentPadding = FMargin(0.0f, 0.0f, 1.0f, 0.0f);
+	FScrollBoxStyle& SetHorizontalScrolledContentPadding(const FMargin& InPadding)
+	{
+		HorizontalScrolledContentPadding = InPadding;
+		return *this;
+	}
+
+	/** Padding scroll panel that presents the scrolled content */
+	UPROPERTY(EditAnywhere, BlueprintReadWrite, Category = Appearance)
+	FMargin VerticalScrolledContentPadding = FMargin(0.0f, 0.0f, 0.0f, 1.0f);
+	FScrollBoxStyle& SetVerticalScrolledContentPadding(const FMargin& InPadding)
+	{
+		VerticalScrolledContentPadding = InPadding;
+		return *this;
+	}
+
 	/**
 	 * Unlinks all colors in this style.
 	 * @see FSlateColor::Unlink
