Slate: Initial phase of FVector2D deprecation

diff --git a/Engine/Source/Runtime/Slate/Public/Widgets/Views/STableRow.h b/Engine/Source/Runtime/Slate/Public/Widgets/Views/STableRow.h
--- a/Engine/Source/Runtime/Slate/Public/Widgets/Views/STableRow.h
+++ b/Engine/Source/Runtime/Slate/Public/Widgets/Views/STableRow.h
@@ -380,9 +380,9 @@ public:
 		else
 		{
 			// Reuse the drop indicator asset for horizontal, by rotating the drawn box 90 degrees.
-			const FVector2D LocalSize(AllottedGeometry.GetLocalSize());
-			const FVector2D Pivot(LocalSize * 0.5f);
-			const FVector2D RotatedLocalSize(LocalSize.Y, LocalSize.X);
+			const FVector2f LocalSize(AllottedGeometry.GetLocalSize());
+			const FVector2f Pivot(LocalSize * 0.5f);
+			const FVector2f RotatedLocalSize(LocalSize.Y, LocalSize.X);
 			FSlateLayoutTransform RotatedTransform(Pivot - RotatedLocalSize * 0.5f);	// Make the box centered to the alloted geometry, so that it can be rotated around the center.
 
 			FSlateDrawElement::MakeRotatedBox(
@@ -723,12 +723,12 @@ public:
 	}
 
 	/** @return the zone (above, onto, below) based on where the user is hovering over within the row */
-	EItemDropZone ZoneFromPointerPosition(FVector2D LocalPointerPos, FVector2D LocalSize, EOrientation Orientation)
+	EItemDropZone ZoneFromPointerPosition(UE::Slate::FDeprecateVector2DParameter LocalPointerPos, UE::Slate::FDeprecateVector2DParameter LocalSize, EOrientation Orientation)
 	{
-		const FVector2D::FReal PointerPos = Orientation == EOrientation::Orient_Horizontal ? LocalPointerPos.X : LocalPointerPos.Y;
-		const FVector2D::FReal Size = Orientation == EOrientation::Orient_Horizontal ? LocalSize.X : LocalSize.Y;
+		const float PointerPos = Orientation == EOrientation::Orient_Horizontal ? LocalPointerPos.X : LocalPointerPos.Y;
+		const float Size = Orientation == EOrientation::Orient_Horizontal ? LocalSize.X : LocalSize.Y;
 
-		const FVector2D::FReal ZoneBoundarySu = FMath::Clamp(Size * 0.25f, 3.0f, 10.0f);
+		const float ZoneBoundarySu = FMath::Clamp(Size * 0.25f, 3.0f, 10.0f);
 		if (PointerPos < ZoneBoundarySu)
 		{
 			return EItemDropZone::AboveItem;
@@ -748,7 +748,7 @@ public:
 		if ( OnCanAcceptDrop.IsBound() )
 		{
 			const TSharedRef< ITypedTableView<ItemType> > OwnerTable = OwnerTablePtr.Pin().ToSharedRef();
-			const FVector2D LocalPointerPos = MyGeometry.AbsoluteToLocal(DragDropEvent.GetScreenSpacePosition());
+			const FVector2f LocalPointerPos = MyGeometry.AbsoluteToLocal(DragDropEvent.GetScreenSpacePosition());
 			const EItemDropZone ItemHoverZone = ZoneFromPointerPosition(LocalPointerPos, MyGeometry.GetLocalSize(), OwnerTable->Private_GetOrientation());
 
 			ItemDropZone = [ItemHoverZone, DragDropEvent, this]()
@@ -786,7 +786,7 @@ public:
 				if (const ItemType* MyItemPtr = GetItemForThis(OwnerTable))
 				{
 					// Which physical drop zone is the drop about to be performed onto?
-					const FVector2D LocalPointerPos = MyGeometry.AbsoluteToLocal(DragDropEvent.GetScreenSpacePosition());
+					const FVector2f LocalPointerPos = MyGeometry.AbsoluteToLocal(DragDropEvent.GetScreenSpacePosition());
 					const EItemDropZone HoveredZone = ZoneFromPointerPosition(LocalPointerPos, MyGeometry.GetLocalSize(), OwnerTable->Private_GetOrientation());
 
 					// The row gets final say over which zone to drop onto regardless of physical location.
diff --git a/Engine/Source/Runtime/Slate/Public/Widgets/Views/STableViewBase.h b/Engine/Source/Runtime/Slate/Public/Widgets/Views/STableViewBase.h
--- a/Engine/Source/Runtime/Slate/Public/Widgets/Views/STableViewBase.h
+++ b/Engine/Source/Runtime/Slate/Public/Widgets/Views/STableViewBase.h
@@ -62,11 +62,11 @@ struct SLATE_API FTableViewDimensions
 {
 	FTableViewDimensions(EOrientation InOrientation);
 	FTableViewDimensions(EOrientation InOrientation, float X, float Y);
-	FTableViewDimensions(EOrientation InOrientation, const FVector2D& Size);
+	FTableViewDimensions(EOrientation InOrientation, const UE::Slate::FDeprecateVector2DParameter& Size);
 
-	FVector2D ToVector2D() const
+	UE::Slate::FDeprecateVector2DResult ToVector2D() const
 	{
-		return Orientation == Orient_Vertical ? FVector2D(LineAxis, ScrollAxis) : FVector2D(ScrollAxis, LineAxis);
+		return Orientation == Orient_Vertical ? FVector2f(LineAxis, ScrollAxis) : FVector2f(ScrollAxis, LineAxis);
 	}
 
 	FTableViewDimensions operator+(const FTableViewDimensions& Other) const
@@ -278,7 +278,7 @@ protected:
 	/**
 	* Get the uniform item
 	*/
-	FVector2D GetItemSize() const;
+	UE::Slate::FDeprecateVector2DResult GetItemSize() const;
 
 	/** @return the number of items that can fit on the screen */
 	virtual float GetNumLiveWidgets() const;
@@ -425,7 +425,7 @@ protected:
 	float AmountScrolledWhileRightMouseDown;
 
 	/** The location in screenspace the view was pressed */
-	FVector2D PressedScreenSpacePosition;
+	FVector2f PressedScreenSpacePosition;
 
 	/** The amount we have scrolled this tick cycle */
 	float TickScrollDelta;
@@ -452,7 +452,7 @@ protected:
 	FInertialScrollManager InertialScrollManager;
 
 	/**	The current position of the software cursor */
-	FVector2D SoftwareCursorPosition;
+	FVector2f SoftwareCursorPosition;
 
 	/**	Whether the software cursor should be drawn in the viewport */
 	bool bShowSoftwareCursor;
diff --git a/Engine/Source/Runtime/SlateCore/Private/Application/SlateApplicationBase.cpp b/Engine/Source/Runtime/SlateCore/Private/Application/SlateApplicationBase.cpp
--- a/Engine/Source/Runtime/SlateCore/Private/Application/SlateApplicationBase.cpp
+++ b/Engine/Source/Runtime/SlateCore/Private/Application/SlateApplicationBase.cpp
@@ -20,7 +20,7 @@ const uint32 FSlateApplicationBase::CursorPointerIndex = ETouchIndex::CursorPoin
 const uint32 FSlateApplicationBase::CursorUserIndex = 0;
 const FPlatformUserId FSlateApplicationBase::SlateAppPrimaryPlatformUser = FPlatformUserId::CreateFromInternalId(0);
 
-FWidgetPath FHitTesting::LocateWidgetInWindow(FVector2D ScreenspaceMouseCoordinate, const TSharedRef<SWindow>& Window, bool bIgnoreEnabledStatus, int32 UserIndex) const
+FWidgetPath FHitTesting::LocateWidgetInWindow(FVector2f ScreenspaceMouseCoordinate, const TSharedRef<SWindow>& Window, bool bIgnoreEnabledStatus, int32 UserIndex) const
 {
 	return SlateApp->LocateWidgetInWindow(ScreenspaceMouseCoordinate, Window, bIgnoreEnabledStatus, UserIndex);
 }
@@ -50,9 +50,9 @@ void FSlateApplicationBase::GetCachedDisplayMetrics(FDisplayMetrics& OutDisplayM
 	OutDisplayMetrics = CachedDisplayMetrics;
 }
 
-void FSlateApplicationBase::GetSafeZoneSize(FMargin& SafeZone, const FVector2D& OverrideSize)
+void FSlateApplicationBase::GetSafeZoneSize(FMargin& SafeZone, const UE::Slate::FDeprecateVector2DParameter& OverrideSize)
 {
-	FVector2D ContainerSize = FVector2D::ZeroVector;
+	FVector2f ContainerSize = FVector2f::ZeroVector;
 
 #if WITH_EDITOR
 	ContainerSize = OverrideSize;
@@ -62,7 +62,7 @@ void FSlateApplicationBase::GetSafeZoneSize(FMargin& SafeZone, const FVector2D&
 	{
 		FDisplayMetrics Metrics;
 		GetCachedDisplayMetrics(Metrics);
-		ContainerSize = FVector2D((float)Metrics.PrimaryDisplayWidth, (float)Metrics.PrimaryDisplayHeight);
+		ContainerSize = FVector2f((float)Metrics.PrimaryDisplayWidth, (float)Metrics.PrimaryDisplayHeight);
 	}
 
 	FMargin SafeZoneRatio;
diff --git a/Engine/Source/Runtime/SlateCore/Private/Application/SlateWindowHelper.cpp b/Engine/Source/Runtime/SlateCore/Private/Application/SlateWindowHelper.cpp
--- a/Engine/Source/Runtime/SlateCore/Private/Application/SlateWindowHelper.cpp
+++ b/Engine/Source/Runtime/SlateCore/Private/Application/SlateWindowHelper.cpp
@@ -57,8 +57,8 @@ bool FSlateWindowHelper::CheckWorkAreaForWindows( const TArray<TSharedRef<SWindo
 	for (TArray<TSharedRef<SWindow>>::TConstIterator CurrentWindowIt(WindowsToSearch); CurrentWindowIt; ++CurrentWindowIt)
 	{
 		const TSharedRef<SWindow>& CurrentWindow = *CurrentWindowIt;
-		const FVector2D Position = CurrentWindow->GetPositionInScreen();
-		const FVector2D Size = CurrentWindow->GetSizeInScreen();
+		const FVector2f Position = CurrentWindow->GetPositionInScreen();
+		const FVector2f Size = CurrentWindow->GetSizeInScreen();
 		const FSlateRect WindowRect(Position.X, Position.Y, Size.X, Size.Y);
 
 		if (FSlateRect::DoRectanglesIntersect(WorkAreaRect, WindowRect) || CheckWorkAreaForWindows(CurrentWindow->GetChildWindows(), WorkAreaRect))
diff --git a/Engine/Source/Runtime/SlateCore/Private/Brushes/SlateDynamicImageBrush.cpp b/Engine/Source/Runtime/SlateCore/Private/Brushes/SlateDynamicImageBrush.cpp
--- a/Engine/Source/Runtime/SlateCore/Private/Brushes/SlateDynamicImageBrush.cpp
+++ b/Engine/Source/Runtime/SlateCore/Private/Brushes/SlateDynamicImageBrush.cpp
@@ -10,7 +10,7 @@
 
 TSharedPtr<FSlateDynamicImageBrush> FSlateDynamicImageBrush::CreateWithImageData(
 	const FName InTextureName,
-	const FVector2D& InImageSize,
+	const UE::Slate::FDeprecateVector2DParameter& InImageSize,
 	const TArray<uint8>& InImageData,
 	const FLinearColor& InTint,
 	ESlateBrushTileType::Type InTiling,
diff --git a/Engine/Source/Runtime/SlateCore/Private/Debugging/ConsoleSlateDebuggerInvalidate.cpp b/Engine/Source/Runtime/SlateCore/Private/Debugging/ConsoleSlateDebuggerInvalidate.cpp
--- a/Engine/Source/Runtime/SlateCore/Private/Debugging/ConsoleSlateDebuggerInvalidate.cpp
+++ b/Engine/Source/Runtime/SlateCore/Private/Debugging/ConsoleSlateDebuggerInvalidate.cpp
@@ -646,13 +646,13 @@ void FConsoleSlateDebuggerInvalidate::HandlePaintDebugInfo(const FPaintArgs& InA
 	CacheDuration = FMath::Max(CacheDuration, 0.01f);
 	const double SlateApplicationCurrentTime = FSlateApplicationBase::Get().GetCurrentTime();
 
-	FVector2D TextElementLocation {16.f, 64.f};
+	FVector2f TextElementLocation {16.f, 64.f};
 	auto MakeText = [&](const FString& Text, const FLinearColor& Color)
 	{
 		FSlateDrawElement::MakeText(
 			InOutDrawElements
 			, InOutLayerId
-			, InAllottedGeometry.ToPaintGeometry(TextElementLocation, FVector2D(1.f, 1.f))
+			, InAllottedGeometry.ToPaintGeometry(FVector2f(1.f, 1.f), FSlateLayoutTransform(TextElementLocation))
 			, Text
 			, FontInfo
 			, ESlateDrawEffect::None
@@ -745,7 +745,7 @@ void FConsoleSlateDebuggerInvalidate::HandlePaintDebugInfo(const FPaintArgs& InA
 		FSlateDrawElement::MakeText(
 			InOutDrawElements
 			, InOutLayerId
-			, InAllottedGeometry.ToPaintGeometry(TextElementLocation, FVector2D(1.f, 1.f))
+			, InAllottedGeometry.ToPaintGeometry(FVector2f(1.f, 1.f), FSlateLayoutTransform(TextElementLocation))
 			, FString::Printf(TEXT("Slate Performance Threshold Reached: %d"), LastPerformanceThresholdFrameCount)
 			, NormalFontInfo
 			, ESlateDrawEffect::None
diff --git a/Engine/Source/Runtime/SlateCore/Private/Debugging/ConsoleSlateDebuggerInvalidate.h b/Engine/Source/Runtime/SlateCore/Private/Debugging/ConsoleSlateDebuggerInvalidate.h
--- a/Engine/Source/Runtime/SlateCore/Private/Debugging/ConsoleSlateDebuggerInvalidate.h
+++ b/Engine/Source/Runtime/SlateCore/Private/Debugging/ConsoleSlateDebuggerInvalidate.h
@@ -55,10 +55,10 @@ private:
 		FConsoleSlateDebuggerUtility::TSWindowId WindowId;
 		FString WidgetInvalidatedName;
 		FString WidgetInvalidatorName;
-		FVector2D InvalidatedPaintLocation;
-		FVector2D InvalidatedPaintSize;
-		FVector2D InvalidatorPaintLocation;
-		FVector2D InvalidatorPaintSize;
+		FVector2f InvalidatedPaintLocation;
+		FVector2f InvalidatedPaintSize;
+		FVector2f InvalidatorPaintLocation;
+		FVector2f InvalidatorPaintSize;
 		EInvalidateWidgetReason WidgetReason;
 		ESlateDebuggingInvalidateRootReason InvalidationRootReason;
 		int32 InvalidationPriority;
diff --git a/Engine/Source/Runtime/SlateCore/Private/Debugging/ConsoleSlateDebuggerInvalidationRoot.cpp b/Engine/Source/Runtime/SlateCore/Private/Debugging/ConsoleSlateDebuggerInvalidationRoot.cpp
--- a/Engine/Source/Runtime/SlateCore/Private/Debugging/ConsoleSlateDebuggerInvalidationRoot.cpp
+++ b/Engine/Source/Runtime/SlateCore/Private/Debugging/ConsoleSlateDebuggerInvalidationRoot.cpp
@@ -174,12 +174,12 @@ void FConsoleSlateDebuggerInvalidationRoot::HandlePaintDebugInfo(const FPaintArg
 
 
 	float TextElementY = 48.f;
-	auto MakeText = [&](const FString& Text, const FVector2D& Location, const FLinearColor& Color)
+	auto MakeText = [&](const FString& Text, const FVector2f& Location, const FLinearColor& Color)
 	{
 		FSlateDrawElement::MakeText(
 			InOutDrawElements
 			, InOutLayerId
-			, InAllottedGeometry.ToPaintGeometry(Location, FVector2D(1.f, 1.f))
+			, InAllottedGeometry.ToPaintGeometry(FVector2f(1.f, 1.f), FSlateLayoutTransform(Location))
 			, Text
 			, FontInfo
 			, ESlateDrawEffect::None
@@ -202,9 +202,9 @@ void FConsoleSlateDebuggerInvalidationRoot::HandlePaintDebugInfo(const FPaintArg
 		//	ESlateDrawEffect::None,
 		//	FLinearColor::Black);
 
-		MakeText(TEXT("Slow Path"), FVector2D(10.f, 10.f+0.f), DrawSlowPathColor);
-		MakeText(TEXT("Fast Path"), FVector2D(10.f, 10.f+12.f), DrawFastPathColor);
-		MakeText(TEXT("No Paint occurred"), FVector2D(10.f, 10.f+24.f), DrawNoneColor);
+		MakeText(TEXT("Slow Path"), FVector2f(10.f, 10.f+0.f), DrawSlowPathColor);
+		MakeText(TEXT("Fast Path"), FVector2f(10.f, 10.f+12.f), DrawFastPathColor);
+		MakeText(TEXT("No Paint occurred"), FVector2f(10.f, 10.f+24.f), DrawNoneColor);
 		TextElementY += 36.f;
 	}
 
@@ -305,7 +305,7 @@ void FConsoleSlateDebuggerInvalidationRoot::HandlePaintDebugInfo(const FPaintArg
 				if (NumberOfWidget < MaxNumberOfWidgetInList)
 				{
 					FString WidgetDisplayName = FString::Printf(TEXT("Id:('%d') - %s"), Itt.Key, *FReflectionMetaData::GetWidgetDebugInfo(Widget));
-					MakeText(WidgetDisplayName, FVector2D(0.f, (12.f * NumberOfWidget) + TextElementY), DrawColor);
+					MakeText(WidgetDisplayName, FVector2f(0.f, (12.f * NumberOfWidget) + TextElementY), DrawColor);
 				}
 			}
 			++NumberOfWidget;
@@ -315,7 +315,7 @@ void FConsoleSlateDebuggerInvalidationRoot::HandlePaintDebugInfo(const FPaintArg
 	if (bDisplayInvalidationRootList && NumberOfWidget == MaxNumberOfWidgetInList)
 	{
 		FString WidgetDisplayName = FString::Printf(TEXT("   %d more invalidation root"), NumberOfWidget - MaxNumberOfWidgetInList);
-		MakeText(WidgetDisplayName, FVector2D(0.f, (12.f * NumberOfWidget) + TextElementY), FLinearColor::White);
+		MakeText(WidgetDisplayName, FVector2f(0.f, (12.f * NumberOfWidget) + TextElementY), FLinearColor::White);
 	}
 }
 
diff --git a/Engine/Source/Runtime/SlateCore/Private/Debugging/ConsoleSlateDebuggerPaint.cpp b/Engine/Source/Runtime/SlateCore/Private/Debugging/ConsoleSlateDebuggerPaint.cpp
--- a/Engine/Source/Runtime/SlateCore/Private/Debugging/ConsoleSlateDebuggerPaint.cpp
+++ b/Engine/Source/Runtime/SlateCore/Private/Debugging/ConsoleSlateDebuggerPaint.cpp
@@ -236,12 +236,12 @@ void FConsoleSlateDebuggerPaint::HandlePaintDebugInfo(const FPaintArgs& InArgs,
 	CacheDuration = FMath::Max(CacheDuration, 0.01f);
 	const double SlateApplicationCurrentTime = FSlateApplicationBase::Get().GetCurrentTime();
 
-	auto MakeText = [&](const FString& Text, const FVector2D& Location, const FLinearColor& Color)
+	auto MakeText = [&](const FString& Text, const FVector2f& Location, const FLinearColor& Color)
 	{
 		FSlateDrawElement::MakeText(
 			InOutDrawElements
 			, InOutLayerId
-			, InAllottedGeometry.ToPaintGeometry(Location, FVector2D(1.f, 1.f))
+			, InAllottedGeometry.ToPaintGeometry(FVector2f(1.f, 1.f), FSlateLayoutTransform(Location))
 			, Text
 			, FontInfo
 			, ESlateDrawEffect::None
@@ -293,7 +293,7 @@ void FConsoleSlateDebuggerPaint::HandlePaintDebugInfo(const FPaintArgs& InArgs,
 			{
 				if (NumberOfWidget < MaxNumberOfWidgetInList)
 				{
-					MakeText(Itt.Value.WidgetName, FVector2D(0, (12 * NumberOfWidget) + TextElementY), DrawWidgetNameColor);
+					MakeText(Itt.Value.WidgetName, FVector2f(0.f, (12.f * NumberOfWidget) + TextElementY), DrawWidgetNameColor);
 				}
 			}
 			++NumberOfWidget;
@@ -303,13 +303,13 @@ void FConsoleSlateDebuggerPaint::HandlePaintDebugInfo(const FPaintArgs& InArgs,
 
 	{
 		FString NumberOfWidgetDrawn = FString::Printf(TEXT("Number of Widget Painted: %d"), NumberOfWidget);
-		MakeText(NumberOfWidgetDrawn, FVector2D(10, 10), DrawWidgetNameColor);
+		MakeText(NumberOfWidgetDrawn, FVector2f(10.f, 10.f), DrawWidgetNameColor);
 	}
 
 	if (bDisplayWidgetsNameList && NumberOfWidget > MaxNumberOfWidgetInList)
 	{
 		FString WidgetDisplayName = FString::Printf(TEXT("   %d more invalidations"), NumberOfWidget - MaxNumberOfWidgetInList);
-		MakeText(WidgetDisplayName, FVector2D(0.f, (12.f * NumberOfWidget) + TextElementY), FLinearColor::White);
+		MakeText(WidgetDisplayName, FVector2f(0.f, (12.f * NumberOfWidget) + TextElementY), FLinearColor::White);
 	}
 }
 
diff --git a/Engine/Source/Runtime/SlateCore/Private/Debugging/ConsoleSlateDebuggerPaint.h b/Engine/Source/Runtime/SlateCore/Private/Debugging/ConsoleSlateDebuggerPaint.h
--- a/Engine/Source/Runtime/SlateCore/Private/Debugging/ConsoleSlateDebuggerPaint.h
+++ b/Engine/Source/Runtime/SlateCore/Private/Debugging/ConsoleSlateDebuggerPaint.h
@@ -80,8 +80,8 @@ private:
 	struct FPaintInfo 
 	{
 		FConsoleSlateDebuggerUtility::TSWindowId Window;
-		FVector2D PaintLocation;
-		FVector2D PaintSize;
+		FVector2f PaintLocation;
+		FVector2f PaintSize;
 		FString WidgetName;
 		double LastPaint;
 		int32 PaintCount;
diff --git a/Engine/Source/Runtime/SlateCore/Private/Debugging/ConsoleSlateDebuggerUpdate.cpp b/Engine/Source/Runtime/SlateCore/Private/Debugging/ConsoleSlateDebuggerUpdate.cpp
--- a/Engine/Source/Runtime/SlateCore/Private/Debugging/ConsoleSlateDebuggerUpdate.cpp
+++ b/Engine/Source/Runtime/SlateCore/Private/Debugging/ConsoleSlateDebuggerUpdate.cpp
@@ -379,12 +379,12 @@ void FConsoleSlateDebuggerUpdate::HandlePaintDebugInfo(const FPaintArgs& InArgs,
 	const double SlateApplicationCurrentTime = FSlateApplicationBase::Get().GetCurrentTime();
 
 	float TextElementY = 48.f;
-	auto MakeText = [&](const FString& Text, const FVector2D& Location, const FLinearColor& Color)
+	auto MakeText = [&](const FString& Text, const FVector2f& Location, const FLinearColor& Color)
 	{
 		FSlateDrawElement::MakeText(
 			InOutDrawElements
 			, InOutLayerId
-			, InAllottedGeometry.ToPaintGeometry(Location, FVector2D(1.f, 1.f))
+			, InAllottedGeometry.ToPaintGeometry(FVector2f(1.f, 1.f), FSlateLayoutTransform(Location))
 			, Text
 			, FontInfo
 			, ESlateDrawEffect::None
@@ -393,10 +393,10 @@ void FConsoleSlateDebuggerUpdate::HandlePaintDebugInfo(const FPaintArgs& InArgs,
 
 	if (bShowLegend)
 	{
-		MakeText(TEXT("Repaint"), FVector2D(10.f, 10.f + 0.f), DrawRepaintColor);
-		MakeText(TEXT("Volatile Paint"), FVector2D(10.f, 10.f + 12.f), DrawVolatilePaintColor);
-		MakeText(TEXT("Tick"), FVector2D(10.f, 10.f + 24.f), DrawTickColor);
-		MakeText(TEXT("Active Timer"), FVector2D(10.f, 10.f + 36.f), DrawActiveTimerColor);
+		MakeText(TEXT("Repaint"), FVector2f(10.f, 10.f + 0.f), DrawRepaintColor);
+		MakeText(TEXT("Volatile Paint"), FVector2f(10.f, 10.f + 12.f), DrawVolatilePaintColor);
+		MakeText(TEXT("Tick"), FVector2f(10.f, 10.f + 24.f), DrawTickColor);
+		MakeText(TEXT("Active Timer"), FVector2f(10.f, 10.f + 36.f), DrawActiveTimerColor);
 		TextElementY += 48.f;
 	}
 
@@ -449,7 +449,7 @@ void FConsoleSlateDebuggerUpdate::HandlePaintDebugInfo(const FPaintArgs& InArgs,
 			{
 				if (NumberOfWidget < MaxNumberOfWidgetInList)
 				{
-					MakeText(Itt.Value.WidgetName, FVector2D(0, (12 * NumberOfWidget) + TextElementY), Color);
+					MakeText(Itt.Value.WidgetName, FVector2f(0.f, (12.f * NumberOfWidget) + TextElementY), Color);
 				}
 			}
 			++NumberOfWidget;
@@ -459,7 +459,7 @@ void FConsoleSlateDebuggerUpdate::HandlePaintDebugInfo(const FPaintArgs& InArgs,
 	if (bDisplayWidgetsNameList && NumberOfWidget > MaxNumberOfWidgetInList)
 	{
 		FString WidgetDisplayName = FString::Printf(TEXT("   %d more updates"), NumberOfWidget - MaxNumberOfWidgetInList);
-		MakeText(WidgetDisplayName, FVector2D(0.f, (12.f * NumberOfWidget) + TextElementY), FLinearColor::White);
+		MakeText(WidgetDisplayName, FVector2f(0.f, (12.f * NumberOfWidget) + TextElementY), FLinearColor::White);
 	}
 }
 
diff --git a/Engine/Source/Runtime/SlateCore/Private/Debugging/ConsoleSlateDebuggerUpdate.h b/Engine/Source/Runtime/SlateCore/Private/Debugging/ConsoleSlateDebuggerUpdate.h
--- a/Engine/Source/Runtime/SlateCore/Private/Debugging/ConsoleSlateDebuggerUpdate.h
+++ b/Engine/Source/Runtime/SlateCore/Private/Debugging/ConsoleSlateDebuggerUpdate.h
@@ -89,8 +89,8 @@ private:
 		void Update(const SWidget* Widget, EWidgetUpdateFlags InUpdateFlags);
 
 		FConsoleSlateDebuggerUtility::TSWindowId WindowId;
-		FVector2D PaintLocation;
-		FVector2D PaintSize;
+		FVector2f PaintLocation;
+		FVector2f PaintSize;
 		FString WidgetName;
 		EWidgetUpdateFlags UpdateFlags;
 		double LastInvalidationTime;
diff --git a/Engine/Source/Runtime/SlateCore/Private/FastUpdate/SlateInvalidationRoot.cpp b/Engine/Source/Runtime/SlateCore/Private/FastUpdate/SlateInvalidationRoot.cpp
--- a/Engine/Source/Runtime/SlateCore/Private/FastUpdate/SlateInvalidationRoot.cpp
+++ b/Engine/Source/Runtime/SlateCore/Private/FastUpdate/SlateInvalidationRoot.cpp
@@ -791,7 +791,7 @@ bool FSlateInvalidationRoot::PaintFastPath(const FSlateInvalidationContext& Cont
 	return bWidgetsNeededRepaint;
 }
 
-void FSlateInvalidationRoot::AdjustWidgetsDesktopGeometry(FVector2D WindowToDesktopTransform)
+void FSlateInvalidationRoot::AdjustWidgetsDesktopGeometry(UE::Slate::FDeprecateVector2DParameter WindowToDesktopTransform)
 {
 	FSlateLayoutTransform WindowToDesktop(WindowToDesktopTransform);
 
diff --git a/Engine/Source/Runtime/SlateCore/Private/FastUpdate/WidgetProxy.cpp b/Engine/Source/Runtime/SlateCore/Private/FastUpdate/WidgetProxy.cpp
--- a/Engine/Source/Runtime/SlateCore/Private/FastUpdate/WidgetProxy.cpp
+++ b/Engine/Source/Runtime/SlateCore/Private/FastUpdate/WidgetProxy.cpp
@@ -107,8 +107,8 @@ void FWidgetProxy::ProcessLayoutInvalidation(FSlateInvalidationWidgetPostHeap& U
 	SWidget* WidgetPtr = GetWidget();
 	SCOPE_CYCLE_SWIDGET(WidgetPtr);
 	// When layout changes compute a new desired size for this widget
-	const FVector2D CurrentDesiredSize = WidgetPtr->GetDesiredSize();
-	FVector2D NewDesiredSize = FVector2D::ZeroVector;
+	const FVector2f CurrentDesiredSize = WidgetPtr->GetDesiredSize();
+	FVector2f NewDesiredSize = FVector2f::ZeroVector;
 	if (!Visibility.IsCollapseDirectly())
 	{
 		if (WidgetPtr->NeedsPrepass())
diff --git a/Engine/Source/Runtime/SlateCore/Private/Input/DragAndDrop.cpp b/Engine/Source/Runtime/SlateCore/Private/Input/DragAndDrop.cpp
--- a/Engine/Source/Runtime/SlateCore/Private/Input/DragAndDrop.cpp
+++ b/Engine/Source/Runtime/SlateCore/Private/Input/DragAndDrop.cpp
@@ -34,7 +34,7 @@ void FDragDropOperation::OnDragged( const class FDragDropEvent& DragDropEvent )
 	{
 		FSlateApplicationBase& SlateApplication = FSlateApplicationBase::Get();
 						
-		FVector2D Position = DragDropEvent.GetScreenSpacePosition() + FSlateApplicationBase::Get().GetCursorSize();
+		FVector2f Position = DragDropEvent.GetScreenSpacePosition() + FSlateApplicationBase::Get().GetCursorSize();
 
 		if (Position != CursorDecoratorWindow->GetPositionInScreen())
 		{
diff --git a/Engine/Source/Runtime/SlateCore/Private/Input/HittestGrid.cpp b/Engine/Source/Runtime/SlateCore/Private/Input/HittestGrid.cpp
--- a/Engine/Source/Runtime/SlateCore/Private/Input/HittestGrid.cpp
+++ b/Engine/Source/Runtime/SlateCore/Private/Input/HittestGrid.cpp
@@ -40,7 +40,7 @@ constexpr bool IsCompatibleUserIndex(int32 RequestedUserIndex, int32 TestUserInd
 // Helper Functions
 //
 
-FVector2D ClosestPointOnSlateRotatedRect(const FVector2D &Point, const FSlateRotatedRect& RotatedRect)
+FVector2f ClosestPointOnSlateRotatedRect(const FVector2f &Point, const FSlateRotatedRect& RotatedRect)
 {
 	//no need to do any testing if we are inside of the rect
 	if (RotatedRect.IsUnderLocation(Point))
@@ -49,26 +49,26 @@ FVector2D ClosestPointOnSlateRotatedRect(const FVector2D &Point, const FSlateRot
 	}
 
 	const static int32 NumOfCorners = 4;
-	FVector2D Corners[NumOfCorners];
-	Corners[0] = FVector2D(RotatedRect.TopLeft);
-	Corners[1] = FVector2D(Corners[0]) + FVector2D(RotatedRect.ExtentX);
-	Corners[2] = FVector2D(Corners[1]) + FVector2D(RotatedRect.ExtentY);
-	Corners[3] = FVector2D(Corners[0]) + FVector2D(RotatedRect.ExtentY);
+	FVector2d Corners[NumOfCorners];
+	Corners[0] = FVector2d(RotatedRect.TopLeft);
+	Corners[1] = FVector2d(Corners[0]) + FVector2d(RotatedRect.ExtentX);
+	Corners[2] = FVector2d(Corners[1]) + FVector2d(RotatedRect.ExtentY);
+	Corners[3] = FVector2d(Corners[0]) + FVector2d(RotatedRect.ExtentY);
 
-	FVector2D RetPoint;
+	FVector2f RetPoint;
 	float ClosestDistSq = FLT_MAX;
 	for (int32 i = 0; i < NumOfCorners; ++i)
 	{
 		//grab the closest point along the line segment
-		const FVector2D ClosestPoint = FMath::ClosestPointOnSegment2D(Point, Corners[i], Corners[(i + 1) % NumOfCorners]);
+		const FVector2d ClosestPoint = FMath::ClosestPointOnSegment2D(FVector2d(Point), Corners[i], Corners[(i + 1) % NumOfCorners]);
 
 		//get the distance between the two
-		const float TestDist = FVector2D::DistSquared(Point, ClosestPoint);
+		const float TestDist = FVector2d::DistSquared(FVector2d(Point), ClosestPoint);
 
 		//if the distance is smaller than the current smallest, update our closest
 		if (TestDist < ClosestDistSq)
 		{
-			RetPoint = ClosestPoint;
+			RetPoint = FVector2f(UE_REAL_TO_FLOAT(ClosestPoint.X), UE_REAL_TO_FLOAT(ClosestPoint.Y));
 			ClosestDistSq = TestDist;
 		}
 	}
@@ -76,12 +76,12 @@ FVector2D ClosestPointOnSlateRotatedRect(const FVector2D &Point, const FSlateRot
 	return RetPoint;
 }
 
-FORCEINLINE float DistanceSqToSlateRotatedRect(const FVector2D &Point, const FSlateRotatedRect& RotatedRect)
+FORCEINLINE float DistanceSqToSlateRotatedRect(const FVector2f &Point, const FSlateRotatedRect& RotatedRect)
 {
-	return FVector2D::DistSquared(ClosestPointOnSlateRotatedRect(Point, RotatedRect), Point);
+	return FVector2f::DistSquared(ClosestPointOnSlateRotatedRect(Point, RotatedRect), Point);
 }
 
-FORCEINLINE bool IsOverlappingSlateRotatedRect(const FVector2D& Point, const float Radius, const FSlateRotatedRect& RotatedRect)
+FORCEINLINE bool IsOverlappingSlateRotatedRect(const FVector2f& Point, const float Radius, const FSlateRotatedRect& RotatedRect)
 {
 	return DistanceSqToSlateRotatedRect( Point, RotatedRect ) <= (Radius * Radius);
 }
@@ -100,7 +100,7 @@ bool ContainsInteractableWidget(const TArray<FWidgetAndPointer>& PathToTest)
 }
 
 
-const FVector2D CellSize(128.0f, 128.0f);
+const FVector2f CellSize(128.0f, 128.0f);
 
 //
 // FHittestGrid::FWidgetIndex
@@ -119,13 +119,13 @@ struct FHittestGrid::FGridTestingParams
 	/** Ctor */
 	FGridTestingParams()
 	: CellCoord(-1, -1)
-	, CursorPositionInGrid(FVector2D::ZeroVector)
+	, CursorPositionInGrid(FVector2f::ZeroVector)
 	, Radius(-1.0f)
 	, bTestWidgetIsInteractive(false)
 	{}
 
 	FIntPoint CellCoord;
-	FVector2D CursorPositionInGrid;
+	FVector2f CursorPositionInGrid;
 	float Radius;
 	bool bTestWidgetIsInteractive;
 };
@@ -163,11 +163,11 @@ FHittestGrid::FHittestGrid()
 {
 }
 
-TArray<FWidgetAndPointer> FHittestGrid::GetBubblePath(FVector2D DesktopSpaceCoordinate, float CursorRadius, bool bIgnoreEnabledStatus, int32 UserIndex)
+TArray<FWidgetAndPointer> FHittestGrid::GetBubblePath(UE::Slate::FDeprecateVector2DParameter DesktopSpaceCoordinate, float CursorRadius, bool bIgnoreEnabledStatus, int32 UserIndex)
 {
 	checkSlow(IsInGameThread());
 
-	const FVector2D CursorPositionInGrid = DesktopSpaceCoordinate - GridOrigin;
+	const FVector2f CursorPositionInGrid = DesktopSpaceCoordinate - GridOrigin;
 
 	if (WidgetArray.Num() > 0 && Cells.Num() > 0)
 	{
@@ -223,7 +223,8 @@ TArray<FWidgetAndPointer> FHittestGrid::GetBubblePath(FVector2D DesktopSpaceCoor
 				{
 					if (BestHitWidgetData.CustomPath.IsValid())
 					{
-						TArray<FWidgetAndPointer> BubblePathExtension = BestHitWidgetData.CustomPath.Pin()->GetBubblePathAndVirtualCursors(FirstHitWidget->GetTickSpaceGeometry(), DesktopSpaceCoordinate, bIgnoreEnabledStatus);
+						FVector2d DesktopSpaceCoordinate2d(DesktopSpaceCoordinate.X, DesktopSpaceCoordinate.Y);
+						TArray<FWidgetAndPointer> BubblePathExtension = BestHitWidgetData.CustomPath.Pin()->GetBubblePathAndVirtualCursors(FirstHitWidget->GetTickSpaceGeometry(), DesktopSpaceCoordinate2d, bIgnoreEnabledStatus);
 						Path.Append(MoveTemp(BubblePathExtension));
 					}
 				}
@@ -236,7 +237,7 @@ TArray<FWidgetAndPointer> FHittestGrid::GetBubblePath(FVector2D DesktopSpaceCoor
 	return TArray<FWidgetAndPointer>();
 }
 
-bool FHittestGrid::SetHittestArea(const FVector2D& HittestPositionInDesktop, const FVector2D& HittestDimensions, const FVector2D& HitestOffsetInWindow)
+bool FHittestGrid::SetHittestArea(const UE::Slate::FDeprecateVector2DParameter& HittestPositionInDesktop, const UE::Slate::FDeprecateVector2DParameter& HittestDimensions, const UE::Slate::FDeprecateVector2DParameter& HitestOffsetInWindow)
 {
 	bool bWasCleared = false;
 
@@ -442,7 +443,7 @@ TSharedPtr<SWidget> FHittestGrid::FindFocusableWidget(FSlateRect WidgetRect, con
 					return TSharedPtr<SWidget>();
 				case EUINavigationRule::Wrap:
 					CurrentSourceSide = DestSideFunc(SweptRect);
-					FVector2D SampleSpot = WidgetRect.GetCenter();
+					FVector2f SampleSpot = WidgetRect.GetCenter();
 					SampleSpot[AxisIndex] = CurrentSourceSide;
 					CurrentCellPoint = GetCellCoordinate(SampleSpot);
 					bWrapped = true;
@@ -467,7 +468,7 @@ TSharedPtr<SWidget> FHittestGrid::FindFocusableWidget(FSlateRect WidgetRect, con
 					break;
 				}
 				CurrentSourceSide = DestSideFunc(SweptRect);
-				FVector2D SampleSpot = WidgetRect.GetCenter();
+				FVector2f SampleSpot = WidgetRect.GetCenter();
 				SampleSpot[AxisIndex] = CurrentSourceSide;
 				CurrentCellPoint = GetCellCoordinate(SampleSpot);
 				bWrapped = true;
@@ -571,7 +572,7 @@ TSharedPtr<SWidget> FHittestGrid::FindNextFocusableWidget(const FArrangedWidget&
 	return Widget;
 }
 
-FIntPoint FHittestGrid::GetCellCoordinate(FVector2D Position) const
+FIntPoint FHittestGrid::GetCellCoordinate(UE::Slate::FDeprecateVector2DParameter Position) const
 {
 	return FIntPoint(
 		FMath::Min(FMath::Max(FMath::FloorToInt(Position.X / CellSize.X), 0), NumCells.X - 1),
@@ -853,7 +854,7 @@ FHittestGrid::FIndexAndDistance FHittestGrid::GetHitIndexFromCellIndex(const FGr
 			const bool bIsValidWidget = TestWidget.IsValid() && (!Params.bTestWidgetIsInteractive || TestWidget->IsInteractable());
 			if (bIsValidWidget)
 			{
-				const FVector2D WindowSpaceCoordinate = Params.CursorPositionInGrid + GridWindowOrigin;
+				const FVector2f WindowSpaceCoordinate = Params.CursorPositionInGrid + GridWindowOrigin;
 
 				const FGeometry& TestGeometry = TestWidget->GetPaintSpaceGeometry();
 
diff --git a/Engine/Source/Runtime/SlateCore/Private/Layout/Clipping.cpp b/Engine/Source/Runtime/SlateCore/Private/Layout/Clipping.cpp
--- a/Engine/Source/Runtime/SlateCore/Private/Layout/Clipping.cpp
+++ b/Engine/Source/Runtime/SlateCore/Private/Layout/Clipping.cpp
@@ -45,12 +45,12 @@ FSlateClippingZone::FSlateClippingZone(const FGeometry& BooundingGeometry)
 	, bAlwaysClip(false)
 {
 	const FSlateRenderTransform& Transform = BooundingGeometry.GetAccumulatedRenderTransform();
-	const FVector2D& LocalSize = BooundingGeometry.GetLocalSize();
+	FVector2f LocalSize = BooundingGeometry.GetLocalSize();
 
 	InitializeFromArbitraryPoints(
-		Transform.TransformPoint(FVector2D(0,0)),
-		Transform.TransformPoint(FVector2D(LocalSize.X, 0)),
-		Transform.TransformPoint(FVector2D(0, LocalSize.Y)),
+		Transform.TransformPoint(FVector2f(0.f, 0.f)),
+		Transform.TransformPoint(FVector2f(LocalSize.X, 0.f)),
+		Transform.TransformPoint(FVector2f(0.f, LocalSize.Y)),
 		Transform.TransformPoint(LocalSize)
 	);
 }
@@ -60,24 +60,24 @@ FSlateClippingZone::FSlateClippingZone(const FPaintGeometry& PaintingGeometry)
 	, bAlwaysClip(false)
 {
 	const FSlateRenderTransform& Transform = PaintingGeometry.GetAccumulatedRenderTransform();
-	const FVector2D& LocalSize = PaintingGeometry.GetLocalSize();
+	FVector2f LocalSize = PaintingGeometry.GetLocalSize();
 
 	InitializeFromArbitraryPoints(
-		Transform.TransformPoint(FVector2D(0, 0)),
-		Transform.TransformPoint(FVector2D(LocalSize.X, 0)),
-		Transform.TransformPoint(FVector2D(0, LocalSize.Y)),
+		Transform.TransformPoint(FVector2f(0.f, 0.f)),
+		Transform.TransformPoint(FVector2f(LocalSize.X, 0.f)),
+		Transform.TransformPoint(FVector2f(0.f, LocalSize.Y)),
 		Transform.TransformPoint(LocalSize)
 	);
 }
 
-FSlateClippingZone::FSlateClippingZone(const FVector2D& InTopLeft, const FVector2D& InTopRight, const FVector2D& InBottomLeft, const FVector2D& InBottomRight)
+FSlateClippingZone::FSlateClippingZone(const UE::Slate::FDeprecateVector2DParameter& InTopLeft, const UE::Slate::FDeprecateVector2DParameter& InTopRight, const UE::Slate::FDeprecateVector2DParameter& InBottomLeft, const UE::Slate::FDeprecateVector2DParameter& InBottomRight)
 	: bIntersect(true)
 	, bAlwaysClip(false)
 {
 	InitializeFromArbitraryPoints(InTopLeft, InTopRight, InBottomLeft, InBottomRight);
 }
 
-void FSlateClippingZone::InitializeFromArbitraryPoints(const FVector2D& InTopLeft, const FVector2D& InTopRight, const FVector2D& InBottomLeft, const FVector2D& InBottomRight)
+void FSlateClippingZone::InitializeFromArbitraryPoints(const UE::Slate::FDeprecateVector2DParameter& InTopLeft, const UE::Slate::FDeprecateVector2DParameter& InTopRight, const UE::Slate::FDeprecateVector2DParameter& InBottomLeft, const UE::Slate::FDeprecateVector2DParameter& InBottomRight)
 {
 	bIsAxisAligned = false;
 
@@ -162,14 +162,14 @@ FSlateRect FSlateClippingZone::GetBoundingBox() const
 	);
 }
 
-static float VectorSign(const FVector2D& Vec, const FVector2D& A, const FVector2D& B)
+static float VectorSign(const FVector2f& Vec, const FVector2f& A, const FVector2f& B)
 {
 	return FMath::Sign((B.X - A.X) * (Vec.Y - A.Y) - (B.Y - A.Y) * (Vec.X - A.X));
 }
 
 // Returns true when the point is inside the triangle
 // Should not return true when the point is on one of the edges
-static bool IsPointInTriangle(const FVector2D& TestPoint, const FVector2D& A, const FVector2D& B, const FVector2D& C)
+static bool IsPointInTriangle(const FVector2f& TestPoint, const FVector2f& A, const FVector2f& B, const FVector2f& C)
 {
 	float BA = VectorSign(B, A, TestPoint);
 	float CB = VectorSign(C, B, TestPoint);
@@ -180,7 +180,7 @@ static bool IsPointInTriangle(const FVector2D& TestPoint, const FVector2D& A, co
 	return BA == CB && CB == AC;
 }
 
-bool FSlateClippingZone::IsPointInside(const FVector2D& Point) const
+bool FSlateClippingZone::IsPointInside(const UE::Slate::FDeprecateVector2DParameter& Point) const
 {
 	if (IsAxisAligned())
 	{
@@ -188,7 +188,7 @@ bool FSlateClippingZone::IsPointInside(const FVector2D& Point) const
 	}
 	else
 	{
-		if (IsPointInTriangle(Point, FVector2D(TopLeft), FVector2D(TopRight), FVector2D(BottomLeft)) || IsPointInTriangle(Point, FVector2D(BottomLeft), FVector2D(TopRight), FVector2D(BottomRight)))
+		if (IsPointInTriangle(Point, TopLeft, TopRight, BottomLeft) || IsPointInTriangle(Point, BottomLeft, TopRight, BottomRight))
 		{
 			return true;
 		}
@@ -219,7 +219,7 @@ FSlateClippingState::FSlateClippingState(const FSlateClippingState& Other)
 {
 }
 
-bool FSlateClippingState::IsPointInside(const FVector2D& Point) const
+bool FSlateClippingState::IsPointInside(const UE::Slate::FDeprecateVector2DParameter& Point) const
 {
 	if (ScissorRect.IsSet())
 	{
diff --git a/Engine/Source/Runtime/SlateCore/Private/Layout/Geometry.cpp b/Engine/Source/Runtime/SlateCore/Private/Layout/Geometry.cpp
--- a/Engine/Source/Runtime/SlateCore/Private/Layout/Geometry.cpp
+++ b/Engine/Source/Runtime/SlateCore/Private/Layout/Geometry.cpp
@@ -24,6 +24,6 @@ FGeometry& FGeometry::operator=(const FGeometry& RHS)
 
 FString FGeometry::ToString() const
 {
-	return FString::Printf(TEXT("[Abs=%s, Scale=%.2f, Size=%s]"), *AbsolutePosition.ToString(), Scale, *Size.ToString());
+	return FString::Printf(TEXT("[Abs=%s, Scale=%.2f, Size=%s]"), *AbsolutePosition.ToString(), Scale, *FVector2f(Size).ToString());
 }
 
diff --git a/Engine/Source/Runtime/SlateCore/Private/Layout/LayoutUtils.cpp b/Engine/Source/Runtime/SlateCore/Private/Layout/LayoutUtils.cpp
--- a/Engine/Source/Runtime/SlateCore/Private/Layout/LayoutUtils.cpp
+++ b/Engine/Source/Runtime/SlateCore/Private/Layout/LayoutUtils.cpp
@@ -2,7 +2,7 @@
 
 #include "Layout/LayoutUtils.h"
 
-FVector2D ComputePopupFitInRect(const FSlateRect& InAnchor, const FSlateRect& PopupRect, const EOrientation& Orientation, const FSlateRect& RectToFit)
+UE::Slate::FDeprecateVector2DResult ComputePopupFitInRect(const FSlateRect& InAnchor, const FSlateRect& PopupRect, const EOrientation& Orientation, const FSlateRect& RectToFit)
 {
 	const bool bAdjustmentNeeded = PopupRect.IntersectionWith(RectToFit) != PopupRect;
 	if (bAdjustmentNeeded)
@@ -83,10 +83,10 @@ FVector2D ComputePopupFitInRect(const FSlateRect& InAnchor, const FSlateRect& Po
 
 		NewPosition += Adjust;
 
-		return FVector2D(NewPosition);
+		return FVector2f(NewPosition);
 	}
 	else
 	{
-		return FVector2D(PopupRect.GetTopLeft());
+		return FVector2f(PopupRect.GetTopLeft());
 	}
 }
diff --git a/Engine/Source/Runtime/SlateCore/Private/Layout/SlateRotatedRect.cpp b/Engine/Source/Runtime/SlateCore/Private/Layout/SlateRotatedRect.cpp
--- a/Engine/Source/Runtime/SlateCore/Private/Layout/SlateRotatedRect.cpp
+++ b/Engine/Source/Runtime/SlateCore/Private/Layout/SlateRotatedRect.cpp
@@ -21,7 +21,7 @@ FSlateRect FSlateRotatedRect::ToBoundingRect() const
 		);
 }
 
-bool FSlateRotatedRect::IsUnderLocation(const FVector2f Location) const
+bool FSlateRotatedRect::IsUnderLocation(const UE::Slate::FDeprecateVector2DParameter Location) const
 {
 	const FVector2f Offset = Location - TopLeft;
 	const float Det = FVector2f::CrossProduct(ExtentX, ExtentY);
@@ -36,11 +36,6 @@ bool FSlateRotatedRect::IsUnderLocation(const FVector2f Location) const
 	return false;
 }
 
-bool FSlateRotatedRect::IsUnderLocation(const FVector2D Location) const
-{
-	return IsUnderLocation(UE::Slate::CastToVector2f(Location));
-}
-
 FSlateRotatedRect FSlateRotatedRect::MakeRotatedRect(const FSlateRect& ClipRectInLayoutWindowSpace, const FTransform2f& LayoutToRenderTransform)
 {
 	const FSlateRotatedRect RotatedRect = TransformRect(LayoutToRenderTransform, FSlateRotatedRect(ClipRectInLayoutWindowSpace));
diff --git a/Engine/Source/Runtime/SlateCore/Private/Layout/WidgetPath.cpp b/Engine/Source/Runtime/SlateCore/Private/Layout/WidgetPath.cpp
--- a/Engine/Source/Runtime/SlateCore/Private/Layout/WidgetPath.cpp
+++ b/Engine/Source/Runtime/SlateCore/Private/Layout/WidgetPath.cpp
@@ -325,7 +325,10 @@ FWeakWidgetPath::EPathResolutionResult::Result FWeakWidgetPath::ToWidgetPath( FW
 						{
 							if (PointerEvent && !VirtualPointerPos.IsSet())
 							{
-								VirtualPointerPos = CurWidget->TranslateMouseCoordinateForCustomHitTestChild(*ChildWidgetPtr.Get(), ParentGeometry, PointerEvent->GetScreenSpacePosition(), PointerEvent->GetLastScreenSpacePosition());
+								FVector2f ScreenPos = PointerEvent->GetScreenSpacePosition();
+								FVector2f LastScreenPos = PointerEvent->GetLastScreenSpacePosition();
+
+								VirtualPointerPos = CurWidget->TranslateMouseCoordinateForCustomHitTestChild(*ChildWidgetPtr.Get(), ParentGeometry, FVector2d(ScreenPos), FVector2d(LastScreenPos));
 							}
 
 							bFoundChild = true;
@@ -388,7 +391,10 @@ FWeakWidgetPath::EPathResolutionResult::Result FWeakWidgetPath::ToWidgetPath( FW
 						{
 							if (PointerEvent && !VirtualPointerPos.IsSet())
 							{
-								VirtualPointerPos = CurWidget->TranslateMouseCoordinateForCustomHitTestChild(ArrangedWidget.Widget.Get(), ParentGeometry, PointerEvent->GetScreenSpacePosition(), PointerEvent->GetLastScreenSpacePosition());
+								FVector2f ScreenPos = PointerEvent->GetScreenSpacePosition();
+								FVector2f LastScreenPos = PointerEvent->GetLastScreenSpacePosition();
+
+								VirtualPointerPos = CurWidget->TranslateMouseCoordinateForCustomHitTestChild(ArrangedWidget.Widget.Get(), ParentGeometry, FVector2d(ScreenPos), FVector2d(LastScreenPos));
 							}
 
 							bFoundChild = true;
diff --git a/Engine/Source/Runtime/SlateCore/Private/Rendering/DrawElements.cpp b/Engine/Source/Runtime/SlateCore/Private/Rendering/DrawElements.cpp
--- a/Engine/Source/Runtime/SlateCore/Private/Rendering/DrawElements.cpp
+++ b/Engine/Source/Runtime/SlateCore/Private/Rendering/DrawElements.cpp
@@ -221,12 +221,7 @@ void FSlateDrawElement::Init(FSlateWindowElementList& ElementList, EElementType
 	}
 }
 
-void FSlateDrawElement::ApplyPositionOffset(FVector2d InOffset)
-{
-	ApplyPositionOffset(UE::Slate::CastToVector2f(InOffset));
-}
-
-void FSlateDrawElement::ApplyPositionOffset(FVector2f InOffset)
+void FSlateDrawElement::ApplyPositionOffset(UE::Slate::FDeprecateVector2DParameter InOffset)
 {
 	SetPosition(Position + InOffset);
 	RenderTransform = Concatenate(RenderTransform, InOffset);
@@ -321,33 +316,6 @@ void FSlateDrawElement::MakeBox(
 	MakeBoxInternal(ElementList, InLayer, PaintGeometry, InBrush, InDrawEffects, InTint);
 }
 
-void FSlateDrawElement::MakeBox( 
-	FSlateWindowElementList& ElementList,
-	uint32 InLayer, 
-	const FPaintGeometry& PaintGeometry, 
-	const FSlateBrush* InBrush, 
-	const FSlateResourceHandle& InRenderingHandle,
-	ESlateDrawEffect InDrawEffects, 
-	const FLinearColor& InTint )
-{
-	MakeBox(ElementList, InLayer, PaintGeometry, InBrush, InDrawEffects, InTint);
-}
-
-void FSlateDrawElement::MakeRotatedBox(
-	FSlateWindowElementList& ElementList,
-	uint32 InLayer,
-	const FPaintGeometry& PaintGeometry,
-	const FSlateBrush* InBrush,
-	ESlateDrawEffect InDrawEffects,
-	float Angle2D,
-	TOptional<FVector2d> InRotationPoint,
-	ERotationSpace RotationSpace,
-	const FLinearColor& InTint)
-{
-	TOptional<FVector2f> RotationPoint = InRotationPoint.IsSet() ? UE::Slate::CastToVector2f(InRotationPoint.GetValue()) : TOptional<FVector2f>();
-	MakeRotatedBox(ElementList, InLayer, PaintGeometry, InBrush, InDrawEffects, Angle2D, RotationPoint, RotationSpace, InTint);
-}
-
 void FSlateDrawElement::MakeRotatedBox(
 	FSlateWindowElementList& ElementList,
 	uint32 InLayer,
@@ -355,7 +323,7 @@ void FSlateDrawElement::MakeRotatedBox(
 	const FSlateBrush* InBrush,
 	ESlateDrawEffect InDrawEffects,
 	float Angle2D,
-	TOptional<FVector2f> InRotationPoint,
+	UE::Slate::FDeprecateOptionalVector2DParameter InRotationPoint,
 	ERotationSpace RotationSpace,
 	const FLinearColor& InTint)
 {
@@ -519,12 +487,7 @@ void FSlateDrawElement::MakeGradient( FSlateWindowElementList& ElementList, uint
 	Element.Init(ElementList, EElementType::ET_Gradient, InLayer, PaintGeometry, InDrawEffects);
 }
 
-void FSlateDrawElement::MakeSpline(FSlateWindowElementList& ElementList, uint32 InLayer, const FPaintGeometry& PaintGeometry, const FVector2d InStart, const FVector2d InStartDir, const FVector2d InEnd, const FVector2d InEndDir, float InThickness, ESlateDrawEffect InDrawEffects, const FLinearColor& InTint)
-{
-	MakeSpline(ElementList, InLayer, PaintGeometry, UE::Slate::CastToVector2f(InStart), UE::Slate::CastToVector2f(InStartDir), UE::Slate::CastToVector2f(InEnd), UE::Slate::CastToVector2f(InEndDir), InThickness, InDrawEffects, InTint);
-}
-
-void FSlateDrawElement::MakeSpline(FSlateWindowElementList& ElementList, uint32 InLayer, const FPaintGeometry& PaintGeometry, const FVector2f InStart, const FVector2f InStartDir, const FVector2f InEnd, const FVector2f InEndDir, float InThickness, ESlateDrawEffect InDrawEffects, const FLinearColor& InTint)
+void FSlateDrawElement::MakeSpline(FSlateWindowElementList& ElementList, uint32 InLayer, const FPaintGeometry& PaintGeometry, const UE::Slate::FDeprecateVector2DParameter InStart, const UE::Slate::FDeprecateVector2DParameter InStartDir, const UE::Slate::FDeprecateVector2DParameter InEnd, const UE::Slate::FDeprecateVector2DParameter InEndDir, float InThickness, ESlateDrawEffect InDrawEffects, const FLinearColor& InTint)
 {
 	PaintGeometry.CommitTransformsIfUsingLegacyConstructor();
 
@@ -541,12 +504,7 @@ void FSlateDrawElement::MakeSpline(FSlateWindowElementList& ElementList, uint32
 	Element.Init(ElementList, EElementType::ET_Spline, InLayer, PaintGeometry, InDrawEffects);
 }
 
-void FSlateDrawElement::MakeCubicBezierSpline(FSlateWindowElementList& ElementList, uint32 InLayer, const FPaintGeometry& PaintGeometry, const FVector2d P0, const FVector2d P1, const FVector2d P2, const FVector2d P3, float InThickness, ESlateDrawEffect InDrawEffects, const FLinearColor& InTint)
-{
-	MakeCubicBezierSpline(ElementList, InLayer, PaintGeometry, UE::Slate::CastToVector2f(P0), UE::Slate::CastToVector2f(P1), UE::Slate::CastToVector2f(P2), UE::Slate::CastToVector2f(P3), InThickness, InDrawEffects, InTint);
-}
-
-void FSlateDrawElement::MakeCubicBezierSpline(FSlateWindowElementList& ElementList, uint32 InLayer, const FPaintGeometry& PaintGeometry, const FVector2f P0, const FVector2f P1, const FVector2f P2, const FVector2f P3, float InThickness, ESlateDrawEffect InDrawEffects, const FLinearColor& InTint)
+void FSlateDrawElement::MakeCubicBezierSpline(FSlateWindowElementList& ElementList, uint32 InLayer, const FPaintGeometry& PaintGeometry, const UE::Slate::FDeprecateVector2DParameter P0, const UE::Slate::FDeprecateVector2DParameter P1, const UE::Slate::FDeprecateVector2DParameter P2, const UE::Slate::FDeprecateVector2DParameter P3, float InThickness, ESlateDrawEffect InDrawEffects, const FLinearColor& InTint)
 {
 	PaintGeometry.CommitTransformsIfUsingLegacyConstructor();
 
@@ -563,52 +521,12 @@ void FSlateDrawElement::MakeCubicBezierSpline(FSlateWindowElementList& ElementLi
 	Element.Init(ElementList, EElementType::ET_Spline, InLayer, PaintGeometry, InDrawEffects);
 }
 
-void FSlateDrawElement::MakeDrawSpaceSpline(FSlateWindowElementList& ElementList, uint32 InLayer, const FVector2d InStart, const FVector2d InStartDir, const FVector2d InEnd, const FVector2d InEndDir, float InThickness, ESlateDrawEffect InDrawEffects, const FLinearColor& InTint)
-{
-	MakeSpline(ElementList, InLayer, FPaintGeometry(), InStart, InStartDir, InEnd, InEndDir, InThickness, InDrawEffects, InTint);
-}
-
-void FSlateDrawElement::MakeDrawSpaceSpline(FSlateWindowElementList& ElementList, uint32 InLayer, const FVector2f InStart, const FVector2f InStartDir, const FVector2f InEnd, const FVector2f InEndDir, float InThickness, ESlateDrawEffect InDrawEffects, const FLinearColor& InTint)
+void FSlateDrawElement::MakeDrawSpaceSpline(FSlateWindowElementList& ElementList, uint32 InLayer, const UE::Slate::FDeprecateVector2DParameter InStart, const UE::Slate::FDeprecateVector2DParameter InStartDir, const UE::Slate::FDeprecateVector2DParameter InEnd, const UE::Slate::FDeprecateVector2DParameter InEndDir, float InThickness, ESlateDrawEffect InDrawEffects, const FLinearColor& InTint)
 {
 	MakeSpline(ElementList, InLayer, FPaintGeometry(), InStart, InStartDir, InEnd, InEndDir, InThickness, InDrawEffects, InTint);
 }
 
-void FSlateDrawElement::MakeDrawSpaceGradientSpline(FSlateWindowElementList& ElementList, uint32 InLayer, const FVector2D& InStart, const FVector2D& InStartDir, const FVector2D& InEnd, const FVector2D& InEndDir, const TArray<FSlateGradientStop>& InGradientStops, float InThickness, ESlateDrawEffect InDrawEffects)
-{
-	const FPaintGeometry PaintGeometry;
-	PaintGeometry.CommitTransformsIfUsingLegacyConstructor();
-
-	if (ShouldCull(ElementList))
-	{
-		return;
-	}
-
-	FSlateDrawElement& Element = ElementList.AddUninitialized(EElementType::ET_Spline);
-
-	FSlateSplinePayload& DataPayload = ElementList.CreatePayload<FSlateSplinePayload>(Element);
-	DataPayload.SetGradientHermiteSpline(InStart, InStartDir, InEnd, InEndDir, InThickness, InGradientStops);
-
-	Element.Init(ElementList, EElementType::ET_Spline, InLayer, PaintGeometry, InDrawEffects);
-}
-
-void FSlateDrawElement::MakeDrawSpaceGradientSpline(FSlateWindowElementList& ElementList, uint32 InLayer, const FVector2D& InStart, const FVector2D& InStartDir, const FVector2D& InEnd, const FVector2D& InEndDir, const FSlateRect InClippingRect, const TArray<FSlateGradientStop>& InGradientStops, float InThickness, ESlateDrawEffect InDrawEffects)
-{
-	const FPaintGeometry PaintGeometry;
-	PaintGeometry.CommitTransformsIfUsingLegacyConstructor();
-
-	if (ShouldCull(ElementList))
-	{
-		return;
-	}
-
-	FSlateDrawElement& Element = ElementList.AddUninitialized(EElementType::ET_Spline);
-
-	FSlateSplinePayload& DataPayload = ElementList.CreatePayload<FSlateSplinePayload>(Element);
-	DataPayload.SetGradientHermiteSpline(InStart, InStartDir, InEnd, InEndDir, InThickness, InGradientStops);
-
-	Element.Init(ElementList, EElementType::ET_Spline, InLayer, PaintGeometry, InDrawEffects);
-}
-
+#if UE_ENABLE_SLATE_VECTOR_DEPRECATION_MECHANISMS
 void FSlateDrawElement::MakeLines(FSlateWindowElementList& ElementList, uint32 InLayer, const FPaintGeometry& PaintGeometry, const TArray<FVector2d>& Points, ESlateDrawEffect InDrawEffects, const FLinearColor& InTint, bool bAntialias, float Thickness)
 {
 	if (ShouldCull(ElementList) || Points.Num() < 2)
@@ -624,6 +542,7 @@ void FSlateDrawElement::MakeLines(FSlateWindowElementList& ElementList, uint32 I
 	}
 	MakeLines(ElementList, InLayer, PaintGeometry, MoveTemp(NewVector), InDrawEffects, InTint, bAntialias, Thickness);
 }
+#endif // UE_ENABLE_SLATE_VECTOR_DEPRECATION_MECHANISMS
 
 void FSlateDrawElement::MakeLines(FSlateWindowElementList& ElementList, uint32 InLayer, const FPaintGeometry& PaintGeometry, TArray<FVector2f> Points, ESlateDrawEffect InDrawEffects, const FLinearColor& InTint, bool bAntialias, float Thickness)
 {
@@ -653,6 +572,7 @@ void FSlateDrawElement::MakeLines(FSlateWindowElementList& ElementList, uint32 I
 	Element.Init(ElementList, EElementType::ET_Line, InLayer, PaintGeometry, DrawEffects);
 }
 
+#if UE_ENABLE_SLATE_VECTOR_DEPRECATION_MECHANISMS
 void FSlateDrawElement::MakeLines(FSlateWindowElementList& ElementList, uint32 InLayer, const FPaintGeometry& PaintGeometry, const TArray<FVector2d>& Points, const TArray<FLinearColor>& PointColors, ESlateDrawEffect InDrawEffects, const FLinearColor& InTint, bool bAntialias, float Thickness)
 {
 	if (ShouldCull(ElementList) || Points.Num() < 2)
@@ -668,6 +588,7 @@ void FSlateDrawElement::MakeLines(FSlateWindowElementList& ElementList, uint32 I
 	}
 	MakeLines(ElementList, InLayer, PaintGeometry, MoveTemp(NewVector), PointColors, InDrawEffects, InTint, bAntialias, Thickness);
 }
+#endif
 
 void FSlateDrawElement::MakeLines( FSlateWindowElementList& ElementList, uint32 InLayer, const FPaintGeometry& PaintGeometry, TArray<FVector2f> Points, TArray<FLinearColor> PointColors, ESlateDrawEffect InDrawEffects, const FLinearColor& InTint, bool bAntialias, float Thickness )
 {
diff --git a/Engine/Source/Runtime/SlateCore/Private/Styling/SlateBrush.cpp b/Engine/Source/Runtime/SlateCore/Private/Styling/SlateBrush.cpp
--- a/Engine/Source/Runtime/SlateCore/Private/Styling/SlateBrush.cpp
+++ b/Engine/Source/Runtime/SlateCore/Private/Styling/SlateBrush.cpp
@@ -32,7 +32,7 @@ FSlateBrush::FSlateBrush( ESlateBrushDrawType::Type InDrawType,
 						  const FMargin& InMargin, 
 						  ESlateBrushTileType::Type InTiling, 
 						  ESlateBrushImageType::Type InImageType, 
-						  const FVector2D& InImageSize, 
+						  const UE::Slate::FDeprecateVector2DParameter& InImageSize, 
 						  const FLinearColor& InTint, 
 						  UObject* InObjectResource, 
 						  bool bInDynamicallyLoaded
@@ -68,7 +68,7 @@ FSlateBrush::FSlateBrush( ESlateBrushDrawType::Type InDrawType,
  						  const FMargin& InMargin,
  						  ESlateBrushTileType::Type InTiling,
  						  ESlateBrushImageType::Type InImageType,
- 						  const FVector2D& InImageSize,
+ 						  const UE::Slate::FDeprecateVector2DParameter& InImageSize,
  						  const TSharedRef< FLinearColor >& InTint,
  						  UObject* InObjectResource, 
  						  bool bInDynamicallyLoaded
@@ -104,7 +104,7 @@ FSlateBrush::FSlateBrush( ESlateBrushDrawType::Type InDrawType,
 						  const FMargin& InMargin,
 						  ESlateBrushTileType::Type InTiling, 
 						  ESlateBrushImageType::Type InImageType, 
-						  const FVector2D& InImageSize, const FSlateColor& InTint, 
+						  const UE::Slate::FDeprecateVector2DParameter& InImageSize, const FSlateColor& InTint, 
 						  UObject* InObjectResource, 
 						  bool bInDynamicallyLoaded
  						)
@@ -151,7 +151,7 @@ const FSlateResourceHandle& FSlateBrush::GetRenderingResource() const
 	return ResourceHandle;
 }
 
-void FSlateBrush::UpdateRenderingResource(FVector2D LocalSize, float DrawScale) const
+void FSlateBrush::UpdateRenderingResource(FVector2f LocalSize, float DrawScale) const
 {
 	if (DrawAs != ESlateBrushDrawType::NoDrawType && (ResourceName != NAME_None || ResourceObject != nullptr))
 	{
@@ -159,7 +159,7 @@ void FSlateBrush::UpdateRenderingResource(FVector2D LocalSize, float DrawScale)
 		// For vector graphics we will rebuild the handle only if the shape needs to be rasterized again and the new size and scale
 		if (!ResourceHandle.IsValid() || ImageType == ESlateBrushImageType::Vector)
 		{
-			ResourceHandle = FSlateApplicationBase::Get().GetRenderer()->GetResourceHandle(*this, UE::Slate::CastToVector2f(LocalSize), DrawScale);
+			ResourceHandle = FSlateApplicationBase::Get().GetRenderer()->GetResourceHandle(*this, LocalSize, DrawScale);
 		}
 	}
 
diff --git a/Engine/Source/Runtime/SlateCore/Private/Styling/SlateStyleSet.cpp b/Engine/Source/Runtime/SlateCore/Private/Styling/SlateStyleSet.cpp
--- a/Engine/Source/Runtime/SlateCore/Private/Styling/SlateStyleSet.cpp
+++ b/Engine/Source/Runtime/SlateCore/Private/Styling/SlateStyleSet.cpp
@@ -17,7 +17,7 @@
 
 FSlateStyleSet::FSlateStyleSet(const FName& InStyleSetName)
 : StyleSetName(InStyleSetName)
-, DefaultBrush(new FSlateImageBrush(FPaths::EngineContentDir() / TEXT("Slate/Checkerboard.png"), FVector2D(16, 16), FLinearColor::White, ESlateBrushTileType::Both))
+, DefaultBrush(new FSlateImageBrush(FPaths::EngineContentDir() / TEXT("Slate/Checkerboard.png"), FVector2f(16.f, 16.f), FLinearColor::White, ESlateBrushTileType::Both))
 {
 	// Add a mapping so that this resource will be discovered by GetStyleResources.
 	Set(TEXT("Default"), GetDefaultBrush());
@@ -197,7 +197,7 @@ float FSlateStyleSet::GetFloat(const FName PropertyName, const ANSICHAR* Specifi
 
 FVector2D FSlateStyleSet::GetVector(const FName PropertyName, const ANSICHAR* Specifier, FVector2D DefaultValue, const ISlateStyle* RequestingStyle) const
 {
-	const FVector2D* const Result = Vector2DValues.Find(Join(PropertyName, Specifier));
+	const FVector2f* const Result = Vector2DValues.Find(Join(PropertyName, Specifier));
 
 	if (Result == nullptr) 
 	{
@@ -213,7 +213,7 @@ FVector2D FSlateStyleSet::GetVector(const FName PropertyName, const ANSICHAR* Sp
 		ReportingStyle->LogMissingResource(EStyleMessageSeverity::Warning, FText::Format(NSLOCTEXT("SlateStyleSet", "UknownSlateVector", "Unable to find vector property '{0}' in style."), FText::FromName(PropertyName)), PropertyName);
 	}
 
-	return Result ? *Result : DefaultValue;
+	return Result ? FVector2D(*Result) : DefaultValue;
 }
 
 const FLinearColor& FSlateStyleSet::GetColor(const FName PropertyName, const ANSICHAR* Specifier, const FLinearColor& DefaultValue, const ISlateStyle* RequestingStyle) const
diff --git a/Engine/Source/Runtime/SlateCore/Private/Styling/SlateTypes.cpp b/Engine/Source/Runtime/SlateCore/Private/Styling/SlateTypes.cpp
--- a/Engine/Source/Runtime/SlateCore/Private/Styling/SlateTypes.cpp
+++ b/Engine/Source/Runtime/SlateCore/Private/Styling/SlateTypes.cpp
@@ -76,7 +76,7 @@ void FCheckBoxStyle::PostSerialize(const FArchive& Ar)
 FTextBlockStyle::FTextBlockStyle()
 	: Font()
 	, ColorAndOpacity()
-	, ShadowOffset(FVector2D::ZeroVector)
+	, ShadowOffset(FVector2f::ZeroVector)
 	, ShadowColorAndOpacity(FLinearColor::Black)
 	, SelectedBackgroundColor(FSlateColor::UseForeground())
 	, HighlightColor(FLinearColor::Black)
@@ -148,7 +148,7 @@ void FButtonStyle::PostSerialize(const FArchive& Ar)
 #endif
 
 FComboButtonStyle::FComboButtonStyle()
-	: ShadowOffset(FVector2D::ZeroVector)
+	: ShadowOffset(FVector2f::ZeroVector)
 	, ShadowColorAndOpacity(FLinearColor::Black)
 	, MenuBorderPadding(0.0f)
 	, ContentPadding(5.0f)
diff --git a/Engine/Source/Runtime/SlateCore/Private/Styling/StarshipCoreStyle.cpp b/Engine/Source/Runtime/SlateCore/Private/Styling/StarshipCoreStyle.cpp
--- a/Engine/Source/Runtime/SlateCore/Private/Styling/StarshipCoreStyle.cpp
+++ b/Engine/Source/Runtime/SlateCore/Private/Styling/StarshipCoreStyle.cpp
@@ -257,7 +257,7 @@ TSharedRef<ISlateStyle> FStarshipCoreStyle::Create()
 
 	// Demo Recording
 	{
-		Style->Set("DemoRecording.CursorPing", new IMAGE_BRUSH("Common/CursorPing", FVector2D(31.f, 31.f)));
+		Style->Set("DemoRecording.CursorPing", new IMAGE_BRUSH("Common/CursorPing", FVector2f(31.f, 31.f)));
 	}
 
 	// Error Reporting
@@ -280,15 +280,15 @@ TSharedRef<ISlateStyle> FStarshipCoreStyle::Create()
 
 		// chrome icon or not (we are testing this)
 #if 0
-		Style->Set("AppIcon", new IMAGE_BRUSH_SVG("Starship/Common/unreal", FVector2D(36, 36), FStyleColors::Foreground));
+		Style->Set("AppIcon", new IMAGE_BRUSH_SVG("Starship/Common/unreal", FVector2f(36.f, 36.f), FStyleColors::Foreground));
 		Style->Set("AppIconPadding", FMargin(11.f, 11.f, 3.f, 5.f));
 #else
-		Style->Set("AppIcon", new IMAGE_BRUSH("Starship/Common/UELogo", FVector2D(45, 45), FStyleColors::White));
+		Style->Set("AppIcon", new IMAGE_BRUSH("Starship/Common/UELogo", FVector2f(45.f, 45.f), FStyleColors::White));
 		Style->Set("AppIconPadding", FMargin(5.f, 5.f, 5.f, 5.f));
 #endif
 
-		Style->Set("UnrealCircle.Thick", new IMAGE_BRUSH_SVG("Starship/Common/unreal-circle-thick", FVector2D(32, 32), FStyleColors::Foreground));
-		Style->Set("UnrealCircle.Thin", new IMAGE_BRUSH_SVG("Starship/Common/unreal-circle-thin", FVector2D(32, 32)));
+		Style->Set("UnrealCircle.Thick", new IMAGE_BRUSH_SVG("Starship/Common/unreal-circle-thick", FVector2f(32.f, 32.f), FStyleColors::Foreground));
+		Style->Set("UnrealCircle.Thin", new IMAGE_BRUSH_SVG("Starship/Common/unreal-circle-thin", FVector2f(32.f, 32.f)));
 		Style->Set("AppIcon.Small", new IMAGE_BRUSH_SVG("Starship/Common/unreal-small", Icon24x24, FStyleColors::Foreground));
 		Style->Set("AppIconPadding.Small", FMargin(4.f, 4.f, 0.f, 0.f));
 
@@ -452,8 +452,8 @@ TSharedRef<ISlateStyle> FStarshipCoreStyle::Create()
 	}
 
 	// SVerticalBox Drag& Drop icon
-	Style->Set("VerticalBoxDragIndicator", new IMAGE_BRUSH("Common/VerticalBoxDragIndicator", FVector2D(6.f, 45.f)));
-	Style->Set("VerticalBoxDragIndicatorShort", new IMAGE_BRUSH("Common/VerticalBoxDragIndicatorShort", FVector2D(6.f, 15.f)));
+	Style->Set("VerticalBoxDragIndicator", new IMAGE_BRUSH("Common/VerticalBoxDragIndicator", FVector2f(6.f, 45.f)));
+	Style->Set("VerticalBoxDragIndicatorShort", new IMAGE_BRUSH("Common/VerticalBoxDragIndicatorShort", FVector2f(6.f, 15.f)));
 
 	SetupButtonStyles(Style);
 	SetupComboButtonStyles(Style);
@@ -515,8 +515,8 @@ TSharedRef<ISlateStyle> FStarshipCoreStyle::Create()
 	{
 		Style->Set("ProgressBar", FProgressBarStyle()
 			.SetBackgroundImage(FSlateColorBrush(FStyleColors::Foldout))
-			.SetFillImage(IMAGE_BRUSH("Starship/CoreWidgets/ProgressBar/ProgressMarquee", FVector2D(20.f, 12.f), FStyleColors::PrimaryHover, ESlateBrushTileType::Horizontal))
-			.SetMarqueeImage(IMAGE_BRUSH("Starship/CoreWidgets/ProgressBar/ProgressMarquee", FVector2D(20.f, 12.f), FStyleColors::PrimaryHover, ESlateBrushTileType::Horizontal))
+			.SetFillImage(IMAGE_BRUSH("Starship/CoreWidgets/ProgressBar/ProgressMarquee", FVector2f(20.f, 12.f), FStyleColors::PrimaryHover, ESlateBrushTileType::Horizontal))
+			.SetMarqueeImage(IMAGE_BRUSH("Starship/CoreWidgets/ProgressBar/ProgressMarquee", FVector2f(20.f, 12.f), FStyleColors::PrimaryHover, ESlateBrushTileType::Horizontal))
 			.SetEnableFillAnimation(true)
 		);
 	}
@@ -591,7 +591,7 @@ TSharedRef<ISlateStyle> FStarshipCoreStyle::Create()
 
 		Style->Set("NumericEntrySpinBox.Decorator", new BOX_BRUSH("Common/TextBoxLabelBorder", FMargin(5.0f / 16.0f)));
 
-		Style->Set("NumericEntrySpinBox.NarrowDecorator", new IMAGE_BRUSH_SVG("Starship/CoreWidgets/NumericEntryBox/NarrowDecorator", FVector2D(2.0f,16.0f)));
+		Style->Set("NumericEntrySpinBox.NarrowDecorator", new IMAGE_BRUSH_SVG("Starship/CoreWidgets/NumericEntryBox/NarrowDecorator", FVector2f(2.0f,16.0f)));
 	}
 
 	SetupColorPickerStyles(Style);
@@ -640,7 +640,7 @@ TSharedRef<ISlateStyle> FStarshipCoreStyle::Create()
 	{
 		Style->Set("Wizard.PageTitle", FTextBlockStyle(NormalText)
 			.SetFont(DEFAULT_FONT("BoldCondensed", 28))
-			.SetShadowOffset(FVector2D(1.f, 1.f))
+			.SetShadowOffset(FVector2f(1.f, 1.f))
 			.SetShadowColorAndOpacity(FLinearColor(0.f, 0.f, 0.f, 0.9f))
 		);
 	}
@@ -654,9 +654,9 @@ TSharedRef<ISlateStyle> FStarshipCoreStyle::Create()
 		Style->Set("NotificationList.ItemBackground_Border", new FSlateRoundedBoxBrush(FStyleColors::Notifications, 5.0f, FStyleColors::Black, 1.0f));
 		Style->Set("NotificationList.ItemBackground_Border_Transparent", new BOX_BRUSH("Old/Notification_Border_Flash", FMargin(8.0f / 64.0f)));
 
-		Style->Set("NotificationList.Throbber", new IMAGE_BRUSH("Starship/Notifications/Throbber", FVector2D(40, 40)));
+		Style->Set("NotificationList.Throbber", new IMAGE_BRUSH("Starship/Notifications/Throbber", FVector2f(40.f, 40.f)));
 
-		Style->Set("NotificationList.Watermark", new IMAGE_BRUSH_SVG("Starship/Common/unreal-circle-thin", FVector2D(128, 128)));
+		Style->Set("NotificationList.Watermark", new IMAGE_BRUSH_SVG("Starship/Common/unreal-circle-thin", FVector2f(128.f, 128.f)));
 
 		// Should be deleted. everywhere using this is looking for a generic "alert" icon
 		Style->Set("NotificationList.DefaultMessage", new IMAGE_BRUSH("Common/EventMessage_Default", Icon40x40));
@@ -684,10 +684,10 @@ TSharedRef<ISlateStyle> FStarshipCoreStyle::Create()
 	// SScrollBox defaults...
 	{
 		Style->Set("ScrollBox", FScrollBoxStyle()
-			.SetTopShadowBrush(BOX_BRUSH("Common/ScrollBoxShadowTop", FVector2D(16.f, 8.f), FMargin(0.5f, 1.f, 0.5f, 0.f)))
-			.SetBottomShadowBrush(BOX_BRUSH("Common/ScrollBoxShadowBottom", FVector2D(16.f, 8.f), FMargin(0.5f, 0.f, 0.5f, 1.f)))
-			.SetLeftShadowBrush(BOX_BRUSH("Common/ScrollBoxShadowLeft", FVector2D(8.f, 16.f), FMargin(1.f, 0.5f, 0.f, 0.5f)))
-			.SetRightShadowBrush(BOX_BRUSH("Common/ScrollBoxShadowRight", FVector2D(8.f, 16.f), FMargin(0.f, 0.5f, 1.f, 0.5f)))
+			.SetTopShadowBrush(BOX_BRUSH("Common/ScrollBoxShadowTop", FVector2f(16.f, 8.f), FMargin(0.5f, 1.f, 0.5f, 0.f)))
+			.SetBottomShadowBrush(BOX_BRUSH("Common/ScrollBoxShadowBottom", FVector2f(16.f, 8.f), FMargin(0.5f, 0.f, 0.5f, 1.f)))
+			.SetLeftShadowBrush(BOX_BRUSH("Common/ScrollBoxShadowLeft", FVector2f(8.f, 16.f), FMargin(1.f, 0.5f, 0.f, 0.5f)))
+			.SetRightShadowBrush(BOX_BRUSH("Common/ScrollBoxShadowRight", FVector2f(8.f, 16.f), FMargin(0.f, 0.5f, 1.f, 0.5f)))
 			.SetBarThickness(8.0)
 		);
 
@@ -703,8 +703,8 @@ TSharedRef<ISlateStyle> FStarshipCoreStyle::Create()
 	// SScrollBorder defaults...
 	{
 		Style->Set("ScrollBorder", FScrollBorderStyle()
-			.SetTopShadowBrush(BOX_BRUSH("Common/ScrollBorderShadowTop", FVector2D(16.f, 8.f), FMargin(0.5f, 1.f, 0.5f, 0.f)))
-			.SetBottomShadowBrush(BOX_BRUSH("Common/ScrollBorderShadowBottom", FVector2D(16.f, 8.f), FMargin(0.5f, 0.f, 0.5f, 1.f)))
+			.SetTopShadowBrush(BOX_BRUSH("Common/ScrollBorderShadowTop", FVector2f(16.f, 8.f), FMargin(0.5f, 1.f, 0.5f, 0.f)))
+			.SetBottomShadowBrush(BOX_BRUSH("Common/ScrollBorderShadowBottom", FVector2f(16.f, 8.f), FMargin(0.5f, 0.f, 0.5f, 1.f)))
 		);
 	}
 
@@ -714,24 +714,24 @@ TSharedRef<ISlateStyle> FStarshipCoreStyle::Create()
 	{
 #if !PLATFORM_MAC
 		const FButtonStyle MinimizeButtonStyle = FButtonStyle(NoBorder)
-			.SetNormal( IMAGE_BRUSH_SVG("Starship/CoreWidgets/Window/minimize", FVector2D(42.0f, 34.0f), FStyleColors::Foreground))
-			.SetHovered(IMAGE_BRUSH_SVG("Starship/CoreWidgets/Window/minimize", FVector2D(42.0f, 34.0f), FStyleColors::ForegroundHover))
-			.SetPressed(IMAGE_BRUSH_SVG("Starship/CoreWidgets/Window/minimize", FVector2D(42.0f, 34.0f), FStyleColors::Foreground));
+			.SetNormal( IMAGE_BRUSH_SVG("Starship/CoreWidgets/Window/minimize", FVector2f(42.0f, 34.0f), FStyleColors::Foreground))
+			.SetHovered(IMAGE_BRUSH_SVG("Starship/CoreWidgets/Window/minimize", FVector2f(42.0f, 34.0f), FStyleColors::ForegroundHover))
+			.SetPressed(IMAGE_BRUSH_SVG("Starship/CoreWidgets/Window/minimize", FVector2f(42.0f, 34.0f), FStyleColors::Foreground));
 
 		const FButtonStyle MaximizeButtonStyle = FButtonStyle(NoBorder)
-			.SetNormal( IMAGE_BRUSH_SVG("Starship/CoreWidgets/Window/maximize", FVector2D(42.0f, 34.0f), FStyleColors::Foreground))
-			.SetHovered(IMAGE_BRUSH_SVG("Starship/CoreWidgets/Window/maximize", FVector2D(42.0f, 34.0f), FStyleColors::ForegroundHover))
-			.SetPressed(IMAGE_BRUSH_SVG("Starship/CoreWidgets/Window/maximize", FVector2D(42.0f, 34.0f), FStyleColors::Foreground));
+			.SetNormal( IMAGE_BRUSH_SVG("Starship/CoreWidgets/Window/maximize", FVector2f(42.0f, 34.0f), FStyleColors::Foreground))
+			.SetHovered(IMAGE_BRUSH_SVG("Starship/CoreWidgets/Window/maximize", FVector2f(42.0f, 34.0f), FStyleColors::ForegroundHover))
+			.SetPressed(IMAGE_BRUSH_SVG("Starship/CoreWidgets/Window/maximize", FVector2f(42.0f, 34.0f), FStyleColors::Foreground));
 
 		const FButtonStyle RestoreButtonStyle = FButtonStyle(NoBorder)
-			.SetNormal( IMAGE_BRUSH_SVG("Starship/CoreWidgets/Window/restore", FVector2D(42.0f, 34.0f), FStyleColors::Foreground))
-			.SetHovered(IMAGE_BRUSH_SVG("Starship/CoreWidgets/Window/restore", FVector2D(42.0f, 34.0f), FStyleColors::ForegroundHover))
-			.SetPressed(IMAGE_BRUSH_SVG("Starship/CoreWidgets/Window/restore", FVector2D(42.0f, 34.0f), FStyleColors::Foreground));
+			.SetNormal( IMAGE_BRUSH_SVG("Starship/CoreWidgets/Window/restore", FVector2f(42.0f, 34.0f), FStyleColors::Foreground))
+			.SetHovered(IMAGE_BRUSH_SVG("Starship/CoreWidgets/Window/restore", FVector2f(42.0f, 34.0f), FStyleColors::ForegroundHover))
+			.SetPressed(IMAGE_BRUSH_SVG("Starship/CoreWidgets/Window/restore", FVector2f(42.0f, 34.0f), FStyleColors::Foreground));
 
 		const FButtonStyle CloseButtonStyle = FButtonStyle(NoBorder)
-			.SetNormal( IMAGE_BRUSH_SVG("Starship/CoreWidgets/Window/close", FVector2D(42.0f, 34.0f), FStyleColors::Foreground))
-			.SetHovered(IMAGE_BRUSH_SVG("Starship/CoreWidgets/Window/close", FVector2D(42.0f, 34.0f), FStyleColors::ForegroundHover))
-			.SetPressed(IMAGE_BRUSH_SVG("Starship/CoreWidgets/Window/close", FVector2D(42.0f, 34.0f), FStyleColors::Foreground));
+			.SetNormal( IMAGE_BRUSH_SVG("Starship/CoreWidgets/Window/close", FVector2f(42.0f, 34.0f), FStyleColors::Foreground))
+			.SetHovered(IMAGE_BRUSH_SVG("Starship/CoreWidgets/Window/close", FVector2f(42.0f, 34.0f), FStyleColors::ForegroundHover))
+			.SetPressed(IMAGE_BRUSH_SVG("Starship/CoreWidgets/Window/close", FVector2f(42.0f, 34.0f), FStyleColors::Foreground));
 #endif
 
 		FWindowStyle Window =
@@ -796,7 +796,7 @@ TSharedRef<ISlateStyle> FStarshipCoreStyle::Create()
 	{
 		Style->Set("MessageLog", FTextBlockStyle(NormalText)
 			.SetFont(DEFAULT_FONT("Mono", 9))
-			.SetShadowOffset(FVector2D::ZeroVector)
+			.SetShadowOffset(FVector2f::ZeroVector)
 		);
 		Style->Set("MessageLog.Error", new IMAGE_BRUSH("MessageLog/Log_Error", Icon16x16));
 		Style->Set("MessageLog.Warning", new IMAGE_BRUSH("MessageLog/Log_Warning", Icon16x16));
@@ -876,7 +876,7 @@ TSharedRef<ISlateStyle> FStarshipCoreStyle::Create()
 			.SetCheckedPressedForegroundColor(FStyleColors::ForegroundHover)
 			.SetPressedForegroundColor(FStyleColors::ForegroundHover);
 
-		Style->Set("FilterBar.FilterImage", new IMAGE_BRUSH_SVG("Starship/CoreWidgets/FilterBar/FilterColorSegment", FVector2D(8, 22)));
+		Style->Set("FilterBar.FilterImage", new IMAGE_BRUSH_SVG("Starship/CoreWidgets/FilterBar/FilterColorSegment", FVector2f(8.f, 22.f)));
 		Style->Set("FilterBar.FilterBackground", new FSlateRoundedBoxBrush(FStyleColors::Secondary, 3.0f));
 
 		/* ... and add the new style */
@@ -1133,7 +1133,7 @@ void FStarshipCoreStyle::SetupTextStyles(TSharedRef<FStyle>& Style)
 	const FTextBlockStyle NormalText = FTextBlockStyle()
 		.SetFont(StyleFonts.Normal)
 		.SetColorAndOpacity(FSlateColor::UseForeground())
-		.SetShadowOffset(FVector2D::ZeroVector)
+		.SetShadowOffset(FVector2f::ZeroVector)
 		.SetSelectedBackgroundColor(FStyleColors::Highlight)
 		.SetHighlightColor(FStyleColors::Black)
 		.SetHighlightShape(FSlateColorBrush(FStyleColors::AccentGreen));
@@ -1162,7 +1162,7 @@ void FStarshipCoreStyle::SetupTextStyles(TSharedRef<FStyle>& Style)
 	const FTextBlockStyle MonospacedText = FTextBlockStyle()
 		.SetFont(DEFAULT_FONT("Mono", 10))
 		.SetColorAndOpacity(FSlateColor::UseForeground())
-		.SetShadowOffset(FVector2D::ZeroVector)
+		.SetShadowOffset(FVector2f::ZeroVector)
 		.SetShadowColorAndOpacity(FLinearColor::Black)
 		.SetHighlightColor(FLinearColor(0.02f, 0.3f, 0.0f))
 		.SetHighlightShape(BOX_BRUSH("Common/TextBlockHighlightShape", FMargin(3.f / 8.f))
@@ -1192,7 +1192,7 @@ void FStarshipCoreStyle::SetupTextStyles(TSharedRef<FStyle>& Style)
 	Style->Set("EmbossedText", FTextBlockStyle(NormalText)
 		.SetFont(DEFAULT_FONT("Regular", 24))
 		.SetColorAndOpacity(FLinearColor::Black)
-		.SetShadowOffset(FVector2D(0.0f, 1.0f))
+		.SetShadowOffset(FVector2f(0.0f, 1.0f))
 		.SetShadowColorAndOpacity(FLinearColor(0.8f, 0.8f, 0.8f, 0.5f))
 	);
 
@@ -1443,7 +1443,7 @@ void FStarshipCoreStyle::SetupComboButtonStyles(TSharedRef<FStyle>& Style)
 	FComboButtonStyle ComboButton = FComboButtonStyle()
 		.SetButtonStyle(ComboButtonButton)
 		.SetContentPadding(0.f)
-		.SetDownArrowImage(IMAGE_BRUSH_SVG("Starship/CoreWidgets/ComboBox/wide-chevron-down", FVector2D(16.f, 16.f)))
+		.SetDownArrowImage(IMAGE_BRUSH_SVG("Starship/CoreWidgets/ComboBox/wide-chevron-down", FVector2f(16.f, 16.f)))
 		.SetMenuBorderBrush(FSlateRoundedBoxBrush(FStyleColors::Dropdown, 0.0f, FStyleColors::DropdownOutline, 1.0f))
 		.SetMenuBorderPadding(0.f);
 	Style->Set("ComboButton", ComboButton);
@@ -1469,8 +1469,8 @@ void FStarshipCoreStyle::SetupComboButtonStyles(TSharedRef<FStyle>& Style)
 	FComboButtonStyle SimpleComboButton = FComboButtonStyle()
 		.SetButtonStyle(SimpleButton)
 		.SetContentPadding(0.f)
-		.SetDownArrowImage(IMAGE_BRUSH_SVG("Starship/CoreWidgets/ComboBox/wide-chevron-down", FVector2D(16.f, 16.f)))
-		.SetShadowOffset(FVector2D(0.0f, 0.0f))
+		.SetDownArrowImage(IMAGE_BRUSH_SVG("Starship/CoreWidgets/ComboBox/wide-chevron-down", FVector2f(16.f, 16.f)))
+		.SetShadowOffset(FVector2f(0.0f, 0.0f))
 		.SetMenuBorderBrush(FSlateRoundedBoxBrush(FStyleColors::Dropdown, 0.0f, FStyleColors::DropdownOutline, 1.0f))
 		.SetMenuBorderPadding(0.0f);
 	Style->Set("SimpleComboButton", SimpleComboButton);
@@ -1642,19 +1642,19 @@ void FStarshipCoreStyle::SetupCheckboxStyles(TSharedRef<FStyle>& Style)
 
 
 	SegmentedBoxLeft
-		.SetUncheckedImage(BOX_BRUSH("/Starship/CoreWidgets/SegmentedBox/left", FVector2D(16.f, 16.f), FMargin(4.0f / 16.0f), FStyleColors::Dropdown))
-		.SetUncheckedHoveredImage(BOX_BRUSH("/Starship/CoreWidgets/SegmentedBox/left", FVector2D(16.f, 16.f), FMargin(4.0f / 16.0f), FStyleColors::Hover))
-		.SetUncheckedPressedImage(BOX_BRUSH("/Starship/CoreWidgets/SegmentedBox/left", FVector2D(16.f, 16.f), FMargin(4.0f / 16.0f), FStyleColors::Secondary))
-		.SetCheckedImage(BOX_BRUSH("/Starship/CoreWidgets/SegmentedBox/left", FVector2D(16.f, 16.f), FMargin(4.0f / 16.0f), FStyleColors::Primary))
-		.SetCheckedHoveredImage(BOX_BRUSH("/Starship/CoreWidgets/SegmentedBox/left", FVector2D(16.f, 16.f), FMargin(4.0f / 16.0f), FStyleColors::PrimaryHover))
-		.SetCheckedPressedImage(BOX_BRUSH("/Starship/CoreWidgets/SegmentedBox/left", FVector2D(16.f, 16.f), FMargin(4.0f / 16.0f), FStyleColors::PrimaryPress));
+		.SetUncheckedImage(BOX_BRUSH("/Starship/CoreWidgets/SegmentedBox/left", FVector2f(16.f, 16.f), FMargin(4.0f / 16.0f), FStyleColors::Dropdown))
+		.SetUncheckedHoveredImage(BOX_BRUSH("/Starship/CoreWidgets/SegmentedBox/left", FVector2f(16.f, 16.f), FMargin(4.0f / 16.0f), FStyleColors::Hover))
+		.SetUncheckedPressedImage(BOX_BRUSH("/Starship/CoreWidgets/SegmentedBox/left", FVector2f(16.f, 16.f), FMargin(4.0f / 16.0f), FStyleColors::Secondary))
+		.SetCheckedImage(BOX_BRUSH("/Starship/CoreWidgets/SegmentedBox/left", FVector2f(16.f, 16.f), FMargin(4.0f / 16.0f), FStyleColors::Primary))
+		.SetCheckedHoveredImage(BOX_BRUSH("/Starship/CoreWidgets/SegmentedBox/left", FVector2f(16.f, 16.f), FMargin(4.0f / 16.0f), FStyleColors::PrimaryHover))
+		.SetCheckedPressedImage(BOX_BRUSH("/Starship/CoreWidgets/SegmentedBox/left", FVector2f(16.f, 16.f), FMargin(4.0f / 16.0f), FStyleColors::PrimaryPress));
 
 	Style->Set("SegmentedCombo.Left", SegmentedBoxLeft);
 
 	const FButtonStyle SegmentedComboRightButton = FButtonStyle()
-	.SetNormal(BOX_BRUSH("/Starship/CoreWidgets/SegmentedBox/right", FVector2D(16.f, 16.f), FMargin(4.0f / 16.0f), FStyleColors::Secondary))
-	.SetHovered(BOX_BRUSH("/Starship/CoreWidgets/SegmentedBox/right", FVector2D(16.f, 16.f), FMargin(4.0f / 16.0f), FStyleColors::Hover))
-	.SetPressed(BOX_BRUSH("/Starship/CoreWidgets/SegmentedBox/right", FVector2D(16.f, 16.f), FMargin(4.0f / 16.0f), FStyleColors::Secondary))
+	.SetNormal(BOX_BRUSH("/Starship/CoreWidgets/SegmentedBox/right", FVector2f(16.f, 16.f), FMargin(4.0f / 16.0f), FStyleColors::Secondary))
+	.SetHovered(BOX_BRUSH("/Starship/CoreWidgets/SegmentedBox/right", FVector2f(16.f, 16.f), FMargin(4.0f / 16.0f), FStyleColors::Hover))
+	.SetPressed(BOX_BRUSH("/Starship/CoreWidgets/SegmentedBox/right", FVector2f(16.f, 16.f), FMargin(4.0f / 16.0f), FStyleColors::Secondary))
 	.SetNormalForeground(FStyleColors::Foreground)
 	.SetHoveredForeground(FStyleColors::ForegroundHover)
 	.SetPressedForeground(FStyleColors::Primary)
@@ -1665,7 +1665,7 @@ void FStarshipCoreStyle::SetupCheckboxStyles(TSharedRef<FStyle>& Style)
 	FComboButtonStyle SegmentedComboRight = FComboButtonStyle(Style->GetWidgetStyle<FComboButtonStyle>("ComboButton"))
 		.SetContentPadding(FMargin(6.f, 2.f))
 		.SetButtonStyle(SegmentedComboRightButton)
-		.SetDownArrowImage(IMAGE_BRUSH_SVG("Starship/Common/ellipsis-vertical-narrow", FVector2D(6, 24)));
+		.SetDownArrowImage(IMAGE_BRUSH_SVG("Starship/Common/ellipsis-vertical-narrow", FVector2f(6.f, 24.f)));
 
 
 	Style->Set("SegmentedCombo.Right", SegmentedComboRight);
@@ -1825,13 +1825,13 @@ void FStarshipCoreStyle::SetupDockingStyles(TSharedRef<FStyle>& Style)
 		.SetPressedPadding(FMargin(2, 3, 2, 1));
 
 	const FComboButtonStyle StatusBarComboButton = FComboButtonStyle(Style->GetWidgetStyle<FComboButtonStyle>("ComboButton"))
-		.SetDownArrowImage(IMAGE_BRUSH_SVG("Starship/CoreWidgets/ComboBox/corner-dropdown", FVector2D(7.0f, 7.0f)))
+		.SetDownArrowImage(IMAGE_BRUSH_SVG("Starship/CoreWidgets/ComboBox/corner-dropdown", FVector2f(7.0f, 7.0f)))
 		.SetButtonStyle(StatusBarButton)
 		.SetDownArrowPadding(FMargin(0.0f))
 		.SetDownArrowAlignment(EVerticalAlignment::VAlign_Bottom);
 
 	const FComboButtonStyle StatusBarEllipsisComboButton = FComboButtonStyle(Style->GetWidgetStyle<FComboButtonStyle>("ComboButton"))
-		.SetDownArrowImage(IMAGE_BRUSH_SVG("Starship/Common/ellipsis-vertical-narrow", FVector2D(6, 24)))
+		.SetDownArrowImage(IMAGE_BRUSH_SVG("Starship/Common/ellipsis-vertical-narrow", FVector2f(6.f, 24.f)))
 		.SetButtonStyle(StatusBarButton)
 		.SetDownArrowPadding(FMargin(0.0f))
 		.SetDownArrowAlignment(EVerticalAlignment::VAlign_Bottom);
@@ -1844,7 +1844,7 @@ void FStarshipCoreStyle::SetupDockingStyles(TSharedRef<FStyle>& Style)
 	Style->Set("StatusBar.DrawerBackground", new FSlateRoundedBoxBrush(FStyleColors::Panel, 5.0f));
 
 	Style->Set("StatusBar.Background", new FSlateColorBrush(FStyleColors::Panel));
-	Style->Set("StatusBar.ProgressOverlay", new FSlateRoundedBoxBrush(FStyleColors::Transparent, 2.0f, FStyleColors::Panel, 1.0f, FVector2D(20, 8)));
+	Style->Set("StatusBar.ProgressOverlay", new FSlateRoundedBoxBrush(FStyleColors::Transparent, 2.0f, FStyleColors::Panel, 1.0f, FVector2f(20.f, 8.f)));
 
 	Style->Set("StatusBar.HelpIcon", new IMAGE_BRUSH_SVG("Starship/Common/help", Icon16x16, FStyleColors::Foreground));
 
@@ -1876,14 +1876,14 @@ void FStarshipCoreStyle::SetupDockingStyles(TSharedRef<FStyle>& Style)
 	Style->Set("Docking.Sidebar.PinTabIcon", new IMAGE_BRUSH_SVG("Starship/Docking/pin", Icon12x12));
 
 	// Dock Cross
-	Style->Set("Docking.Cross.DockLeft", new IMAGE_BRUSH("/Docking/OuterDockingIndicator", FVector2D(6.f, 6.f), FLinearColor(1.0f, 0.35f, 0.0f, 0.25f)));
-	Style->Set("Docking.Cross.DockLeft_Hovered", new IMAGE_BRUSH("/Docking/OuterDockingIndicator", FVector2D(6.f, 6.f), FLinearColor(1.0f, 0.35f, 0.0f)));
-	Style->Set("Docking.Cross.DockTop", new IMAGE_BRUSH("/Docking/OuterDockingIndicator", FVector2D(6.f, 6.f), FLinearColor(1.0f, 0.35f, 0.0f, 0.25f)));
-	Style->Set("Docking.Cross.DockTop_Hovered", new IMAGE_BRUSH("/Docking/OuterDockingIndicator", FVector2D(6.f, 6.f), FLinearColor(1.0f, 0.35f, 0.0f)));
-	Style->Set("Docking.Cross.DockRight", new IMAGE_BRUSH("/Docking/OuterDockingIndicator", FVector2D(6.f, 6.f), FLinearColor(1.0f, 0.35f, 0.0f, 0.25f)));
-	Style->Set("Docking.Cross.DockRight_Hovered", new IMAGE_BRUSH("/Docking/OuterDockingIndicator", FVector2D(6.f, 6.f), FLinearColor(1.0f, 0.35f, 0.0f)));
-	Style->Set("Docking.Cross.DockBottom", new IMAGE_BRUSH("/Docking/OuterDockingIndicator", FVector2D(6.f, 6.f), FLinearColor(1.0f, 0.35f, 0.0f, 0.25f)));
-	Style->Set("Docking.Cross.DockBottom_Hovered", new IMAGE_BRUSH("/Docking/OuterDockingIndicator", FVector2D(6.f, 6.f), FLinearColor(1.0f, 0.35f, 0.0f)));
+	Style->Set("Docking.Cross.DockLeft", new IMAGE_BRUSH("/Docking/OuterDockingIndicator", FVector2f(6.f, 6.f), FLinearColor(1.0f, 0.35f, 0.0f, 0.25f)));
+	Style->Set("Docking.Cross.DockLeft_Hovered", new IMAGE_BRUSH("/Docking/OuterDockingIndicator", FVector2f(6.f, 6.f), FLinearColor(1.0f, 0.35f, 0.0f)));
+	Style->Set("Docking.Cross.DockTop", new IMAGE_BRUSH("/Docking/OuterDockingIndicator", FVector2f(6.f, 6.f), FLinearColor(1.0f, 0.35f, 0.0f, 0.25f)));
+	Style->Set("Docking.Cross.DockTop_Hovered", new IMAGE_BRUSH("/Docking/OuterDockingIndicator", FVector2f(6.f, 6.f), FLinearColor(1.0f, 0.35f, 0.0f)));
+	Style->Set("Docking.Cross.DockRight", new IMAGE_BRUSH("/Docking/OuterDockingIndicator", FVector2f(6.f, 6.f), FLinearColor(1.0f, 0.35f, 0.0f, 0.25f)));
+	Style->Set("Docking.Cross.DockRight_Hovered", new IMAGE_BRUSH("/Docking/OuterDockingIndicator", FVector2f(6.f, 6.f), FLinearColor(1.0f, 0.35f, 0.0f)));
+	Style->Set("Docking.Cross.DockBottom", new IMAGE_BRUSH("/Docking/OuterDockingIndicator", FVector2f(6.f, 6.f), FLinearColor(1.0f, 0.35f, 0.0f, 0.25f)));
+	Style->Set("Docking.Cross.DockBottom_Hovered", new IMAGE_BRUSH("/Docking/OuterDockingIndicator", FVector2f(6.f, 6.f), FLinearColor(1.0f, 0.35f, 0.0f)));
 	Style->Set("Docking.Cross.DockCenter", new IMAGE_BRUSH("/Docking/DockingIndicator_Center", Icon64x64, FLinearColor(1.0f, 0.35f, 0.0f, 0.25f)));
 	Style->Set("Docking.Cross.DockCenter_Hovered", new IMAGE_BRUSH("/Docking/DockingIndicator_Center", Icon64x64, FLinearColor(1.0f, 0.35f, 0.0f)));
 
@@ -1910,19 +1910,19 @@ void FStarshipCoreStyle::SetupColorPickerStyles(TSharedRef<FStyle>& Style)
 		Style->Set("ColorPicker.EyeDropper", new IMAGE_BRUSH("Icons/eyedropper_16px", Icon16x16));
 		Style->Set("ColorPicker.Font", FStyleFonts::Get().Normal);
 		Style->Set("ColorPicker.Mode", new IMAGE_BRUSH("Common/ColorPicker_Mode_16x", Icon16x16));
-		Style->Set("ColorPicker.Separator", new IMAGE_BRUSH("Common/ColorPicker_Separator", FVector2D(2.0f, 2.0f)));
+		Style->Set("ColorPicker.Separator", new IMAGE_BRUSH("Common/ColorPicker_Separator", FVector2f(2.0f, 2.0f)));
 		Style->Set("ColorPicker.Selector", new IMAGE_BRUSH("Common/Circle", Icon8x8));
 		Style->Set("ColorPicker.Slider", FSliderStyle()
-			.SetDisabledThumbImage(IMAGE_BRUSH("Common/ColorPicker_SliderHandle", FVector2D(8.0f, 32.0f)))
-			.SetNormalThumbImage(IMAGE_BRUSH("Common/ColorPicker_SliderHandle", FVector2D(8.0f, 32.0f)))
-			.SetHoveredThumbImage(IMAGE_BRUSH("Common/ColorPicker_SliderHandle", FVector2D(8.0f, 32.0f)))
+			.SetDisabledThumbImage(IMAGE_BRUSH("Common/ColorPicker_SliderHandle", FVector2f(8.0f, 32.0f)))
+			.SetNormalThumbImage(IMAGE_BRUSH("Common/ColorPicker_SliderHandle", FVector2f(8.0f, 32.0f)))
+			.SetHoveredThumbImage(IMAGE_BRUSH("Common/ColorPicker_SliderHandle", FVector2f(8.0f, 32.0f)))
 		);
 	}
 
 	// SColorSpectrum defaults...
 	{
-		Style->Set("ColorSpectrum.Spectrum", new IMAGE_BRUSH("Common/ColorSpectrum", FVector2D(256.f, 256.f)));
-		Style->Set("ColorSpectrum.Selector", new IMAGE_BRUSH("Common/Circle", FVector2D(8.f, 8.f)));
+		Style->Set("ColorSpectrum.Spectrum", new IMAGE_BRUSH("Common/ColorSpectrum", FVector2f(256.f, 256.f)));
+		Style->Set("ColorSpectrum.Selector", new IMAGE_BRUSH("Common/Circle", FVector2f(8.f, 8.f)));
 	}
 
 	// SColorThemes defaults...
@@ -1932,14 +1932,14 @@ void FStarshipCoreStyle::SetupColorPickerStyles(TSharedRef<FStyle>& Style)
 
 	// SColorWheel defaults...
 	{
-		Style->Set("ColorWheel.HueValueCircle", new IMAGE_BRUSH("Common/ColorWheel", FVector2D(192.f, 192.f)));
-		Style->Set("ColorWheel.Selector", new IMAGE_BRUSH("Common/Circle", FVector2D(8.f, 8.f)));
+		Style->Set("ColorWheel.HueValueCircle", new IMAGE_BRUSH("Common/ColorWheel", FVector2f(192.f, 192.f)));
+		Style->Set("ColorWheel.Selector", new IMAGE_BRUSH("Common/Circle", FVector2f(8.f, 8.f)));
 	}
 
 	// SColorGradingWheel defaults...
 	{
-		Style->Set("ColorGradingWheel.HueValueCircle", new IMAGE_BRUSH("Common/ColorGradingWheel", FVector2D(192.f, 192.f)));
-		Style->Set("ColorGradingWheel.Selector", new IMAGE_BRUSH("Common/Circle", FVector2D(8.f, 8.f)));
+		Style->Set("ColorGradingWheel.HueValueCircle", new IMAGE_BRUSH("Common/ColorGradingWheel", FVector2f(192.f, 192.f)));
+		Style->Set("ColorGradingWheel.Selector", new IMAGE_BRUSH("Common/Circle", FVector2f(8.f, 8.f)));
 	}
 }
 
@@ -2042,7 +2042,7 @@ void FStarshipCoreStyle::SetupTableViewStyles(TSharedRef<FStyle>& Style)
 		.SetNormalBrush(FSlateColorBrush(FStyleColors::Header))
 		.SetHoveredBrush(FSlateColorBrush(FStyleColors::Dropdown))
 
-		.SetMenuDropdownImage(IMAGE_BRUSH_SVG("Starship/Common/ellipsis-vertical-narrow", FVector2D(6.f, 24.f)))
+		.SetMenuDropdownImage(IMAGE_BRUSH_SVG("Starship/Common/ellipsis-vertical-narrow", FVector2f(6.f, 24.f)))
 		.SetMenuDropdownNormalBorderBrush(FSlateNoResource())
 		.SetMenuDropdownHoveredBorderBrush(FSlateNoResource());
 
@@ -2056,7 +2056,7 @@ void FStarshipCoreStyle::SetupTableViewStyles(TSharedRef<FStyle>& Style)
 		.SetNormalBrush(FSlateColorBrush(FStyleColors::Header))
 		.SetHoveredBrush(FSlateColorBrush(FStyleColors::Dropdown))
 
-		.SetMenuDropdownImage(IMAGE_BRUSH_SVG("Starship/Common/ellipsis-vertical-narrow", FVector2D(6.f, 24.f)))
+		.SetMenuDropdownImage(IMAGE_BRUSH_SVG("Starship/Common/ellipsis-vertical-narrow", FVector2f(6.f, 24.f)))
 		.SetMenuDropdownNormalBorderBrush(FSlateNoResource())
 		.SetMenuDropdownHoveredBorderBrush(FSlateNoResource());
 
@@ -2205,7 +2205,7 @@ void FStarshipCoreStyle::SetupMultiboxStyles(TSharedRef<FStyle>& Style)
 		FComboButtonStyle SlimToolBarSettingsComboButton = FComboButtonStyle(SlimToolBarComboButton)
 			.SetContentPadding(0)
 			.SetButtonStyle(SlimToolBarButtonSettingsRegion)
-			.SetDownArrowImage(IMAGE_BRUSH_SVG("Starship/Common/ellipsis-vertical-narrow", FVector2D(6, 20)));
+			.SetDownArrowImage(IMAGE_BRUSH_SVG("Starship/Common/ellipsis-vertical-narrow", FVector2f(6.f, 20.f)));
 
 		SlimToolbarStyle.SetSettingsComboButtonStyle(SlimToolBarSettingsComboButton);
 
