Slate: Initial phase of FVector2D deprecation

diff --git a/Engine/Source/Runtime/SlateCore/Private/Styling/UMGCoreStyle.cpp b/Engine/Source/Runtime/SlateCore/Private/Styling/UMGCoreStyle.cpp
--- a/Engine/Source/Runtime/SlateCore/Private/Styling/UMGCoreStyle.cpp
+++ b/Engine/Source/Runtime/SlateCore/Private/Styling/UMGCoreStyle.cpp
@@ -224,7 +224,7 @@ TSharedRef<ISlateStyle> FUMGCoreStyle::Create()
 	static const FTextBlockStyle NormalText = FTextBlockStyle()
 		.SetFont(DEFAULT_FONT("Regular", RegularTextSize))
 		.SetColorAndOpacity(FSlateColor::UseForeground())
-		.SetShadowOffset(FVector2D::ZeroVector)
+		.SetShadowOffset(FVector2f::ZeroVector)
 		.SetShadowColorAndOpacity(FLinearColor::Black)
 		.SetHighlightColor(FLinearColor(0.02f, 0.3f, 0.0f))
 		.SetHighlightShape(BOX_BRUSH("Common/TextBlockHighlightShape", FMargin(3.f /8.f)));
@@ -236,7 +236,7 @@ TSharedRef<ISlateStyle> FUMGCoreStyle::Create()
 	static const FTextBlockStyle MonospacedText = FTextBlockStyle()
 		.SetFont(DEFAULT_FONT("Mono", 10))
 		.SetColorAndOpacity(FSlateColor::UseForeground())
-		.SetShadowOffset(FVector2D::ZeroVector)
+		.SetShadowOffset(FVector2f::ZeroVector)
 		.SetShadowColorAndOpacity(FLinearColor::Black)
 		.SetHighlightColor(FLinearColor(0.02f, 0.3f, 0.0f))
 		.SetHighlightShape(BOX_BRUSH("Common/TextBlockHighlightShape", FMargin(3.f/8.f))
@@ -259,7 +259,7 @@ TSharedRef<ISlateStyle> FUMGCoreStyle::Create()
 	Style->Set("EmbossedText", FTextBlockStyle(NormalText)
 		.SetFont(DEFAULT_FONT("Regular", 24))
 		.SetColorAndOpacity(FLinearColor::Black )
-		.SetShadowOffset( FVector2D(0.0f, 1.0f))
+		.SetShadowOffset( FVector2f(0.0f, 1.0f))
 		.SetShadowColorAndOpacity(FLinearColor(0.8f, 0.8f, 0.8f, 0.5))
 		);
 
@@ -310,7 +310,7 @@ TSharedRef<ISlateStyle> FUMGCoreStyle::Create()
 	
 	// Demo Recording
 	{
-		Style->Set("DemoRecording.CursorPing", new IMAGE_BRUSH("Common/CursorPing", FVector2D(31,31)));
+		Style->Set("DemoRecording.CursorPing", new IMAGE_BRUSH("Common/CursorPing", FVector2f(31.f,31.f)));
 	}
 
 	// Error Reporting
@@ -330,8 +330,8 @@ TSharedRef<ISlateStyle> FUMGCoreStyle::Create()
 
 	// Common Icons
 	{
-		Style->Set("TrashCan", new IMAGE_BRUSH( "Icons/TrashCan", FVector2D(64, 64)));
-		Style->Set("TrashCan_Small", new IMAGE_BRUSH( "Icons/TrashCan_Small", FVector2D(18, 18)));
+		Style->Set("TrashCan", new IMAGE_BRUSH( "Icons/TrashCan", FVector2f(64.f, 64.f)));
+		Style->Set("TrashCan_Small", new IMAGE_BRUSH( "Icons/TrashCan_Small", FVector2f(18.f, 18.f)));
 	}
 
 	// Common Icons
@@ -370,14 +370,14 @@ TSharedRef<ISlateStyle> FUMGCoreStyle::Create()
 	}
 
 	// SVerticalBox Drag& Drop Icon
-	Style->Set("VerticalBoxDragIndicator", new IMAGE_BRUSH("Common/VerticalBoxDragIndicator", FVector2D(6, 45)));
-	Style->Set("VerticalBoxDragIndicatorShort", new IMAGE_BRUSH("Common/VerticalBoxDragIndicatorShort", FVector2D(6, 15)));
+	Style->Set("VerticalBoxDragIndicator", new IMAGE_BRUSH("Common/VerticalBoxDragIndicator", FVector2f(6.f, 45.f)));
+	Style->Set("VerticalBoxDragIndicatorShort", new IMAGE_BRUSH("Common/VerticalBoxDragIndicatorShort", FVector2f(6.f, 15.f)));
 	// SScrollBar defaults...
 	static const FScrollBarStyle ScrollBar = FScrollBarStyle()
-		.SetVerticalTopSlotImage(IMAGE_BRUSH("Common/Scrollbar_Background_Vertical", FVector2D(8, 8)))
-		.SetVerticalBottomSlotImage(IMAGE_BRUSH("Common/Scrollbar_Background_Vertical", FVector2D(8, 8)))
-		.SetHorizontalTopSlotImage(IMAGE_BRUSH("Common/Scrollbar_Background_Horizontal", FVector2D(8, 8)))
-		.SetHorizontalBottomSlotImage(IMAGE_BRUSH("Common/Scrollbar_Background_Horizontal", FVector2D(8, 8)))
+		.SetVerticalTopSlotImage(IMAGE_BRUSH("Common/Scrollbar_Background_Vertical", FVector2f(8.f, 8.f)))
+		.SetVerticalBottomSlotImage(IMAGE_BRUSH("Common/Scrollbar_Background_Vertical", FVector2f(8.f, 8.f)))
+		.SetHorizontalTopSlotImage(IMAGE_BRUSH("Common/Scrollbar_Background_Horizontal", FVector2f(8.f, 8.f)))
+		.SetHorizontalBottomSlotImage(IMAGE_BRUSH("Common/Scrollbar_Background_Horizontal", FVector2f(8.f, 8.f)))
 		.SetNormalThumbImage( BOX_BRUSH( "Common/Scrollbar_Thumb", FMargin(4.f/16.f) ) )
 		.SetDraggedThumbImage( BOX_BRUSH( "Common/Scrollbar_Thumb", FMargin(4.f/16.f) ) )
 		.SetHoveredThumbImage( BOX_BRUSH( "Common/Scrollbar_Thumb", FMargin(4.f/16.f) ) );
@@ -580,7 +580,7 @@ TSharedRef<ISlateStyle> FUMGCoreStyle::Create()
 	{
 		FTextBlockStyle InlineEditableTextBlockReadOnly = FTextBlockStyle(NormalText)
 			.SetColorAndOpacity( FSlateColor::UseForeground() )
-			.SetShadowOffset( FVector2D::ZeroVector )
+			.SetShadowOffset( FVector2f::ZeroVector )
 			.SetShadowColorAndOpacity( FLinearColor::Black );
 
 		FTextBlockStyle InlineEditableTextBlockSmallReadOnly = FTextBlockStyle(InlineEditableTextBlockReadOnly)
@@ -653,14 +653,14 @@ TSharedRef<ISlateStyle> FUMGCoreStyle::Create()
 		Style->Set( "ProgressBar", FProgressBarStyle()
 			.SetBackgroundImage( BOX_BRUSH( "Common/ProgressBar_Background", FMargin(5.f/12.f) ) )
 			.SetFillImage( BOX_BRUSH( "Common/ProgressBar_Fill", FMargin(5.f/12.f), FLinearColor( 1.0f, 0.22f, 0.0f )  ) )
-			.SetMarqueeImage( IMAGE_BRUSH( "Common/ProgressBar_Marquee", FVector2D(20,12), FLinearColor::White, ESlateBrushTileType::Horizontal ) )
+			.SetMarqueeImage( IMAGE_BRUSH( "Common/ProgressBar_Marquee", FVector2f(20.f,12.f), FLinearColor::White, ESlateBrushTileType::Horizontal ) )
 			);
 	}
 
 	// SThrobber, SCircularThrobber defaults...
 	{
-		Style->Set( "Throbber.Chunk", new IMAGE_BRUSH( "Common/Throbber_Piece", FVector2D(16,16) ) );
-		Style->Set( "Throbber.CircleChunk", new IMAGE_BRUSH( "Common/Throbber_Piece", FVector2D(8,8) ) );
+		Style->Set( "Throbber.Chunk", new IMAGE_BRUSH( "Common/Throbber_Piece", FVector2f(16.f,16.f) ) );
+		Style->Set( "Throbber.CircleChunk", new IMAGE_BRUSH( "Common/Throbber_Piece", FVector2f(8.f,8.f) ) );
 	}
 
 	// SExpandableArea defaults...
@@ -698,9 +698,9 @@ TSharedRef<ISlateStyle> FUMGCoreStyle::Create()
 			.SetNormalBarImage(FSlateColorBrush(FColor::White))
 			.SetHoveredBarImage(FSlateColorBrush(FColor::White))
 			.SetDisabledBarImage(FSlateColorBrush(FLinearColor::Gray))
-			.SetNormalThumbImage( IMAGE_BRUSH( "Common/Button", FVector2D(8.0f, 14.0f) ) )
-			.SetHoveredThumbImage(IMAGE_BRUSH("Common/Button", FVector2D(8.0f, 14.0f)))
-			.SetDisabledThumbImage( IMAGE_BRUSH( "Common/Button_Disabled", FVector2D(8.0f, 14.0f) ) )
+			.SetNormalThumbImage( IMAGE_BRUSH( "Common/Button", FVector2f(8.0f, 14.0f) ) )
+			.SetHoveredThumbImage(IMAGE_BRUSH("Common/Button", FVector2f(8.0f, 14.0f)))
+			.SetDisabledThumbImage( IMAGE_BRUSH( "Common/Button_Disabled", FVector2f(8.0f, 14.0f) ) )
 			.SetBarThickness(2.0f);
 		Style->Set( "Slider", SliderStyle );
 
@@ -760,19 +760,19 @@ TSharedRef<ISlateStyle> FUMGCoreStyle::Create()
 		Style->Set( "ColorPicker.EyeDropper", new IMAGE_BRUSH( "Icons/eyedropper_16px", Icon16x16) );
 		Style->Set( "ColorPicker.Font", DEFAULT_FONT( "Regular", 10 ) );
 		Style->Set( "ColorPicker.Mode", new IMAGE_BRUSH( "Common/ColorPicker_Mode_16x", Icon16x16) );
-		Style->Set( "ColorPicker.Separator", new IMAGE_BRUSH( "Common/ColorPicker_Separator", FVector2D(2.0f, 2.0f) ) );
-		Style->Set( "ColorPicker.Selector", new IMAGE_BRUSH( "Common/Circle", FVector2D(8, 8) ) );
+		Style->Set( "ColorPicker.Separator", new IMAGE_BRUSH( "Common/ColorPicker_Separator", FVector2f(2.0f, 2.0f) ) );
+		Style->Set( "ColorPicker.Selector", new IMAGE_BRUSH( "Common/Circle", FVector2f(8.f, 8.f) ) );
 		Style->Set( "ColorPicker.Slider", FSliderStyle()
-			.SetDisabledThumbImage( IMAGE_BRUSH( "Common/ColorPicker_SliderHandle", FVector2D(8.0f, 32.0f)) )
-			.SetNormalThumbImage( IMAGE_BRUSH( "Common/ColorPicker_SliderHandle", FVector2D(8.0f, 32.0f)) )
-			.SetHoveredThumbImage( IMAGE_BRUSH("Common/ColorPicker_SliderHandle", FVector2D(8.0f, 32.0f)) )
+			.SetDisabledThumbImage( IMAGE_BRUSH( "Common/ColorPicker_SliderHandle", FVector2f(8.0f, 32.0f)) )
+			.SetNormalThumbImage( IMAGE_BRUSH( "Common/ColorPicker_SliderHandle", FVector2f(8.0f, 32.0f)) )
+			.SetHoveredThumbImage( IMAGE_BRUSH("Common/ColorPicker_SliderHandle", FVector2f(8.0f, 32.0f)) )
 			);
 	}
 
 	// SColorSpectrum defaults...
 	{
-		Style->Set( "ColorSpectrum.Spectrum", new IMAGE_BRUSH( "Common/ColorSpectrum", FVector2D(256, 256) ) );
-		Style->Set( "ColorSpectrum.Selector", new IMAGE_BRUSH( "Common/Circle", FVector2D(8, 8) ) );
+		Style->Set( "ColorSpectrum.Spectrum", new IMAGE_BRUSH( "Common/ColorSpectrum", FVector2f(256.f, 256.f) ) );
+		Style->Set( "ColorSpectrum.Selector", new IMAGE_BRUSH( "Common/Circle", FVector2f(8.f, 8.f) ) );
 	}
 
 	// SColorThemes defaults...
@@ -782,14 +782,14 @@ TSharedRef<ISlateStyle> FUMGCoreStyle::Create()
 
 	// SColorWheel defaults...
 	{
-		Style->Set( "ColorWheel.HueValueCircle", new IMAGE_BRUSH( "Common/ColorWheel", FVector2D(192, 192) ) );
-		Style->Set( "ColorWheel.Selector", new IMAGE_BRUSH( "Common/Circle", FVector2D(8, 8) ) );
+		Style->Set( "ColorWheel.HueValueCircle", new IMAGE_BRUSH( "Common/ColorWheel", FVector2f(192.f, 192.f) ) );
+		Style->Set( "ColorWheel.Selector", new IMAGE_BRUSH( "Common/Circle", FVector2f(8.f, 8.f) ) );
 	}
 
 	// SColorGradingWheel defaults...
 	{
-		Style->Set("ColorGradingWheel.HueValueCircle", new IMAGE_BRUSH("Common/ColorGradingWheel", FVector2D(192, 192)));
-		Style->Set("ColorGradingWheel.Selector", new IMAGE_BRUSH("Common/Circle", FVector2D(8, 8)));
+		Style->Set("ColorGradingWheel.HueValueCircle", new IMAGE_BRUSH("Common/ColorGradingWheel", FVector2f(192.f, 192.f)));
+		Style->Set("ColorGradingWheel.Selector", new IMAGE_BRUSH("Common/Circle", FVector2f(8.f, 8.f)));
 	}
 
 	// SSplitter
@@ -1140,7 +1140,7 @@ TSharedRef<ISlateStyle> FUMGCoreStyle::Create()
 	{
 		Style->Set("Wizard.PageTitle", FTextBlockStyle(NormalText)
 			.SetFont(DEFAULT_FONT("BoldCondensed", 28))
-			.SetShadowOffset(FVector2D(1, 1))
+			.SetShadowOffset(FVector2f(1.f, 1.f))
 			.SetShadowColorAndOpacity(FLinearColor(0, 0, 0, 0.9f))
 			);
 	}
@@ -1176,13 +1176,13 @@ TSharedRef<ISlateStyle> FUMGCoreStyle::Create()
 		FTextBlockStyle DockingTabFont = FTextBlockStyle(NormalText)
 			.SetFont( DEFAULT_FONT( "Regular", 9 ) )
 			.SetColorAndOpacity( FLinearColor(0.72f, 0.72f, 0.72f, 1.f) )
-			.SetShadowOffset( FVector2D( 1,1 ) )
+			.SetShadowOffset( FVector2f( 1.f,1.f ) )
 			.SetShadowColorAndOpacity( FLinearColor::Black ) ;
 
 		Style->Set( "Docking.UnhideTabwellButton", FButtonStyle(Button)
-			.SetNormal ( IMAGE_BRUSH( "/Docking/ShowTabwellButton_Normal", FVector2D(10,10) ) )
-			.SetPressed( IMAGE_BRUSH( "/Docking/ShowTabwellButton_Pressed", FVector2D(10,10) ) )
-			.SetHovered( IMAGE_BRUSH( "/Docking/ShowTabwellButton_Hovered", FVector2D(10,10) ) )
+			.SetNormal ( IMAGE_BRUSH( "/Docking/ShowTabwellButton_Normal", FVector2f(10.f,10.f) ) )
+			.SetPressed( IMAGE_BRUSH( "/Docking/ShowTabwellButton_Pressed", FVector2f(10.f,10.f) ) )
+			.SetHovered( IMAGE_BRUSH( "/Docking/ShowTabwellButton_Hovered", FVector2f(10.f,10.f) ) )
 			.SetNormalPadding(0)
 			.SetPressedPadding(0)
 			);
@@ -1231,14 +1231,14 @@ TSharedRef<ISlateStyle> FUMGCoreStyle::Create()
 			);
 
 		// Dock Cross
-		Style->Set( "Docking.Cross.DockLeft", new IMAGE_BRUSH( "/Docking/OuterDockingIndicator", FVector2D(6, 6), FLinearColor(1.0f, 0.35f, 0.0f, 0.25f) ) );
-		Style->Set( "Docking.Cross.DockLeft_Hovered", new IMAGE_BRUSH( "/Docking/OuterDockingIndicator", FVector2D(6, 6), FLinearColor(1.0f, 0.35f, 0.0f) ) );
-		Style->Set( "Docking.Cross.DockTop", new IMAGE_BRUSH( "/Docking/OuterDockingIndicator", FVector2D(6, 6), FLinearColor(1.0f, 0.35f, 0.0f, 0.25f) ) );
-		Style->Set( "Docking.Cross.DockTop_Hovered", new IMAGE_BRUSH( "/Docking/OuterDockingIndicator", FVector2D(6, 6), FLinearColor(1.0f, 0.35f, 0.0f) ) );
-		Style->Set( "Docking.Cross.DockRight", new IMAGE_BRUSH( "/Docking/OuterDockingIndicator", FVector2D(6, 6), FLinearColor(1.0f, 0.35f, 0.0f, 0.25f) ) );
-		Style->Set( "Docking.Cross.DockRight_Hovered", new IMAGE_BRUSH( "/Docking/OuterDockingIndicator", FVector2D(6, 6), FLinearColor(1.0f, 0.35f, 0.0f) ) );
-		Style->Set( "Docking.Cross.DockBottom", new IMAGE_BRUSH( "/Docking/OuterDockingIndicator", FVector2D(6, 6), FLinearColor(1.0f, 0.35f, 0.0f, 0.25f) ) );
-		Style->Set( "Docking.Cross.DockBottom_Hovered", new IMAGE_BRUSH( "/Docking/OuterDockingIndicator", FVector2D(6, 6), FLinearColor(1.0f, 0.35f, 0.0f) ) );
+		Style->Set( "Docking.Cross.DockLeft", new IMAGE_BRUSH( "/Docking/OuterDockingIndicator", FVector2f(6.f, 6.f), FLinearColor(1.0f, 0.35f, 0.0f, 0.25f) ) );
+		Style->Set( "Docking.Cross.DockLeft_Hovered", new IMAGE_BRUSH( "/Docking/OuterDockingIndicator", FVector2f(6.f, 6.f), FLinearColor(1.0f, 0.35f, 0.0f) ) );
+		Style->Set( "Docking.Cross.DockTop", new IMAGE_BRUSH( "/Docking/OuterDockingIndicator", FVector2f(6.f, 6.f), FLinearColor(1.0f, 0.35f, 0.0f, 0.25f) ) );
+		Style->Set( "Docking.Cross.DockTop_Hovered", new IMAGE_BRUSH( "/Docking/OuterDockingIndicator", FVector2f(6.f, 6.f), FLinearColor(1.0f, 0.35f, 0.0f) ) );
+		Style->Set( "Docking.Cross.DockRight", new IMAGE_BRUSH( "/Docking/OuterDockingIndicator", FVector2f(6.f, 6.f), FLinearColor(1.0f, 0.35f, 0.0f, 0.25f) ) );
+		Style->Set( "Docking.Cross.DockRight_Hovered", new IMAGE_BRUSH( "/Docking/OuterDockingIndicator", FVector2f(6.f, 6.f), FLinearColor(1.0f, 0.35f, 0.0f) ) );
+		Style->Set( "Docking.Cross.DockBottom", new IMAGE_BRUSH( "/Docking/OuterDockingIndicator", FVector2f(6.f, 6.f), FLinearColor(1.0f, 0.35f, 0.0f, 0.25f) ) );
+		Style->Set( "Docking.Cross.DockBottom_Hovered", new IMAGE_BRUSH( "/Docking/OuterDockingIndicator", FVector2f(6.f, 6.f), FLinearColor(1.0f, 0.35f, 0.0f) ) );
 		Style->Set( "Docking.Cross.DockCenter", new IMAGE_BRUSH( "/Docking/DockingIndicator_Center", Icon64x64,  FLinearColor(1.0f, 0.35f, 0.0f, 0.25f) ) );
 		Style->Set( "Docking.Cross.DockCenter_Hovered", new IMAGE_BRUSH( "/Docking/DockingIndicator_Center", Icon64x64,  FLinearColor(1.0f, 0.35f, 0.0f) ) );
 		
@@ -1256,18 +1256,18 @@ TSharedRef<ISlateStyle> FUMGCoreStyle::Create()
 	// SScrollBox defaults...
 	{
 		Style->Set( "ScrollBox", FScrollBoxStyle()
-			.SetTopShadowBrush( BOX_BRUSH( "Common/ScrollBoxShadowTop", FVector2D(16, 8), FMargin(0.5, 1, 0.5, 0) ) )
-			.SetBottomShadowBrush(BOX_BRUSH("Common/ScrollBoxShadowBottom", FVector2D(16, 8), FMargin(0.5, 0, 0.5, 1)))
-			.SetLeftShadowBrush(BOX_BRUSH("Common/ScrollBoxShadowLeft", FVector2D(8, 16), FMargin(1, 0.5, 0, 0.5)))
-			.SetRightShadowBrush(BOX_BRUSH("Common/ScrollBoxShadowRight", FVector2D(8, 16), FMargin(0, 0.5, 1, 0.5)))
+			.SetTopShadowBrush( BOX_BRUSH( "Common/ScrollBoxShadowTop", FVector2f(16.f, 8.f), FMargin(0.5, 1, 0.5, 0) ) )
+			.SetBottomShadowBrush(BOX_BRUSH("Common/ScrollBoxShadowBottom", FVector2f(16.f, 8.f), FMargin(0.5, 0, 0.5, 1)))
+			.SetLeftShadowBrush(BOX_BRUSH("Common/ScrollBoxShadowLeft", FVector2f(8.f, 16.f), FMargin(1, 0.5, 0, 0.5)))
+			.SetRightShadowBrush(BOX_BRUSH("Common/ScrollBoxShadowRight", FVector2f(8.f, 16.f), FMargin(0, 0.5, 1, 0.5)))
 			);
 	}
 
 	// SScrollBorder defaults...
 	{
 		Style->Set( "ScrollBorder", FScrollBorderStyle()
-			.SetTopShadowBrush(BOX_BRUSH("Common/ScrollBorderShadowTop", FVector2D(16, 8), FMargin(0.5, 1, 0.5, 0)))
-			.SetBottomShadowBrush(BOX_BRUSH("Common/ScrollBorderShadowBottom", FVector2D(16, 8), FMargin(0.5, 0, 0.5, 1)))
+			.SetTopShadowBrush(BOX_BRUSH("Common/ScrollBorderShadowTop", FVector2f(16.f, 8.f), FMargin(0.5, 1, 0.5, 0)))
+			.SetBottomShadowBrush(BOX_BRUSH("Common/ScrollBorderShadowBottom", FVector2f(16.f, 8.f), FMargin(0.5, 0, 0.5, 1)))
 			);
 	}
 
@@ -1275,30 +1275,30 @@ TSharedRef<ISlateStyle> FUMGCoreStyle::Create()
 	{
 #if !PLATFORM_MAC
 		static const FButtonStyle MinimizeButtonStyle = FButtonStyle(Button)
-			.SetNormal (IMAGE_BRUSH("Common/Window/WindowButton_Minimize_Normal", FVector2D(27.0f, 18.0f)))
-			.SetHovered(IMAGE_BRUSH("Common/Window/WindowButton_Minimize_Hovered", FVector2D(27.0f, 18.0f)))
-			.SetPressed(IMAGE_BRUSH("Common/Window/WindowButton_Minimize_Pressed", FVector2D(27.0f, 18.0f)));
+			.SetNormal (IMAGE_BRUSH("Common/Window/WindowButton_Minimize_Normal", FVector2f(27.0f, 18.0f)))
+			.SetHovered(IMAGE_BRUSH("Common/Window/WindowButton_Minimize_Hovered", FVector2f(27.0f, 18.0f)))
+			.SetPressed(IMAGE_BRUSH("Common/Window/WindowButton_Minimize_Pressed", FVector2f(27.0f, 18.0f)));
 
 		static const FButtonStyle MaximizeButtonStyle = FButtonStyle(Button)
-			.SetNormal (IMAGE_BRUSH("Common/Window/WindowButton_Maximize_Normal", FVector2D(23.0f, 18.0f)))
-			.SetHovered(IMAGE_BRUSH("Common/Window/WindowButton_Maximize_Hovered", FVector2D(23.0f, 18.0f)))
-			.SetPressed(IMAGE_BRUSH("Common/Window/WindowButton_Maximize_Pressed", FVector2D(23.0f, 18.0f)));
+			.SetNormal (IMAGE_BRUSH("Common/Window/WindowButton_Maximize_Normal", FVector2f(23.0f, 18.0f)))
+			.SetHovered(IMAGE_BRUSH("Common/Window/WindowButton_Maximize_Hovered", FVector2f(23.0f, 18.0f)))
+			.SetPressed(IMAGE_BRUSH("Common/Window/WindowButton_Maximize_Pressed", FVector2f(23.0f, 18.0f)));
 
 		static const FButtonStyle RestoreButtonStyle = FButtonStyle(Button)
-			.SetNormal (IMAGE_BRUSH("Common/Window/WindowButton_Restore_Normal", FVector2D(23.0f, 18)))
-			.SetHovered(IMAGE_BRUSH("Common/Window/WindowButton_Restore_Hovered", FVector2D(23.0f, 18)))
-			.SetPressed(IMAGE_BRUSH("Common/Window/WindowButton_Restore_Pressed", FVector2D(23.0f, 18)));
+			.SetNormal (IMAGE_BRUSH("Common/Window/WindowButton_Restore_Normal", FVector2f(23.0f, 18.f)))
+			.SetHovered(IMAGE_BRUSH("Common/Window/WindowButton_Restore_Hovered", FVector2f(23.0f, 18.f)))
+			.SetPressed(IMAGE_BRUSH("Common/Window/WindowButton_Restore_Pressed", FVector2f(23.0f, 18.f)));
 
 		static const FButtonStyle CloseButtonStyle = FButtonStyle(Button)
-			.SetNormal (IMAGE_BRUSH("Common/Window/WindowButton_Close_Normal", FVector2D(44.0f, 18.0f)))
-			.SetHovered(IMAGE_BRUSH("Common/Window/WindowButton_Close_Hovered", FVector2D(44.0f, 18.0f)))
-			.SetPressed(IMAGE_BRUSH("Common/Window/WindowButton_Close_Pressed", FVector2D(44.0f, 18.0f)));
+			.SetNormal (IMAGE_BRUSH("Common/Window/WindowButton_Close_Normal", FVector2f(44.0f, 18.0f)))
+			.SetHovered(IMAGE_BRUSH("Common/Window/WindowButton_Close_Hovered", FVector2f(44.0f, 18.0f)))
+			.SetPressed(IMAGE_BRUSH("Common/Window/WindowButton_Close_Pressed", FVector2f(44.0f, 18.0f)));
 #endif
 
 		static const FTextBlockStyle TitleTextStyle = FTextBlockStyle(NormalText)
 			.SetFont(DEFAULT_FONT("Regular", 9))
 			.SetColorAndOpacity(FLinearColor::White)
-			.SetShadowOffset(FVector2D(1.0f, 1.0f))
+			.SetShadowOffset(FVector2f(1.0f, 1.0f))
 			.SetShadowColorAndOpacity(FLinearColor::Black);
 
 		Style->Set( "Window", FWindowStyle()
@@ -1315,15 +1315,15 @@ TSharedRef<ISlateStyle> FUMGCoreStyle::Create()
 			.SetOutlineBrush(BORDER_BRUSH("Common/Window/WindowOutline", FMargin(3.0f / 32.0f)))
 			.SetOutlineColor(FLinearColor(0.1f, 0.1f, 0.1f, 1.0f))
 			.SetBorderBrush(BOX_BRUSH("Common/Window/WindowBorder", 0.48f))
-			.SetBackgroundBrush(IMAGE_BRUSH( "Common/Window/WindowBackground", FVector2D(74.0f, 74.0f), FLinearColor::White, ESlateBrushTileType::Both))
-			.SetChildBackgroundBrush(IMAGE_BRUSH( "Common/NoiseBackground", FVector2D(64.0f, 64.0f), FLinearColor::White, ESlateBrushTileType::Both))
+			.SetBackgroundBrush(IMAGE_BRUSH( "Common/Window/WindowBackground", FVector2f(74.0f, 74.0f), FLinearColor::White, ESlateBrushTileType::Both))
+			.SetChildBackgroundBrush(IMAGE_BRUSH( "Common/NoiseBackground", FVector2f(64.0f, 64.0f), FLinearColor::White, ESlateBrushTileType::Both))
 			);
 	}
 
 	// STutorialWrapper defaults...
 	{
-		Style->Set("Tutorials.Border", new BOX_BRUSH("Tutorials/TutorialBorder", FVector2D(64.0f, 64.0f), FMargin(25.0f/ 64.0f)));
-		Style->Set("Tutorials.Shadow", new BOX_BRUSH("Tutorials/TutorialShadow", FVector2D(256.0f, 256.0f), FMargin(114.0f / 256.0f)));
+		Style->Set("Tutorials.Border", new BOX_BRUSH("Tutorials/TutorialBorder", FVector2f(64.0f, 64.0f), FMargin(25.0f/ 64.0f)));
+		Style->Set("Tutorials.Shadow", new BOX_BRUSH("Tutorials/TutorialShadow", FVector2f(256.0f, 256.0f), FMargin(114.0f / 256.0f)));
 	}
 
 	// Standard Dialog Settings
@@ -1357,7 +1357,7 @@ TSharedRef<ISlateStyle> FUMGCoreStyle::Create()
 	{
 		Style->Set("MessageLog", FTextBlockStyle(NormalText)
 			.SetFont(DEFAULT_FONT("Regular", 8))
-			.SetShadowOffset(FVector2D::ZeroVector)
+			.SetShadowOffset(FVector2f::ZeroVector)
 		);
 		Style->Set("MessageLog.Error", new IMAGE_BRUSH("MessageLog/Log_Error", Icon16x16));
 		Style->Set("MessageLog.Warning", new IMAGE_BRUSH("MessageLog/Log_Warning", Icon16x16));
diff --git a/Engine/Source/Runtime/SlateCore/Private/Types/PaintArgs.cpp b/Engine/Source/Runtime/SlateCore/Private/Types/PaintArgs.cpp
--- a/Engine/Source/Runtime/SlateCore/Private/Types/PaintArgs.cpp
+++ b/Engine/Source/Runtime/SlateCore/Private/Types/PaintArgs.cpp
@@ -5,7 +5,7 @@
 #include "Input/HittestGrid.h"
 #include "Widgets/SWindow.h"
 
-FPaintArgs::FPaintArgs(const SWidget* PaintParent, FHittestGrid& InRootHittestGrid, FHittestGrid& InCurrentHitTestGrid, FVector2D InWindowOffset, double InCurrentTime, float InDeltaTime)
+FPaintArgs::FPaintArgs(const SWidget* PaintParent, FHittestGrid& InRootHittestGrid, FHittestGrid& InCurrentHitTestGrid, UE::Slate::FDeprecateVector2DParameter InWindowOffset, double InCurrentTime, float InDeltaTime)
 	: RootGrid(InRootHittestGrid)
 	, CurrentGrid(InCurrentHitTestGrid)
 	, WindowOffset(InWindowOffset)
@@ -17,7 +17,7 @@ FPaintArgs::FPaintArgs(const SWidget* PaintParent, FHittestGrid& InRootHittestGr
 {
 }
 
-FPaintArgs::FPaintArgs(const SWidget* PaintParent, FHittestGrid& InRootHittestGrid, FVector2D InWindowOffset, double InCurrentTime, float InDeltaTime)
+FPaintArgs::FPaintArgs(const SWidget* PaintParent, FHittestGrid& InRootHittestGrid, UE::Slate::FDeprecateVector2DParameter InWindowOffset, double InCurrentTime, float InDeltaTime)
 	: FPaintArgs(PaintParent, InRootHittestGrid, InRootHittestGrid, InWindowOffset, InCurrentTime, InDeltaTime)
 {
 
diff --git a/Engine/Source/Runtime/SlateCore/Private/Widgets/Accessibility/SlateCoreAccessibleWidgets.cpp b/Engine/Source/Runtime/SlateCore/Private/Widgets/Accessibility/SlateCoreAccessibleWidgets.cpp
--- a/Engine/Source/Runtime/SlateCore/Private/Widgets/Accessibility/SlateCoreAccessibleWidgets.cpp
+++ b/Engine/Source/Runtime/SlateCore/Private/Widgets/Accessibility/SlateCoreAccessibleWidgets.cpp
@@ -74,7 +74,10 @@ FBox2D FSlateAccessibleWidget::GetBounds() const
 	if (Widget.IsValid())
 	{
 		const FGeometry& Geometry = Widget.Pin()->GetCachedGeometry();
-		return FBox2D(Geometry.GetAbsolutePosition(), Geometry.GetAbsolutePosition() + Geometry.GetAbsoluteSize());
+		const FVector2f Min = Geometry.GetAbsolutePosition();
+		const FVector2f Max = Min + Geometry.GetAbsoluteSize();
+
+		return FBox2D(FVector2d(Min), FVector2d(Max));
 	}
 	return FBox2D();
 }
@@ -370,7 +373,7 @@ TSharedPtr<IAccessibleWidget> FSlateAccessibleWindow::GetChildAtPosition(int32 X
 		if (UseHitTestGrid)
 		{
 			TSharedPtr<SWindow> SlateWindow = StaticCastSharedPtr<SWindow>(Widget.Pin());
-			TArray<FWidgetAndPointer> Hits = SlateWindow->GetHittestGrid().GetBubblePath(FVector2D((float)X, (float)Y), 0.0f, false, INDEX_NONE);
+			TArray<FWidgetAndPointer> Hits = SlateWindow->GetHittestGrid().GetBubblePath(FVector2f((float)X, (float)Y), 0.0f, false, INDEX_NONE);
 			TSharedPtr<SWidget> LastAccessibleWidget = nullptr;
 			for (int32 i = 0; i < Hits.Num(); ++i)
 			{
diff --git a/Engine/Source/Runtime/SlateCore/Private/Widgets/Images/SImage.cpp b/Engine/Source/Runtime/SlateCore/Private/Widgets/Images/SImage.cpp
--- a/Engine/Source/Runtime/SlateCore/Private/Widgets/Images/SImage.cpp
+++ b/Engine/Source/Runtime/SlateCore/Private/Widgets/Images/SImage.cpp
@@ -71,7 +71,7 @@ FVector2D SImage::ComputeDesiredSize( float ) const
 	{
 		const TOptional<FVector2D>& CurrentSizeOverride = DesiredSizeOverrideAttribute.Get();
 
-		return CurrentSizeOverride.IsSet() ? CurrentSizeOverride.GetValue() : ImageBrush->ImageSize;
+		return CurrentSizeOverride.IsSet() ? CurrentSizeOverride.GetValue() : FVector2D(ImageBrush->ImageSize.X, ImageBrush->ImageSize.Y);
 	}
 	return FVector2D::ZeroVector;
 }
diff --git a/Engine/Source/Runtime/SlateCore/Private/Widgets/SBoxPanel.cpp b/Engine/Source/Runtime/SlateCore/Private/Widgets/SBoxPanel.cpp
--- a/Engine/Source/Runtime/SlateCore/Private/Widgets/SBoxPanel.cpp
+++ b/Engine/Source/Runtime/SlateCore/Private/Widgets/SBoxPanel.cpp
@@ -149,7 +149,7 @@ static FVector2D ComputeDesiredSizeForBox( const TPanelChildren<SlotType>& Child
 		if ( CurChild.GetWidget()->GetVisibility() != EVisibility::Collapsed )
 		{
 
-			const FVector2D& CurChildDesiredSize = CurChild.GetWidget()->GetDesiredSize();
+			FVector2f CurChildDesiredSize = CurChild.GetWidget()->GetDesiredSize();
 
 			if (Orientation == Orient_Vertical)
 			{
@@ -246,7 +246,7 @@ void SDragAndDropVerticalBox::Construct(const FArguments& InArgs)
 	OnDragLeave_Handler = InArgs._OnDragLeave;
 	OnDrop_Handler = InArgs._OnDrop;
 
-	CurrentDragOperationScreenSpaceLocation = FVector2D::ZeroVector;
+	CurrentDragOperationScreenSpaceLocation = FVector2f::ZeroVector;
 	CurrentDragOverSlotIndex = INDEX_NONE;
 }
 
@@ -290,7 +290,7 @@ void SDragAndDropVerticalBox::OnDragEnter(FGeometry const& MyGeometry, FDragDrop
 void SDragAndDropVerticalBox::OnDragLeave(FDragDropEvent const& DragDropEvent)
 {
 	ItemDropZone = TOptional<EItemDropZone>();
-	CurrentDragOperationScreenSpaceLocation = FVector2D::ZeroVector;
+	CurrentDragOperationScreenSpaceLocation = FVector2f::ZeroVector;
 	CurrentDragOverSlotIndex = INDEX_NONE;
 
 	if (OnDragLeave_Handler.IsBound())
@@ -299,7 +299,7 @@ void SDragAndDropVerticalBox::OnDragLeave(FDragDropEvent const& DragDropEvent)
 	}
 }
 
-SDragAndDropVerticalBox::EItemDropZone SDragAndDropVerticalBox::ZoneFromPointerPosition(FVector2D LocalPointerPos, const FGeometry& CurrentGeometry, const FGeometry& StartGeometry) const
+SDragAndDropVerticalBox::EItemDropZone SDragAndDropVerticalBox::ZoneFromPointerPosition(FVector2f LocalPointerPos, const FGeometry& CurrentGeometry, const FGeometry& StartGeometry) const
 {
 	FSlateLayoutTransform StartGeometryLayoutTransform = StartGeometry.GetAccumulatedLayoutTransform();
 	FSlateLayoutTransform CurrentGeometryLayoutTransform = CurrentGeometry.GetAccumulatedLayoutTransform();
@@ -340,7 +340,7 @@ FReply SDragAndDropVerticalBox::OnDragOver(const FGeometry& MyGeometry, const FD
 
 			if (ArrangedChildren.IsValidIndex(DragOverSlotIndex))
 			{
-				FVector2D LocalPointerPos = ArrangedChildren[DragOverSlotIndex].Geometry.AbsoluteToLocal(DragDropEvent.GetScreenSpacePosition());
+				FVector2f LocalPointerPos = ArrangedChildren[DragOverSlotIndex].Geometry.AbsoluteToLocal(DragDropEvent.GetScreenSpacePosition());
 				EItemDropZone ItemHoverZone = ZoneFromPointerPosition(LocalPointerPos, ArrangedChildren[DragOverSlotIndex].Geometry, ArrangedChildren[DragOp->SlotIndexBeingDragged].Geometry);
 
 				if (Children.IsValidIndex(DragOverSlotIndex))
@@ -397,7 +397,7 @@ FReply SDragAndDropVerticalBox::OnDrop(const FGeometry& MyGeometry, const FDragD
 			}
 
 			ItemDropZone = TOptional<EItemDropZone>();
-			CurrentDragOperationScreenSpaceLocation = FVector2D::ZeroVector;
+			CurrentDragOperationScreenSpaceLocation = FVector2f::ZeroVector;
 			CurrentDragOverSlotIndex = INDEX_NONE;
 		}		
 	}
diff --git a/Engine/Source/Runtime/SlateCore/Private/Widgets/SCompoundWidget.cpp b/Engine/Source/Runtime/SlateCore/Private/Widgets/SCompoundWidget.cpp
--- a/Engine/Source/Runtime/SlateCore/Private/Widgets/SCompoundWidget.cpp
+++ b/Engine/Source/Runtime/SlateCore/Private/Widgets/SCompoundWidget.cpp
@@ -61,7 +61,8 @@ FVector2D SCompoundWidget::ComputeDesiredSize( float ) const
 	EVisibility ChildVisibility = ChildSlot.GetWidget()->GetVisibility();
 	if ( ChildVisibility != EVisibility::Collapsed )
 	{
-		return ChildSlot.GetWidget()->GetDesiredSize() + ChildSlot.GetPadding().GetDesiredSize();
+		FVector2f LocalDesiredSize = ChildSlot.GetWidget()->GetDesiredSize() + ChildSlot.GetPadding().GetDesiredSize();
+		return FVector2D(LocalDesiredSize);
 	}
 	
 	return FVector2D::ZeroVector;
diff --git a/Engine/Source/Runtime/SlateCore/Private/Widgets/SOverlay.cpp b/Engine/Source/Runtime/SlateCore/Private/Widgets/SOverlay.cpp
--- a/Engine/Source/Runtime/SlateCore/Private/Widgets/SOverlay.cpp
+++ b/Engine/Source/Runtime/SlateCore/Private/Widgets/SOverlay.cpp
@@ -142,8 +142,8 @@ void SOverlay::OnArrangeChildren( const FGeometry& AllottedGeometry, FArrangedCh
 
 			ArrangedChildren.AddWidget( ChildVisibility, AllottedGeometry.MakeChild(
 				CurChild.GetWidget(),
-				FVector2D(XResult.Offset, YResult.Offset),
-				FVector2D(XResult.Size, YResult.Size)
+				FVector2f(XResult.Size, YResult.Size),
+				FSlateLayoutTransform(FVector2f(XResult.Offset, YResult.Offset))
 			) );
 		}
 	}
@@ -158,7 +158,7 @@ FVector2D SOverlay::ComputeDesiredSize( float ) const
 		const EVisibility ChildVisibilty = CurSlot.GetWidget()->GetVisibility();
 		if ( ChildVisibilty != EVisibility::Collapsed )
 		{
-			FVector2D ChildDesiredSize = CurSlot.GetWidget()->GetDesiredSize() + CurSlot.GetPadding().GetDesiredSize();
+			FVector2f ChildDesiredSize = CurSlot.GetWidget()->GetDesiredSize() + CurSlot.GetPadding().GetDesiredSize();
 			MaxSize.X = FMath::Max( MaxSize.X, ChildDesiredSize.X );
 			MaxSize.Y = FMath::Max( MaxSize.Y, ChildDesiredSize.Y );
 		}
diff --git a/Engine/Source/Runtime/SlateCore/Private/Widgets/SWidget.cpp b/Engine/Source/Runtime/SlateCore/Private/Widgets/SWidget.cpp
--- a/Engine/Source/Runtime/SlateCore/Private/Widgets/SWidget.cpp
+++ b/Engine/Source/Runtime/SlateCore/Private/Widgets/SWidget.cpp
@@ -631,11 +631,6 @@ FPopupMethodReply SWidget::OnQueryPopupMethod() const
 	return FPopupMethodReply::Unhandled();
 }
 
-TSharedPtr<struct FVirtualPointerPosition> SWidget::TranslateMouseCoordinateForCustomHitTestChild(const TSharedRef<SWidget>& ChildWidget, const FGeometry& MyGeometry, const FVector2D& ScreenSpaceMouseCoordinate, const FVector2D& LastScreenSpaceMouseCoordinate) const
-{
-	return nullptr;
-}
-
 TOptional<FVirtualPointerPosition> SWidget::TranslateMouseCoordinateForCustomHitTestChild(const SWidget& ChildWidget, const FGeometry& MyGeometry, const FVector2D ScreenSpaceMouseCoordinate, const FVector2D LastScreenSpaceMouseCoordinate) const
 {
 	return TOptional<FVirtualPointerPosition>();
@@ -708,9 +703,9 @@ void SWidget::InvalidateChildRemovedFromTree(SWidget& Child)
 	}
 }
 
-FVector2D SWidget::GetDesiredSize() const
+UE::Slate::FDeprecateVector2DResult SWidget::GetDesiredSize() const
 {
-	return FVector2D(DesiredSize.Get(FVector2f::ZeroVector));
+	return UE::Slate::FDeprecateVector2DResult(DesiredSize.Get(FVector2f::ZeroVector));
 }
 
 void SWidget::AssignParentWidget(TSharedPtr<SWidget> InParent)
@@ -1049,11 +1044,11 @@ FGeometry SWidget::FindChildGeometry( const FGeometry& MyGeometry, TSharedRef<SW
 
 int32 SWidget::FindChildUnderMouse( const FArrangedChildren& Children, const FPointerEvent& MouseEvent )
 {
-	const FVector2D& AbsoluteCursorLocation = MouseEvent.GetScreenSpacePosition();
+	FVector2f AbsoluteCursorLocation = MouseEvent.GetScreenSpacePosition();
 	return SWidget::FindChildUnderPosition( Children, AbsoluteCursorLocation );
 }
 
-int32 SWidget::FindChildUnderPosition( const FArrangedChildren& Children, const FVector2D& ArrangedSpacePosition )
+int32 SWidget::FindChildUnderPosition( const FArrangedChildren& Children, const UE::Slate::FDeprecateVector2DParameter& ArrangedSpacePosition )
 {
 	const int32 NumChildren = Children.Num();
 	for( int32 ChildIndex=NumChildren-1; ChildIndex >= 0; --ChildIndex )
@@ -1356,8 +1351,8 @@ FSlateRect SWidget::CalculateCullingAndClippingRules(const FGeometry& AllottedGe
 			break;
 		case EWidgetClipping::OnDemand:
 			const float OverflowEpsilon = 1.0f;
-			const FVector2D& CurrentSize = GetDesiredSize();
-			const FVector2D& LocalSize = AllottedGeometry.GetLocalSize();
+			const FVector2f& CurrentSize = GetDesiredSize();
+			const FVector2f& LocalSize = AllottedGeometry.GetLocalSize();
 			bClipToBounds =
 				(CurrentSize.X - OverflowEpsilon) > LocalSize.X ||
 				(CurrentSize.Y - OverflowEpsilon) > LocalSize.Y;
diff --git a/Engine/Source/Runtime/SlateCore/Private/Widgets/SWindow.cpp b/Engine/Source/Runtime/SlateCore/Private/Widgets/SWindow.cpp
--- a/Engine/Source/Runtime/SlateCore/Private/Widgets/SWindow.cpp
+++ b/Engine/Source/Runtime/SlateCore/Private/Widgets/SWindow.cpp
@@ -133,7 +133,7 @@ private:
 				const float ChildLocalToScreenScale = CurChild.Scale_Attribute.Get();
 				FVector2D ChildLocalToScreenOffset = CurChild.DesktopPosition_Attribute.Get();
 				// The size of the child is either the desired size of the widget (computed in the child's local space) or the size override (specified in screen space)
-				const FVector2D ChildSizeChildLocal = CurChild.GetWidget()->GetDesiredSize();
+				const FVector2D ChildSizeChildLocal = FVector2D(UE::Slate::CastToVector2f(CurChild.GetWidget()->GetDesiredSize()));
 				// Convert the desired size to screen space. Here is were we convert a vector to screenspace
 				// before we have the final position in screenspace (which would be needed to transform a point).
 				FVector2D ChildSizeScreenspace = TransformVector(ChildLocalToScreenScale, ChildSizeChildLocal);
@@ -154,11 +154,11 @@ private:
 				}
 
 				// We now have the final position, so construct the transform from ChildLocal to Screenspace
-				const FSlateLayoutTransform ChildLocalToScreen(ChildLocalToScreenScale, ChildLocalToScreenOffset);
+				const FSlateLayoutTransform ChildLocalToScreen(ChildLocalToScreenScale, UE::Slate::CastToVector2f(ChildLocalToScreenOffset));
 				// Using this we can compute the transform from ChildLocal to WindowLocal, which is effectively the LayoutTransform of the child widget.
 				const FSlateLayoutTransform ChildLocalToWindowLocal = Concatenate(ChildLocalToScreen, ScreenToWindowLocal);
 				// The ChildSize needs to be given in ChildLocal space when constructing a geometry.
-				const FVector2D ChildSizeLocalspace = TransformVector(Inverse(ChildLocalToScreen), ChildSizeScreenspace);
+				const FVector2f ChildSizeLocalspace = TransformVector(Inverse(ChildLocalToScreen), UE::Slate::CastToVector2f(ChildSizeScreenspace));
 
 				// The position is explicitly in desktop pixels.
 				// The size and DPI scale come from the widget that is using
@@ -190,7 +190,7 @@ private:
 };
 
 
-FVector2D SWindow::GetWindowSizeFromClientSize(FVector2D InClientSize, TOptional<float> DPIScale)
+UE::Slate::FDeprecateVector2DResult SWindow::GetWindowSizeFromClientSize(UE::Slate::FDeprecateVector2DParameter InClientSize, TOptional<float> DPIScale)
 {
 	// If this is a regular non-OS window, we need to compensate for the border and title bar area that we will add
 	// Note: Windows with an OS border do this in ReshapeWindow
@@ -272,7 +272,7 @@ void SWindow::Construct(const FArguments& InArgs)
 	bCreateTitleBar = InArgs._CreateTitleBar && !bIsPopupWindow && Type != EWindowType::CursorDecorator && !bHasOSWindowBorder;
 
 	// calculate initial window position
-	FVector2D WindowPosition = InArgs._ScreenPosition;
+	FVector2f WindowPosition = InArgs._ScreenPosition;
 
 	const bool bAnchorWindowWindowPositionTopLeft = FPlatformApplicationMisc::AnchorWindowWindowPositionTopLeft();
 	if (bAnchorWindowWindowPositionTopLeft)
@@ -292,7 +292,7 @@ void SWindow::Construct(const FArguments& InArgs)
 	FDisplayMetrics DisplayMetrics;
 	FSlateApplicationBase::Get().GetCachedDisplayMetrics( DisplayMetrics );
 	const FPlatformRect& VirtualDisplayRect = DisplayMetrics.VirtualDisplayRect;
-	FPlatformRect PrimaryDisplayRect = AutoCenterRule == EAutoCenter::PrimaryWorkArea ? DisplayMetrics.PrimaryDisplayWorkAreaRect : DisplayMetrics.GetMonitorWorkAreaFromPoint(WindowPosition);
+	FPlatformRect PrimaryDisplayRect = AutoCenterRule == EAutoCenter::PrimaryWorkArea ? DisplayMetrics.PrimaryDisplayWorkAreaRect : DisplayMetrics.GetMonitorWorkAreaFromPoint(FVector2d(WindowPosition));
 
 	if (PrimaryDisplayRect == FPlatformRect(0, 0, 0, 0))
 	{
@@ -360,8 +360,8 @@ void SWindow::Construct(const FArguments& InArgs)
 	}
 
 	// If the window has no OS border, simulate it ourselves, enlarging window by the size that OS border would have.
-	const FVector2D DPIScaledClientSize = InArgs._AdjustInitialSizeAndPositionForDPIScale ? InArgs._ClientSize * DPIScale: InArgs._ClientSize;
-	FVector2D WindowSize = GetWindowSizeFromClientSize(DPIScaledClientSize, DPIScale);
+	const FVector2f DPIScaledClientSize = InArgs._AdjustInitialSizeAndPositionForDPIScale ? InArgs._ClientSize * DPIScale: FVector2f(InArgs._ClientSize);
+	FVector2f WindowSize = GetWindowSizeFromClientSize(DPIScaledClientSize, DPIScale);
 
 	// If we're manually positioning the window we need to check if it's outside
 	// of the virtual bounds of the current displays or too large.
@@ -392,8 +392,8 @@ void SWindow::Construct(const FArguments& InArgs)
 		}
 
 		// Setup a position and size for the main frame window that's centered in the desktop work area
-		const FVector2D DisplayTopLeft( AutoCenterRect.Left, AutoCenterRect.Top );
-		const FVector2D DisplaySize( AutoCenterRect.Right - AutoCenterRect.Left, AutoCenterRect.Bottom - AutoCenterRect.Top );
+		const FVector2f DisplayTopLeft( AutoCenterRect.Left, AutoCenterRect.Top );
+		const FVector2f DisplaySize( AutoCenterRect.Right - AutoCenterRect.Left, AutoCenterRect.Bottom - AutoCenterRect.Top );
 
 		WindowPosition = DisplayTopLeft + (DisplaySize - WindowSize) * 0.5f;
 
@@ -493,12 +493,12 @@ TSharedRef<SWindow> SWindow::MakeStyledCursorDecorator(const FWindowStyle& InSty
 	return NewWindow;
 }
 
-FVector2D SWindow::ComputeWindowSizeForContent( FVector2D ContentSize )
+UE::Slate::FDeprecateVector2DResult SWindow::ComputeWindowSizeForContent( UE::Slate::FDeprecateVector2DParameter ContentSize )
 {
 	// @todo mainframe: This code should be updated to handle the case where we're spawning a window that doesn't have
 	//                  a traditional title bar, such as a window that contains a primary SDockingArea.  Currently, the
 	//                  size reported here will be too large!
-	return ContentSize + FVector2D(0, SWindowDefs::DefaultTitleBarSize);
+	return UE::Slate::FDeprecateVector2DResult(ContentSize + FVector2f(0, SWindowDefs::DefaultTitleBarSize));
 }
 
 TSharedRef<SWidget> SWindow::MakeWindowTitleBar(const TSharedRef<SWindow>& Window, const TSharedPtr<SWidget>& CenterContent, EHorizontalAlignment TitleContentAlignment)
@@ -746,9 +746,9 @@ void SWindow::Tick( const FGeometry& AllottedGeometry, const double InCurrentTim
 			}
 			else // if animating position
 			{
-				const FVector2D StartPosition( Morpher.StartingMorphShape.Left, Morpher.StartingMorphShape.Top );
-				const FVector2D TargetPosition( Morpher.TargetMorphShape.Left, Morpher.TargetMorphShape.Top );
-				const FVector2D NewPosition( FMath::Lerp( StartPosition, TargetPosition, InterpAlpha ) );
+				const FVector2f StartPosition( Morpher.StartingMorphShape.Left, Morpher.StartingMorphShape.Top );
+				const FVector2f TargetPosition( Morpher.TargetMorphShape.Left, Morpher.TargetMorphShape.Top );
+				const FVector2f NewPosition( FMath::Lerp( StartPosition, TargetPosition, InterpAlpha ) );
 				if( NewPosition != this->GetPositionInScreen() )
 				{
 					this->MoveWindowTo( NewPosition );
@@ -771,7 +771,7 @@ void SWindow::Tick( const FGeometry& AllottedGeometry, const double InCurrentTim
 			}
 			else // if animating position
 			{
-				const FVector2D TargetPosition( Morpher.TargetMorphShape.Left, Morpher.TargetMorphShape.Top );
+				const FVector2f TargetPosition( Morpher.TargetMorphShape.Left, Morpher.TargetMorphShape.Top );
 				if( TargetPosition != this->GetPositionInScreen() )
 				{
 					this->MoveWindowTo( TargetPosition );
@@ -785,12 +785,12 @@ void SWindow::Tick( const FGeometry& AllottedGeometry, const double InCurrentTim
 	}
 }
 
-FVector2D SWindow::GetInitialDesiredSizeInScreen() const
+UE::Slate::FDeprecateVector2DResult SWindow::GetInitialDesiredSizeInScreen() const
 {
 	return InitialDesiredSize;
 }
 
-FVector2D SWindow::GetInitialDesiredPositionInScreen() const
+UE::Slate::FDeprecateVector2DResult SWindow::GetInitialDesiredPositionInScreen() const
 {
 	return InitialDesiredScreenPosition;
 }
@@ -810,7 +810,7 @@ FGeometry SWindow::GetWindowGeometryInWindow() const
 	// FGeometry expects Size in Local space, but our size is stored in screen space (same as window space + screen offset).
 	// So we need to transform Size into the window's local space for FGeometry.
 	FSlateLayoutTransform LocalToWindow = GetLocalToWindowTransform();
-	FVector2D ViewSize = GetViewportSize();
+	FVector2f ViewSize = GetViewportSize();
 	return FGeometry::MakeRoot(TransformVector(Inverse(LocalToWindow), ViewSize), LocalToWindow );
 }
 
@@ -825,12 +825,12 @@ FSlateLayoutTransform SWindow::GetLocalToWindowTransform() const
 }
 
 
-FVector2D SWindow::GetPositionInScreen() const
+UE::Slate::FDeprecateVector2DResult SWindow::GetPositionInScreen() const
 {
 	return ScreenPosition;
 }
 
-FVector2D SWindow::GetSizeInScreen() const
+UE::Slate::FDeprecateVector2DResult SWindow::GetSizeInScreen() const
 {
 	return Size;
 }
@@ -879,14 +879,14 @@ FSlateRect SWindow::GetClientRectInScreen() const
 		.InsetBy(FMargin(0.0f, TitleBarSize, 0.0f, 0.0f));
 }
 
-FVector2D SWindow::GetClientSizeInScreen() const
+UE::Slate::FDeprecateVector2DResult SWindow::GetClientSizeInScreen() const
 {
 	return GetClientRectInScreen().GetSize();
 }
 
 FSlateRect SWindow::GetClippingRectangleInWindow() const
 {
-	FVector2D ViewSize = GetViewportSize();
+	FVector2f ViewSize = GetViewportSize();
 	return FSlateRect( 0, 0, ViewSize.X, ViewSize.Y );
 }
 
@@ -925,7 +925,7 @@ FMargin SWindow::GetNonMaximizedWindowBorderSize() const
 }
 
 
-void SWindow::MoveWindowTo( FVector2D NewPosition )
+void SWindow::MoveWindowTo( UE::Slate::FDeprecateVector2DParameter NewPosition )
 {
 	if (NativeWindow.IsValid())
 	{
@@ -934,7 +934,7 @@ void SWindow::MoveWindowTo( FVector2D NewPosition )
 		// This expectation is generally invalid (see UE-1308) as there may be a delay before the OS reports it back.
 		// This hack sets the position speculatively, keeping Slate happy while also giving the OS chance to report it
 		// correctly after or even during the actual call.
-		FVector2D SpeculativeScreenPosition(FMath::TruncToFloat(NewPosition.X), FMath::TruncToFloat(NewPosition.Y));
+		FVector2f SpeculativeScreenPosition(FMath::TruncToFloat(NewPosition.X), FMath::TruncToFloat(NewPosition.Y));
 		SetCachedScreenPosition(SpeculativeScreenPosition);
 #endif // PLATFORM_LINUX
 
@@ -946,16 +946,16 @@ void SWindow::MoveWindowTo( FVector2D NewPosition )
 	}
 }
 
-void SWindow::ReshapeWindow( FVector2D NewPosition, FVector2D NewSize )
+void SWindow::ReshapeWindow( UE::Slate::FDeprecateVector2DParameter NewPosition, UE::Slate::FDeprecateVector2DParameter NewSize )
 {
-	const FVector2D CurrentPosition = GetPositionInScreen();
-	const FVector2D CurrentSize = GetSizeInScreen();
+	const FVector2f CurrentPosition = GetPositionInScreen();
+	const FVector2f CurrentSize = GetSizeInScreen();
 
 	// Ceil (Minus a tad for float precision) to ensure contents are not a sub-pixel larger than the window, which will create unnecessary scroll bars 
-	const FVector2D OldPositionTruncated = FVector2D(FMath::TruncToInt(CurrentPosition.X), FMath::TruncToInt(CurrentPosition.Y));
-	const FVector2D OldSizeRounded = FVector2D(FMath::CeilToInt(CurrentSize.X - KINDA_SMALL_NUMBER), FMath::CeilToInt(CurrentSize.Y - KINDA_SMALL_NUMBER));
-	const FVector2D NewPositionTruncated = FVector2D(FMath::TruncToInt(NewPosition.X), FMath::TruncToInt(NewPosition.Y));
-	const FVector2D NewSizeRounded = FVector2D(FMath::CeilToInt(NewSize.X - KINDA_SMALL_NUMBER), FMath::CeilToInt(NewSize.Y - KINDA_SMALL_NUMBER));
+	const FVector2f OldPositionTruncated = FVector2f(FMath::TruncToInt(CurrentPosition.X), FMath::TruncToInt(CurrentPosition.Y));
+	const FVector2f OldSizeRounded = FVector2f(FMath::CeilToInt(CurrentSize.X - KINDA_SMALL_NUMBER), FMath::CeilToInt(CurrentSize.Y - KINDA_SMALL_NUMBER));
+	const FVector2f NewPositionTruncated = FVector2f(FMath::TruncToInt(NewPosition.X), FMath::TruncToInt(NewPosition.Y));
+	const FVector2f NewSizeRounded = FVector2f(FMath::CeilToInt(NewSize.X - KINDA_SMALL_NUMBER), FMath::CeilToInt(NewSize.Y - KINDA_SMALL_NUMBER));
 
 	if (OldPositionTruncated != NewPositionTruncated || OldSizeRounded != NewSizeRounded )
 	{
@@ -981,15 +981,15 @@ void SWindow::ReshapeWindow( FVector2D NewPosition, FVector2D NewSize )
 
 void SWindow::ReshapeWindow( const FSlateRect& InNewShape )
 {
-	ReshapeWindow(FVector2D(InNewShape.Left, InNewShape.Top), FVector2D(InNewShape.Right - InNewShape.Left, InNewShape.Bottom - InNewShape.Top));
+	ReshapeWindow(FVector2f(InNewShape.Left, InNewShape.Top), FVector2f(InNewShape.Right - InNewShape.Left, InNewShape.Bottom - InNewShape.Top));
 }
 
-void SWindow::Resize( FVector2D NewClientSize )
+void SWindow::Resize( UE::Slate::FDeprecateVector2DParameter NewClientSize )
 {
 	ResizeWindowSize(GetWindowSizeFromClientSize(NewClientSize));
 }
 
-void SWindow::ResizeWindowSize( FVector2D NewWindowSize )
+void SWindow::ResizeWindowSize( FVector2f NewWindowSize )
 {
 	Morpher.Sequence.JumpToEnd();
 
@@ -1011,10 +1011,10 @@ void SWindow::ResizeWindowSize( FVector2D NewWindowSize )
 		}
 		else
 		{
-			InitialDesiredSize = NewIntSize;
+			InitialDesiredSize = FVector2f(NewIntSize);
 		}
 	}
-	SetCachedSize(NewIntSize);
+	SetCachedSize(FVector2f(NewIntSize));
 }
 
 FSlateRect SWindow::GetFullScreenInfo() const
@@ -1035,7 +1035,7 @@ FSlateRect SWindow::GetFullScreenInfo() const
 	return FSlateRect();
 }
 
-void SWindow::SetCachedScreenPosition(FVector2D NewPosition)
+void SWindow::SetCachedScreenPosition(UE::Slate::FDeprecateVector2DParameter NewPosition)
 {
 	ScreenPosition = NewPosition;
 
@@ -1044,11 +1044,12 @@ void SWindow::SetCachedScreenPosition(FVector2D NewPosition)
 	OnWindowMoved.ExecuteIfBound( SharedThis( this ) );
 }
 
-void SWindow::SetCachedSize( FVector2D NewSize )
+void SWindow::SetCachedSize( UE::Slate::FDeprecateVector2DParameter NewSize )
 {
 	if( NativeWindow.IsValid() )
 	{
-		NativeWindow->AdjustCachedSize( NewSize );
+		FVector2d NewSize2d(NewSize);
+		NativeWindow->AdjustCachedSize( NewSize2d );
 	}
 
 	if(Size != NewSize)
@@ -1069,7 +1070,7 @@ bool SWindow::IsMorphingSize() const
 }
 
 
-void SWindow::MorphToPosition( const FCurveSequence& Sequence, const float TargetOpacity, const FVector2D& TargetPosition )
+void SWindow::MorphToPosition( const FCurveSequence& Sequence, const float TargetOpacity, const UE::Slate::FDeprecateVector2DParameter& TargetPosition )
 {
 	Morpher.bIsAnimatingWindowSize = false;
 	Morpher.Sequence = Sequence;
@@ -1143,15 +1144,15 @@ void SWindow::UpdateMorphTargetShape( const FSlateRect& TargetShape )
 	Morpher.TargetMorphShape = TargetShape;
 }
 
-void SWindow::UpdateMorphTargetPosition( const FVector2D& TargetPosition )
+void SWindow::UpdateMorphTargetPosition( const UE::Slate::FDeprecateVector2DParameter& TargetPosition )
 {
 	Morpher.TargetMorphShape.Left = Morpher.TargetMorphShape.Right = TargetPosition.X;
 	Morpher.TargetMorphShape.Top = Morpher.TargetMorphShape.Bottom = TargetPosition.Y;
 }
 
-FVector2D SWindow::GetMorphTargetPosition() const
+UE::Slate::FDeprecateVector2DResult SWindow::GetMorphTargetPosition() const
 {
-	return FVector2D( Morpher.TargetMorphShape.Left, Morpher.TargetMorphShape.Top );
+	return UE::Slate::FDeprecateVector2DResult(FVector2f( Morpher.TargetMorphShape.Left, Morpher.TargetMorphShape.Top ));
 }
 
 
@@ -1415,7 +1416,7 @@ void SWindow::ShowWindow()
 		if ( SizingRule == ESizingRule::Autosized && AutoCenterRule != EAutoCenter::None )
 		{
 			SlatePrepass( FSlateApplicationBase::Get().GetApplicationScale() * GetDPIScaleFactor() );
-			const FVector2D WindowDesiredSizePixels = GetDesiredSizeDesktopPixels();
+			const FVector2f WindowDesiredSizePixels = GetDesiredSizeDesktopPixels();
 			ReshapeWindow(InitialDesiredScreenPosition - (WindowDesiredSizePixels * 0.5f), WindowDesiredSizePixels);
 		}
 
@@ -1581,9 +1582,9 @@ bool SWindow::IsTopmostWindow() const
 }
 
 /** @return true if mouse coordinates is within this window */
-bool SWindow::IsScreenspaceMouseWithin(FVector2D ScreenspaceMouseCoordinate) const
+bool SWindow::IsScreenspaceMouseWithin(UE::Slate::FDeprecateVector2DParameter ScreenspaceMouseCoordinate) const
 {
-	const FVector2D LocalMouseCoordinate = ScreenspaceMouseCoordinate - ScreenPosition;
+	const FVector2f LocalMouseCoordinate = ScreenspaceMouseCoordinate - ScreenPosition;
 	return !LocalMouseCoordinate.ContainsNaN() && NativeWindow->IsPointInWindow(FMath::TruncToInt(LocalMouseCoordinate.X), FMath::TruncToInt(LocalMouseCoordinate.Y));
 }
 
@@ -1876,13 +1877,13 @@ int32 SWindow::SwitchWorlds( int32 WorldId ) const
 	return OnWorldSwitchHack.IsBound() ? OnWorldSwitchHack.Execute( WorldId ) : false;
 }
 
-bool PointWithinSlateRect(const FVector2D& Point, const FSlateRect& Rect)
+bool PointWithinSlateRect(const FVector2f& Point, const FSlateRect& Rect)
 {
 	return Point.X >= Rect.Left && Point.X < Rect.Right &&
 		Point.Y >= Rect.Top && Point.Y < Rect.Bottom;
 }
 
-EWindowZone::Type SWindow::GetCurrentWindowZone(FVector2D LocalMousePosition)
+EWindowZone::Type SWindow::GetCurrentWindowZone(UE::Slate::FDeprecateVector2DParameter LocalMousePosition)
 {
 	const bool bIsFullscreenMode = GetWindowMode() == EWindowMode::WindowedFullscreen || GetWindowMode() == EWindowMode::Fullscreen;
 	const bool bIsBorderlessGameWindow = Type == EWindowType::GameWindow && !bHasOSWindowBorder;
@@ -2013,12 +2014,12 @@ SWindow::SWindow()
 	, bManualManageDPI( false )
 	, bAllowFastUpdate( false )
 	, WindowActivationPolicy( EWindowActivationPolicy::Always )
-	, InitialDesiredScreenPosition( FVector2D::ZeroVector )
-	, InitialDesiredSize( FVector2D::ZeroVector )
-	, ScreenPosition( FVector2D::ZeroVector )
-	, PreFullscreenPosition( FVector2D::ZeroVector )
-	, Size( FVector2D::ZeroVector )
-	, ViewportSize( FVector2D::ZeroVector )
+	, InitialDesiredScreenPosition( FVector2f::ZeroVector )
+	, InitialDesiredSize( FVector2f::ZeroVector )
+	, ScreenPosition( FVector2f::ZeroVector )
+	, PreFullscreenPosition( FVector2f::ZeroVector )
+	, Size( FVector2f::ZeroVector )
+	, ViewportSize( FVector2f::ZeroVector )
 	, TitleBarSize( SWindowDefs::DefaultTitleBarSize )
 	, ContentSlot(nullptr)
 	, Style( &FCoreStyle::Get().GetWidgetStyle<FWindowStyle>("Window") )
@@ -2149,7 +2150,7 @@ FOptionalSize SWindow::GetTitleBarSize() const
 }
 
 
-FVector2D SWindow::GetDesiredSizeDesktopPixels() const
+UE::Slate::FDeprecateVector2DResult SWindow::GetDesiredSizeDesktopPixels() const
 {
 	// Note that the window already takes the layout multiplier
 	// into account when computing its desired size.
@@ -2201,7 +2202,7 @@ void SWindow::SetWindowMode( EWindowMode::Type NewWindowMode )
 
 		NativeWindow->SetWindowMode( NewWindowMode );
 
-		const FVector2D vp = IsMirrorWindow() ? GetSizeInScreen() : GetViewportSize();
+		const FVector2f vp = IsMirrorWindow() ? GetSizeInScreen() : GetViewportSize();
 		FSlateApplicationBase::Get().GetRenderer()->UpdateFullscreenState(SharedThis(this), (uint32)vp.X, (uint32)vp.Y);
 
 		if( TitleArea.IsValid() )
diff --git a/Engine/Source/Runtime/SlateCore/Public/Application/SlateApplicationBase.h b/Engine/Source/Runtime/SlateCore/Public/Application/SlateApplicationBase.h
--- a/Engine/Source/Runtime/SlateCore/Public/Application/SlateApplicationBase.h
+++ b/Engine/Source/Runtime/SlateCore/Public/Application/SlateApplicationBase.h
@@ -10,6 +10,7 @@
 #include "Rendering/SlateRenderer.h"
 #include "Misc/CoreDelegates.h"
 #include "Async/TaskGraphInterfaces.h"
+#include "Types/SlateVector2.h"
 
 DECLARE_MULTICAST_DELEGATE_TwoParams(FOnDebugSafeZoneChanged, const FMargin&, bool);
 
@@ -78,7 +79,7 @@ private:
 private:	
 	FSlateApplicationBase* SlateApp;
 	// @see FSlateApplicationBase::LocateWidgetInWindow
-	FWidgetPath LocateWidgetInWindow(FVector2D ScreenspaceMouseCoordinate, const TSharedRef<SWindow>& Window, bool bIgnoreEnabledStatus, int32 UserIndex) const;
+	FWidgetPath LocateWidgetInWindow(FVector2f ScreenspaceMouseCoordinate, const TSharedRef<SWindow>& Window, bool bIgnoreEnabledStatus, int32 UserIndex) const;
 };
 
 /**
@@ -211,21 +212,21 @@ public:
 	 *
 	 * @return Cursor position.
 	 */
-	virtual FVector2D GetCursorPos( ) const = 0;
+	virtual UE::Slate::FDeprecateVector2DResult GetCursorPos( ) const = 0;
 
 	/**
 	* Gets the last known position of the cursor.
 	*
 	* @return Cursor position.
 	*/
-	virtual FVector2D GetLastCursorPos( ) const = 0;
+	virtual UE::Slate::FDeprecateVector2DResult GetLastCursorPos( ) const = 0;
 
 	/**
 	 * Gets the size of the cursor..
 	 *
 	 * @return Cursor size.
 	 */
-	virtual FVector2D GetCursorSize( ) const = 0;
+	virtual UE::Slate::FDeprecateVector2DResult GetCursorSize( ) const = 0;
 
 	/** 
 	 * Whether the software cursor is enabled for this application.  
@@ -255,7 +256,7 @@ public:
 	*/
 	void GetCachedDisplayMetrics(FDisplayMetrics& OutDisplayMetrics) const;
 
-	void GetSafeZoneSize(FMargin& SafeZone, const FVector2D& OverrideSize);
+	void GetSafeZoneSize(FMargin& SafeZone, const UE::Slate::FDeprecateVector2DParameter& OverrideSize);
 
 	/**
 	 * Get the highest level of window transparency support currently enabled by this application
@@ -365,7 +366,7 @@ public:
 	 *
 	 * @return The path to the widget.
 	 */
-	virtual FWidgetPath LocateWindowUnderMouse( FVector2D ScreenspaceMouseCoordinate, const TArray< TSharedRef<SWindow > >& Windows, bool bIgnoreEnabledStatus = false, int32 UserIndex = INDEX_NONE) = 0;
+	virtual FWidgetPath LocateWindowUnderMouse( UE::Slate::FDeprecateVector2DParameter ScreenspaceMouseCoordinate, const TArray< TSharedRef<SWindow > >& Windows, bool bIgnoreEnabledStatus = false, int32 UserIndex = INDEX_NONE) = 0;
 
 	/**
 	 * Calculates the tooltip window position.
@@ -374,7 +375,7 @@ public:
 	 * @param InSize The size of the tooltip window.
 	 * @return The suggested position.
 	 */
-	virtual FVector2D CalculateTooltipWindowPosition(const FSlateRect& InAnchorRect, const FVector2D& InSize, bool bAutoAdjustForDPIScale) const = 0;
+	virtual UE::Slate::FDeprecateVector2DResult CalculateTooltipWindowPosition(const FSlateRect& InAnchorRect, const UE::Slate::FDeprecateVector2DParameter& InSize, bool bAutoAdjustForDPIScale) const = 0;
 
 	/** @return true if 'WindowToTest' is being used to display the current tooltip and the tooltip is interactive. */
 	virtual bool IsWindowHousingInteractiveTooltip(const TSharedRef<const SWindow>& WindowToTest) const = 0;
@@ -597,7 +598,7 @@ protected:
 	virtual bool ShowUserFocus(const TSharedPtr<const SWidget> Widget) const = 0;
 
 	/** Given a window, locate a widget under the cursor in it; returns an invalid path if cursor is not over this window. */
-	virtual FWidgetPath LocateWidgetInWindow(FVector2D ScreenspaceMouseCoordinate, const TSharedRef<SWindow>& Window, bool bIgnoreEnabledStatus, int32 UserIndex) const = 0;
+	virtual FWidgetPath LocateWidgetInWindow(UE::Slate::FDeprecateVector2DParameter ScreenspaceMouseCoordinate, const TSharedRef<SWindow>& Window, bool bIgnoreEnabledStatus, int32 UserIndex) const = 0;
 	
 	void UpdateCustomSafeZone(const FMargin& NewSafeZoneRatio, bool bShouldRecacheMetrics);
 #if WITH_EDITOR
diff --git a/Engine/Source/Runtime/SlateCore/Public/Brushes/SlateBorderBrush.h b/Engine/Source/Runtime/SlateCore/Public/Brushes/SlateBorderBrush.h
--- a/Engine/Source/Runtime/SlateCore/Public/Brushes/SlateBorderBrush.h
+++ b/Engine/Source/Runtime/SlateCore/Public/Brushes/SlateBorderBrush.h
@@ -4,6 +4,7 @@
 
 #include "CoreMinimal.h"
 #include "Styling/SlateBrush.h"
+#include "Types/SlateVector2.h"
 
 /**
  * Similar to FSlateBoxBrush but has no middle and the sides tile instead of stretching.
@@ -19,84 +20,84 @@ struct SLATECORE_API FSlateBorderBrush
 	 * @param InImageType The type of image this this is
 	 */
 	FORCENOINLINE FSlateBorderBrush( const FName& InImageName, const FMargin& InMargin, const FLinearColor& InColorAndOpacity = FLinearColor(1,1,1,1), ESlateBrushImageType::Type InImageType = ESlateBrushImageType::FullColor )
-		: FSlateBrush(ESlateBrushDrawType::Border, InImageName, InMargin, ESlateBrushTileType::Both, InImageType, FVector2D::ZeroVector, InColorAndOpacity)
+		: FSlateBrush(ESlateBrushDrawType::Border, InImageName, InMargin, ESlateBrushTileType::Both, InImageType, FVector2f::ZeroVector, InColorAndOpacity)
 	{
 		// A border with no margin will not show up at all.
 		check(InMargin.GetDesiredSize().SizeSquared() > 0);
 	}
 
 	FORCENOINLINE FSlateBorderBrush( const FString& InImageName, const FMargin& InMargin, const FLinearColor& InColorAndOpacity = FLinearColor(1,1,1,1), ESlateBrushImageType::Type InImageType = ESlateBrushImageType::FullColor )
-		: FSlateBrush(ESlateBrushDrawType::Border, *InImageName, InMargin, ESlateBrushTileType::Both, InImageType, FVector2D::ZeroVector, InColorAndOpacity)
+		: FSlateBrush(ESlateBrushDrawType::Border, *InImageName, InMargin, ESlateBrushTileType::Both, InImageType, FVector2f::ZeroVector, InColorAndOpacity)
 	{
 		// A border with no margin will not show up at all.
 		check(InMargin.GetDesiredSize().SizeSquared() > 0);
 	}
 
 	FORCENOINLINE FSlateBorderBrush( const ANSICHAR* InImageName, const FMargin& InMargin, const FLinearColor& InColorAndOpacity = FLinearColor(1,1,1,1), ESlateBrushImageType::Type InImageType = ESlateBrushImageType::FullColor )
-		: FSlateBrush(ESlateBrushDrawType::Border, InImageName, InMargin, ESlateBrushTileType::Both, InImageType, FVector2D::ZeroVector, InColorAndOpacity)
+		: FSlateBrush(ESlateBrushDrawType::Border, InImageName, InMargin, ESlateBrushTileType::Both, InImageType, FVector2f::ZeroVector, InColorAndOpacity)
 	{
 		// A border with no margin will not show up at all.
 		check(InMargin.GetDesiredSize().SizeSquared() > 0);
 	}
 
 	FORCENOINLINE FSlateBorderBrush( const WIDECHAR* InImageName, const FMargin& InMargin, const FLinearColor& InColorAndOpacity = FLinearColor(1,1,1,1), ESlateBrushImageType::Type InImageType = ESlateBrushImageType::FullColor )
-		: FSlateBrush(ESlateBrushDrawType::Border, InImageName, InMargin, ESlateBrushTileType::Both, InImageType, FVector2D::ZeroVector, InColorAndOpacity)
+		: FSlateBrush(ESlateBrushDrawType::Border, InImageName, InMargin, ESlateBrushTileType::Both, InImageType, FVector2f::ZeroVector, InColorAndOpacity)
 	{
 		// A border with no margin will not show up at all.
 		check(InMargin.GetDesiredSize().SizeSquared() > 0);
 	}
 
 	FORCENOINLINE FSlateBorderBrush( const FName& InImageName, const FMargin& InMargin, const TSharedRef< FLinearColor >& InColorAndOpacity, ESlateBrushImageType::Type InImageType = ESlateBrushImageType::FullColor )
-		: FSlateBrush(ESlateBrushDrawType::Border, InImageName, InMargin, ESlateBrushTileType::Both, InImageType, FVector2D::ZeroVector, InColorAndOpacity)
+		: FSlateBrush(ESlateBrushDrawType::Border, InImageName, InMargin, ESlateBrushTileType::Both, InImageType, FVector2f::ZeroVector, InColorAndOpacity)
 	{
 		// A border with no margin will not show up at all.
 		check(InMargin.GetDesiredSize().SizeSquared() > 0);
 	}
 
 	FORCENOINLINE FSlateBorderBrush( const FString& InImageName, const FMargin& InMargin, const TSharedRef< FLinearColor >& InColorAndOpacity, ESlateBrushImageType::Type InImageType = ESlateBrushImageType::FullColor )
-		: FSlateBrush(ESlateBrushDrawType::Border, *InImageName, InMargin, ESlateBrushTileType::Both, InImageType, FVector2D::ZeroVector, InColorAndOpacity)
+		: FSlateBrush(ESlateBrushDrawType::Border, *InImageName, InMargin, ESlateBrushTileType::Both, InImageType, FVector2f::ZeroVector, InColorAndOpacity)
 	{
 		// A border with no margin will not show up at all.
 		check(InMargin.GetDesiredSize().SizeSquared() > 0);
 	}
 
 	FORCENOINLINE FSlateBorderBrush( const ANSICHAR* InImageName, const FMargin& InMargin, const TSharedRef< FLinearColor >& InColorAndOpacity, ESlateBrushImageType::Type InImageType = ESlateBrushImageType::FullColor )
-		: FSlateBrush(ESlateBrushDrawType::Border, InImageName, InMargin, ESlateBrushTileType::Both, InImageType, FVector2D::ZeroVector, InColorAndOpacity)
+		: FSlateBrush(ESlateBrushDrawType::Border, InImageName, InMargin, ESlateBrushTileType::Both, InImageType, FVector2f::ZeroVector, InColorAndOpacity)
 	{
 		// A border with no margin will not show up at all.
 		check(InMargin.GetDesiredSize().SizeSquared() > 0);
 	}
 
 	FORCENOINLINE FSlateBorderBrush( const TCHAR* InImageName, const FMargin& InMargin, const TSharedRef< FLinearColor >& InColorAndOpacity, ESlateBrushImageType::Type InImageType = ESlateBrushImageType::FullColor )
-		: FSlateBrush(ESlateBrushDrawType::Border, InImageName, InMargin, ESlateBrushTileType::Both, InImageType, FVector2D::ZeroVector, InColorAndOpacity)
+		: FSlateBrush(ESlateBrushDrawType::Border, InImageName, InMargin, ESlateBrushTileType::Both, InImageType, FVector2f::ZeroVector, InColorAndOpacity)
 	{
 		// A border with no margin will not show up at all.
 		check(InMargin.GetDesiredSize().SizeSquared() > 0);
 	}
 
 	FORCENOINLINE FSlateBorderBrush( const FName& InImageName, const FMargin& InMargin, const FSlateColor& InColorAndOpacity, ESlateBrushImageType::Type InImageType = ESlateBrushImageType::FullColor )
-		: FSlateBrush(ESlateBrushDrawType::Border, InImageName, InMargin, ESlateBrushTileType::Both, InImageType, FVector2D::ZeroVector, InColorAndOpacity)
+		: FSlateBrush(ESlateBrushDrawType::Border, InImageName, InMargin, ESlateBrushTileType::Both, InImageType, FVector2f::ZeroVector, InColorAndOpacity)
 	{
 		// A border with no margin will not show up at all.
 		check(InMargin.GetDesiredSize().SizeSquared() > 0);
 	}
 
 	FORCENOINLINE FSlateBorderBrush( const FString& InImageName, const FMargin& InMargin, const FSlateColor& InColorAndOpacity, ESlateBrushImageType::Type InImageType = ESlateBrushImageType::FullColor )
-		: FSlateBrush(ESlateBrushDrawType::Border, *InImageName, InMargin, ESlateBrushTileType::Both, InImageType, FVector2D::ZeroVector, InColorAndOpacity)
+		: FSlateBrush(ESlateBrushDrawType::Border, *InImageName, InMargin, ESlateBrushTileType::Both, InImageType, FVector2f::ZeroVector, InColorAndOpacity)
 	{
 		// A border with no margin will not show up at all.
 		check(InMargin.GetDesiredSize().SizeSquared() > 0);
 	}
 
 	FORCENOINLINE FSlateBorderBrush( const ANSICHAR* InImageName, const FMargin& InMargin, const FSlateColor& InColorAndOpacity, ESlateBrushImageType::Type InImageType = ESlateBrushImageType::FullColor )
-		: FSlateBrush(ESlateBrushDrawType::Border, InImageName, InMargin, ESlateBrushTileType::Both, InImageType, FVector2D::ZeroVector, InColorAndOpacity)
+		: FSlateBrush(ESlateBrushDrawType::Border, InImageName, InMargin, ESlateBrushTileType::Both, InImageType, FVector2f::ZeroVector, InColorAndOpacity)
 	{
 		// A border with no margin will not show up at all.
 		check(InMargin.GetDesiredSize().SizeSquared() > 0);
 	}
 
 	FORCENOINLINE FSlateBorderBrush( const TCHAR* InImageName, const FMargin& InMargin, const FSlateColor& InColorAndOpacity, ESlateBrushImageType::Type InImageType = ESlateBrushImageType::FullColor )
-		: FSlateBrush(ESlateBrushDrawType::Border, InImageName, InMargin, ESlateBrushTileType::Both, InImageType, FVector2D::ZeroVector, InColorAndOpacity)
+		: FSlateBrush(ESlateBrushDrawType::Border, InImageName, InMargin, ESlateBrushTileType::Both, InImageType, FVector2f::ZeroVector, InColorAndOpacity)
 	{
 		// A border with no margin will not show up at all.
 		check(InMargin.GetDesiredSize().SizeSquared() > 0);
@@ -109,7 +110,7 @@ struct SLATECORE_API FSlateBorderBrush
 	 * @param InImageType		The type of image this this is
 	 */
 	FORCENOINLINE FSlateBorderBrush(UObject* InResourceObject, const FMargin& InMargin, const FSlateColor& InColorAndOpacity = FSlateColor(FLinearColor(1, 1, 1, 1)), ESlateBrushImageType::Type InImageType = ESlateBrushImageType::FullColor)
-		: FSlateBrush(ESlateBrushDrawType::Border, NAME_None, InMargin, ESlateBrushTileType::Both, InImageType, FVector2D::ZeroVector, InColorAndOpacity, InResourceObject)
+		: FSlateBrush(ESlateBrushDrawType::Border, NAME_None, InMargin, ESlateBrushTileType::Both, InImageType, FVector2f::ZeroVector, InColorAndOpacity, InResourceObject)
 	{
 		// A border with no margin will not show up at all.
 		check(InMargin.GetDesiredSize().SizeSquared() > 0);
diff --git a/Engine/Source/Runtime/SlateCore/Public/Brushes/SlateBoxBrush.h b/Engine/Source/Runtime/SlateCore/Public/Brushes/SlateBoxBrush.h
--- a/Engine/Source/Runtime/SlateCore/Public/Brushes/SlateBoxBrush.h
+++ b/Engine/Source/Runtime/SlateCore/Public/Brushes/SlateBoxBrush.h
@@ -4,6 +4,7 @@
 
 #include "CoreMinimal.h"
 #include "Styling/SlateBrush.h"
+#include "Types/SlateVector2.h"
 
 /**
  * A 3x3 box where the sides stretch horizontally and vertically and the middle stretches to fill.
@@ -31,51 +32,51 @@ struct SLATECORE_API FSlateBoxBrush
 	 * @param InColorAndOpacity	Color and opacity scale.
 	 */
 	FORCENOINLINE FSlateBoxBrush( const FName& InImageName, const FMargin& InMargin, const FLinearColor& InColorAndOpacity = FLinearColor(1,1,1,1), ESlateBrushImageType::Type InImageType = ESlateBrushImageType::FullColor )
-		: FSlateBrush(ESlateBrushDrawType::Box, InImageName, InMargin, ESlateBrushTileType::NoTile, InImageType, FVector2D::ZeroVector, InColorAndOpacity)
+		: FSlateBrush(ESlateBrushDrawType::Box, InImageName, InMargin, ESlateBrushTileType::NoTile, InImageType, FVector2f::ZeroVector, InColorAndOpacity)
 	{ }
 
 	FORCENOINLINE FSlateBoxBrush( const FString& InImageName, const FMargin& InMargin, const FLinearColor& InColorAndOpacity = FLinearColor(1,1,1,1), ESlateBrushImageType::Type InImageType = ESlateBrushImageType::FullColor )
-		: FSlateBrush(ESlateBrushDrawType::Box, *InImageName, InMargin, ESlateBrushTileType::NoTile, InImageType, FVector2D::ZeroVector, InColorAndOpacity)
+		: FSlateBrush(ESlateBrushDrawType::Box, *InImageName, InMargin, ESlateBrushTileType::NoTile, InImageType, FVector2f::ZeroVector, InColorAndOpacity)
 	{ }
 
 	FORCENOINLINE FSlateBoxBrush( const ANSICHAR* InImageName, const FMargin& InMargin, const FLinearColor& InColorAndOpacity = FLinearColor(1,1,1,1), ESlateBrushImageType::Type InImageType = ESlateBrushImageType::FullColor )
-		: FSlateBrush(ESlateBrushDrawType::Box, InImageName, InMargin, ESlateBrushTileType::NoTile, InImageType, FVector2D::ZeroVector, InColorAndOpacity)
+		: FSlateBrush(ESlateBrushDrawType::Box, InImageName, InMargin, ESlateBrushTileType::NoTile, InImageType, FVector2f::ZeroVector, InColorAndOpacity)
 	{ } 
 
 	FORCENOINLINE FSlateBoxBrush( const WIDECHAR* InImageName, const FMargin& InMargin, const FLinearColor& InColorAndOpacity = FLinearColor(1,1,1,1), ESlateBrushImageType::Type InImageType = ESlateBrushImageType::FullColor )
-		: FSlateBrush(ESlateBrushDrawType::Box, InImageName, InMargin, ESlateBrushTileType::NoTile, InImageType, FVector2D::ZeroVector, InColorAndOpacity)
+		: FSlateBrush(ESlateBrushDrawType::Box, InImageName, InMargin, ESlateBrushTileType::NoTile, InImageType, FVector2f::ZeroVector, InColorAndOpacity)
 	{ }
 
 	FORCENOINLINE FSlateBoxBrush( const FName& InImageName, const FMargin& InMargin, const TSharedRef< FLinearColor >& InColorAndOpacity, ESlateBrushImageType::Type InImageType = ESlateBrushImageType::FullColor )
-		: FSlateBrush(ESlateBrushDrawType::Box, InImageName, InMargin, ESlateBrushTileType::NoTile, InImageType, FVector2D::ZeroVector, InColorAndOpacity)
+		: FSlateBrush(ESlateBrushDrawType::Box, InImageName, InMargin, ESlateBrushTileType::NoTile, InImageType, FVector2f::ZeroVector, InColorAndOpacity)
 	{ }
 
 	FORCENOINLINE FSlateBoxBrush( const FString& InImageName, const FMargin& InMargin, const TSharedRef< FLinearColor >& InColorAndOpacity, ESlateBrushImageType::Type InImageType = ESlateBrushImageType::FullColor )
-		: FSlateBrush(ESlateBrushDrawType::Box, *InImageName, InMargin, ESlateBrushTileType::NoTile, InImageType, FVector2D::ZeroVector, InColorAndOpacity)
+		: FSlateBrush(ESlateBrushDrawType::Box, *InImageName, InMargin, ESlateBrushTileType::NoTile, InImageType, FVector2f::ZeroVector, InColorAndOpacity)
 	{ }
 
 	FORCENOINLINE FSlateBoxBrush( const ANSICHAR* InImageName, const FMargin& InMargin, const TSharedRef< FLinearColor >& InColorAndOpacity, ESlateBrushImageType::Type InImageType = ESlateBrushImageType::FullColor )
-		: FSlateBrush(ESlateBrushDrawType::Box, InImageName, InMargin, ESlateBrushTileType::NoTile, InImageType, FVector2D::ZeroVector, InColorAndOpacity)
+		: FSlateBrush(ESlateBrushDrawType::Box, InImageName, InMargin, ESlateBrushTileType::NoTile, InImageType, FVector2f::ZeroVector, InColorAndOpacity)
 	{ }
 
 	FORCENOINLINE FSlateBoxBrush( const TCHAR* InImageName, const FMargin& InMargin, const TSharedRef< FLinearColor >& InColorAndOpacity, ESlateBrushImageType::Type InImageType = ESlateBrushImageType::FullColor )
-		: FSlateBrush(ESlateBrushDrawType::Box, InImageName, InMargin, ESlateBrushTileType::NoTile, InImageType, FVector2D::ZeroVector, InColorAndOpacity)
+		: FSlateBrush(ESlateBrushDrawType::Box, InImageName, InMargin, ESlateBrushTileType::NoTile, InImageType, FVector2f::ZeroVector, InColorAndOpacity)
 	{ }
 
 	FORCENOINLINE FSlateBoxBrush( const FName& InImageName, const FMargin& InMargin, const FSlateColor& InColorAndOpacity, ESlateBrushImageType::Type InImageType = ESlateBrushImageType::FullColor )
-		: FSlateBrush(ESlateBrushDrawType::Box, InImageName, InMargin, ESlateBrushTileType::NoTile, InImageType, FVector2D::ZeroVector, InColorAndOpacity)
+		: FSlateBrush(ESlateBrushDrawType::Box, InImageName, InMargin, ESlateBrushTileType::NoTile, InImageType, FVector2f::ZeroVector, InColorAndOpacity)
 	{ }
 
 	FORCENOINLINE FSlateBoxBrush( const FString& InImageName, const FMargin& InMargin, const FSlateColor& InColorAndOpacity, ESlateBrushImageType::Type InImageType = ESlateBrushImageType::FullColor )
-		: FSlateBrush(ESlateBrushDrawType::Box, *InImageName, InMargin, ESlateBrushTileType::NoTile, InImageType, FVector2D::ZeroVector, InColorAndOpacity)
+		: FSlateBrush(ESlateBrushDrawType::Box, *InImageName, InMargin, ESlateBrushTileType::NoTile, InImageType, FVector2f::ZeroVector, InColorAndOpacity)
 	{ }
 
 	FORCENOINLINE FSlateBoxBrush( const ANSICHAR* InImageName, const FMargin& InMargin, const FSlateColor& InColorAndOpacity, ESlateBrushImageType::Type InImageType = ESlateBrushImageType::FullColor )
-		: FSlateBrush(ESlateBrushDrawType::Box, InImageName, InMargin, ESlateBrushTileType::NoTile, InImageType, FVector2D::ZeroVector, InColorAndOpacity)
+		: FSlateBrush(ESlateBrushDrawType::Box, InImageName, InMargin, ESlateBrushTileType::NoTile, InImageType, FVector2f::ZeroVector, InColorAndOpacity)
 	{ }
 
 	FORCENOINLINE FSlateBoxBrush( const TCHAR* InImageName, const FMargin& InMargin, const FSlateColor& InColorAndOpacity, ESlateBrushImageType::Type InImageType = ESlateBrushImageType::FullColor )
-		: FSlateBrush(ESlateBrushDrawType::Box, InImageName, InMargin, ESlateBrushTileType::NoTile, InImageType, FVector2D::ZeroVector, InColorAndOpacity)
+		: FSlateBrush(ESlateBrushDrawType::Box, InImageName, InMargin, ESlateBrushTileType::NoTile, InImageType, FVector2f::ZeroVector, InColorAndOpacity)
 	{ }
 
 	/**
@@ -86,51 +87,51 @@ struct SLATECORE_API FSlateBoxBrush
 	 * @param InMargin          The size of corners and sides in texture space.
 	 * @param InColorAndOpacity	Color and opacity scale. Note of the image type is ImageType_TintMask, this value should be in HSV
 	 */
-	FORCENOINLINE FSlateBoxBrush( const FName& InImageName, const FVector2D& ImageSize, const FMargin& InMargin, const FLinearColor& InColorAndOpacity = FLinearColor(1,1,1,1), ESlateBrushImageType::Type InImageType = ESlateBrushImageType::FullColor )
+	FORCENOINLINE FSlateBoxBrush( const FName& InImageName, const UE::Slate::FDeprecateVector2DParameter& ImageSize, const FMargin& InMargin, const FLinearColor& InColorAndOpacity = FLinearColor(1,1,1,1), ESlateBrushImageType::Type InImageType = ESlateBrushImageType::FullColor )
 		: FSlateBrush(ESlateBrushDrawType::Box, InImageName, InMargin, ESlateBrushTileType::NoTile, InImageType, ImageSize, InColorAndOpacity)
 	{ }
 
-	FORCENOINLINE FSlateBoxBrush( const FString& InImageName, const FVector2D& ImageSize, const FMargin& InMargin, const FLinearColor& InColorAndOpacity = FLinearColor(1,1,1,1), ESlateBrushImageType::Type InImageType = ESlateBrushImageType::FullColor )
+	FORCENOINLINE FSlateBoxBrush( const FString& InImageName, const UE::Slate::FDeprecateVector2DParameter& ImageSize, const FMargin& InMargin, const FLinearColor& InColorAndOpacity = FLinearColor(1,1,1,1), ESlateBrushImageType::Type InImageType = ESlateBrushImageType::FullColor )
 		: FSlateBrush(ESlateBrushDrawType::Box, *InImageName, InMargin, ESlateBrushTileType::NoTile, InImageType, ImageSize, InColorAndOpacity)
 	{ }
 
-	FORCENOINLINE FSlateBoxBrush( const ANSICHAR* InImageName, const FVector2D& ImageSize, const FMargin& InMargin, const FLinearColor& InColorAndOpacity = FLinearColor(1,1,1,1), ESlateBrushImageType::Type InImageType = ESlateBrushImageType::FullColor )
+	FORCENOINLINE FSlateBoxBrush( const ANSICHAR* InImageName, const UE::Slate::FDeprecateVector2DParameter& ImageSize, const FMargin& InMargin, const FLinearColor& InColorAndOpacity = FLinearColor(1,1,1,1), ESlateBrushImageType::Type InImageType = ESlateBrushImageType::FullColor )
 		: FSlateBrush(ESlateBrushDrawType::Box, InImageName, InMargin, ESlateBrushTileType::NoTile, InImageType, ImageSize, InColorAndOpacity)
 	{ }
 
-	FORCENOINLINE FSlateBoxBrush( const WIDECHAR* InImageName, const FVector2D& ImageSize, const FMargin& InMargin, const FLinearColor& InColorAndOpacity = FLinearColor(1,1,1,1), ESlateBrushImageType::Type InImageType = ESlateBrushImageType::FullColor )
+	FORCENOINLINE FSlateBoxBrush( const WIDECHAR* InImageName, const UE::Slate::FDeprecateVector2DParameter& ImageSize, const FMargin& InMargin, const FLinearColor& InColorAndOpacity = FLinearColor(1,1,1,1), ESlateBrushImageType::Type InImageType = ESlateBrushImageType::FullColor )
 		: FSlateBrush(ESlateBrushDrawType::Box, InImageName, InMargin, ESlateBrushTileType::NoTile, InImageType, ImageSize, InColorAndOpacity)
 	{ }
 
-	FORCENOINLINE FSlateBoxBrush( const FName& InImageName, const FVector2D& ImageSize, const FMargin& InMargin, const TSharedRef< FLinearColor >& InColorAndOpacity, ESlateBrushImageType::Type InImageType = ESlateBrushImageType::FullColor )
+	FORCENOINLINE FSlateBoxBrush( const FName& InImageName, const UE::Slate::FDeprecateVector2DParameter& ImageSize, const FMargin& InMargin, const TSharedRef< FLinearColor >& InColorAndOpacity, ESlateBrushImageType::Type InImageType = ESlateBrushImageType::FullColor )
 		: FSlateBrush(ESlateBrushDrawType::Box, InImageName, InMargin, ESlateBrushTileType::NoTile, InImageType, ImageSize, InColorAndOpacity)
 	{ }
 
-	FORCENOINLINE FSlateBoxBrush( const FString& InImageName, const FVector2D& ImageSize, const FMargin& InMargin, const TSharedRef< FLinearColor >& InColorAndOpacity, ESlateBrushImageType::Type InImageType = ESlateBrushImageType::FullColor )
+	FORCENOINLINE FSlateBoxBrush( const FString& InImageName, const UE::Slate::FDeprecateVector2DParameter& ImageSize, const FMargin& InMargin, const TSharedRef< FLinearColor >& InColorAndOpacity, ESlateBrushImageType::Type InImageType = ESlateBrushImageType::FullColor )
 		: FSlateBrush(ESlateBrushDrawType::Box, *InImageName, InMargin, ESlateBrushTileType::NoTile, InImageType, ImageSize, InColorAndOpacity)
 	{ }
 
-	FORCENOINLINE FSlateBoxBrush( const ANSICHAR* InImageName, const FVector2D& ImageSize, const FMargin& InMargin, const TSharedRef< FLinearColor >& InColorAndOpacity, ESlateBrushImageType::Type InImageType = ESlateBrushImageType::FullColor )
+	FORCENOINLINE FSlateBoxBrush( const ANSICHAR* InImageName, const UE::Slate::FDeprecateVector2DParameter& ImageSize, const FMargin& InMargin, const TSharedRef< FLinearColor >& InColorAndOpacity, ESlateBrushImageType::Type InImageType = ESlateBrushImageType::FullColor )
 		: FSlateBrush(ESlateBrushDrawType::Box, InImageName, InMargin, ESlateBrushTileType::NoTile, InImageType, ImageSize, InColorAndOpacity)
 	{ }
 
-	FORCENOINLINE FSlateBoxBrush( const TCHAR* InImageName, const FVector2D& ImageSize, const FMargin& InMargin, const TSharedRef< FLinearColor >& InColorAndOpacity, ESlateBrushImageType::Type InImageType = ESlateBrushImageType::FullColor )
+	FORCENOINLINE FSlateBoxBrush( const TCHAR* InImageName, const UE::Slate::FDeprecateVector2DParameter& ImageSize, const FMargin& InMargin, const TSharedRef< FLinearColor >& InColorAndOpacity, ESlateBrushImageType::Type InImageType = ESlateBrushImageType::FullColor )
 		: FSlateBrush(ESlateBrushDrawType::Box, InImageName, InMargin, ESlateBrushTileType::NoTile, InImageType, ImageSize, InColorAndOpacity)
 	{ }
 
-	FORCENOINLINE FSlateBoxBrush( const FName& InImageName, const FVector2D& ImageSize, const FMargin& InMargin, const FSlateColor& InColorAndOpacity, ESlateBrushImageType::Type InImageType = ESlateBrushImageType::FullColor )
+	FORCENOINLINE FSlateBoxBrush( const FName& InImageName, const UE::Slate::FDeprecateVector2DParameter& ImageSize, const FMargin& InMargin, const FSlateColor& InColorAndOpacity, ESlateBrushImageType::Type InImageType = ESlateBrushImageType::FullColor )
 		: FSlateBrush(ESlateBrushDrawType::Box, InImageName, InMargin, ESlateBrushTileType::NoTile, InImageType, ImageSize, InColorAndOpacity)
 	{ }
 
-	FORCENOINLINE FSlateBoxBrush( const FString& InImageName, const FVector2D& ImageSize, const FMargin& InMargin, const FSlateColor& InColorAndOpacity, ESlateBrushImageType::Type InImageType = ESlateBrushImageType::FullColor )
+	FORCENOINLINE FSlateBoxBrush( const FString& InImageName, const UE::Slate::FDeprecateVector2DParameter& ImageSize, const FMargin& InMargin, const FSlateColor& InColorAndOpacity, ESlateBrushImageType::Type InImageType = ESlateBrushImageType::FullColor )
 		: FSlateBrush(ESlateBrushDrawType::Box, *InImageName, InMargin, ESlateBrushTileType::NoTile, InImageType, ImageSize, InColorAndOpacity)
 	{ }
 
-	FORCENOINLINE FSlateBoxBrush( const ANSICHAR* InImageName, const FVector2D& ImageSize, const FMargin& InMargin, const FSlateColor& InColorAndOpacity, ESlateBrushImageType::Type InImageType = ESlateBrushImageType::FullColor )
+	FORCENOINLINE FSlateBoxBrush( const ANSICHAR* InImageName, const UE::Slate::FDeprecateVector2DParameter& ImageSize, const FMargin& InMargin, const FSlateColor& InColorAndOpacity, ESlateBrushImageType::Type InImageType = ESlateBrushImageType::FullColor )
 		: FSlateBrush(ESlateBrushDrawType::Box, InImageName, InMargin, ESlateBrushTileType::NoTile, InImageType, ImageSize, InColorAndOpacity)
 	{ }
 
-	FORCENOINLINE FSlateBoxBrush( const TCHAR* InImageName, const FVector2D& ImageSize, const FMargin& InMargin, const FSlateColor& InColorAndOpacity, ESlateBrushImageType::Type InImageType = ESlateBrushImageType::FullColor )
+	FORCENOINLINE FSlateBoxBrush( const TCHAR* InImageName, const UE::Slate::FDeprecateVector2DParameter& ImageSize, const FMargin& InMargin, const FSlateColor& InColorAndOpacity, ESlateBrushImageType::Type InImageType = ESlateBrushImageType::FullColor )
 		: FSlateBrush(ESlateBrushDrawType::Box, InImageName, InMargin, ESlateBrushTileType::NoTile, InImageType, ImageSize, InColorAndOpacity)
 	{ }
 
@@ -143,7 +144,7 @@ struct SLATECORE_API FSlateBoxBrush
 	 */
 
 	FORCENOINLINE FSlateBoxBrush(UObject* InResourceObject, const FMargin& InMargin, const FSlateColor& InColorAndOpacity = FSlateColor(FLinearColor(1, 1, 1, 1)), ESlateBrushImageType::Type InImageType = ESlateBrushImageType::FullColor)
-		: FSlateBrush(ESlateBrushDrawType::Box, NAME_None, InMargin, ESlateBrushTileType::NoTile, InImageType, FVector2D::ZeroVector, InColorAndOpacity, InResourceObject)
+		: FSlateBrush(ESlateBrushDrawType::Box, NAME_None, InMargin, ESlateBrushTileType::NoTile, InImageType, FVector2f::ZeroVector, InColorAndOpacity, InResourceObject)
 	{ }
 
 	/**
@@ -155,7 +156,7 @@ struct SLATECORE_API FSlateBoxBrush
 	 * @param InColorAndOpacity	Color and opacity scale. Note of the image type is ImageType_TintMask, this value should be in HSV
 	 */
 
-	FORCENOINLINE FSlateBoxBrush(UObject* InResourceObject, const FVector2D& ImageSize, const FMargin& InMargin, const FSlateColor& InColorAndOpacity = FSlateColor(FLinearColor(1, 1, 1, 1)), ESlateBrushImageType::Type InImageType = ESlateBrushImageType::FullColor)
+	FORCENOINLINE FSlateBoxBrush(UObject* InResourceObject, const UE::Slate::FDeprecateVector2DParameter& ImageSize, const FMargin& InMargin, const FSlateColor& InColorAndOpacity = FSlateColor(FLinearColor(1, 1, 1, 1)), ESlateBrushImageType::Type InImageType = ESlateBrushImageType::FullColor)
 		: FSlateBrush(ESlateBrushDrawType::Box, NAME_None, InMargin, ESlateBrushTileType::NoTile, InImageType, ImageSize, InColorAndOpacity, InResourceObject)
 	{ }
 };
@@ -164,7 +165,7 @@ struct SLATECORE_API FSlateBoxBrush
 struct SLATECORE_API FSlateVectorBoxBrush
 	: public FSlateBoxBrush
 {
-	FSlateVectorBoxBrush(const FString& InImageName, const FVector2D& ImageSize, const FMargin& InMargin, const FLinearColor& InColorAndOpacity = FLinearColor(1, 1, 1, 1))
+	FSlateVectorBoxBrush(const FString& InImageName, const UE::Slate::FDeprecateVector2DParameter& ImageSize, const FMargin& InMargin, const FLinearColor& InColorAndOpacity = FLinearColor(1, 1, 1, 1))
 		: FSlateBoxBrush(InImageName, ImageSize, InMargin, InColorAndOpacity, ESlateBrushImageType::Vector)
 	{ }
 
diff --git a/Engine/Source/Runtime/SlateCore/Public/Brushes/SlateColorBrush.h b/Engine/Source/Runtime/SlateCore/Public/Brushes/SlateColorBrush.h
--- a/Engine/Source/Runtime/SlateCore/Public/Brushes/SlateColorBrush.h
+++ b/Engine/Source/Runtime/SlateCore/Public/Brushes/SlateColorBrush.h
@@ -18,7 +18,7 @@ struct SLATECORE_API FSlateColorBrush
 	 * @param InColor The linear color to use
 	 */
 	FORCENOINLINE FSlateColorBrush( const FLinearColor& InColor )
-		: FSlateBrush(ESlateBrushDrawType::Image, NAME_None, FMargin(0.0f), ESlateBrushTileType::NoTile, ESlateBrushImageType::NoImage, FVector2D::ZeroVector, InColor)
+		: FSlateBrush(ESlateBrushDrawType::Image, NAME_None, FMargin(0.0f), ESlateBrushTileType::NoTile, ESlateBrushImageType::NoImage, FVector2f::ZeroVector, InColor)
 	{ }
 
 	/** 
@@ -27,7 +27,7 @@ struct SLATECORE_API FSlateColorBrush
 	 * @param InColor The color value to use
 	 */
 	FORCENOINLINE FSlateColorBrush( const FColor& InColor )
-		: FSlateBrush(ESlateBrushDrawType::Image, NAME_None, FMargin(0.0f), ESlateBrushTileType::NoTile, ESlateBrushImageType::NoImage, FVector2D::ZeroVector, InColor.ReinterpretAsLinear())
+		: FSlateBrush(ESlateBrushDrawType::Image, NAME_None, FMargin(0.0f), ESlateBrushTileType::NoTile, ESlateBrushImageType::NoImage, FVector2f::ZeroVector, InColor.ReinterpretAsLinear())
 	{ }
 
 	/** 
@@ -36,7 +36,7 @@ struct SLATECORE_API FSlateColorBrush
 	 * @param InColor The shared color to use
 	 */
 	FORCENOINLINE FSlateColorBrush( const TSharedRef< FLinearColor >& InColor )
-		: FSlateBrush(ESlateBrushDrawType::Image, NAME_None, FMargin(0.0f), ESlateBrushTileType::NoTile, ESlateBrushImageType::NoImage, FVector2D::ZeroVector, InColor)
+		: FSlateBrush(ESlateBrushDrawType::Image, NAME_None, FMargin(0.0f), ESlateBrushTileType::NoTile, ESlateBrushImageType::NoImage, FVector2f::ZeroVector, InColor)
 	{ }
 
 	/**
@@ -45,6 +45,6 @@ struct SLATECORE_API FSlateColorBrush
 	 * @param InColor The shared color to use
 	 */
 	FORCENOINLINE FSlateColorBrush(const FSlateColor& InColor)
-		: FSlateBrush(ESlateBrushDrawType::Image, NAME_None, FMargin(0.0f), ESlateBrushTileType::NoTile, ESlateBrushImageType::NoImage, FVector2D::ZeroVector, InColor)
+		: FSlateBrush(ESlateBrushDrawType::Image, NAME_None, FMargin(0.0f), ESlateBrushTileType::NoTile, ESlateBrushImageType::NoImage, FVector2f::ZeroVector, InColor)
 	{ }
 };
diff --git a/Engine/Source/Runtime/SlateCore/Public/Brushes/SlateDynamicImageBrush.h b/Engine/Source/Runtime/SlateCore/Public/Brushes/SlateDynamicImageBrush.h
--- a/Engine/Source/Runtime/SlateCore/Public/Brushes/SlateDynamicImageBrush.h
+++ b/Engine/Source/Runtime/SlateCore/Public/Brushes/SlateDynamicImageBrush.h
@@ -5,6 +5,7 @@
 #include "CoreMinimal.h"
 #include "Layout/Margin.h"
 #include "Styling/SlateBrush.h"
+#include "Types/SlateVector2.h"
 
 /**
  * Ignores the Margin. Just renders the image. Can tile the image instead of stretching.
@@ -21,7 +22,7 @@ struct SLATECORE_API FSlateDynamicImageBrush
 	 */
 	FORCENOINLINE FSlateDynamicImageBrush( 
 		class UTexture2D* InTexture, 
-		const FVector2D& InImageSize,
+		const UE::Slate::FDeprecateVector2DParameter& InImageSize,
 		const FName InTextureName,
 		const FLinearColor& InTint = FLinearColor(1.0f, 1.0f, 1.0f, 1.0f), 
 		ESlateBrushTileType::Type InTiling = ESlateBrushTileType::NoTile, 
@@ -44,7 +45,7 @@ struct SLATECORE_API FSlateDynamicImageBrush
 	 */
 	FORCENOINLINE FSlateDynamicImageBrush( 
 		class UTexture2DDynamic* InTexture, 
-		const FVector2D& InImageSize,
+		const UE::Slate::FDeprecateVector2DParameter& InImageSize,
 		const FName InTextureName,
 		const FLinearColor& InTint = FLinearColor(1.0f, 1.0f, 1.0f, 1.0f), 
 		ESlateBrushTileType::Type InTiling = ESlateBrushTileType::NoTile, 
@@ -67,7 +68,7 @@ struct SLATECORE_API FSlateDynamicImageBrush
 	 */
 	FORCENOINLINE FSlateDynamicImageBrush( 
 		const FName InTextureName,
-		const FVector2D& InImageSize,
+		const UE::Slate::FDeprecateVector2DParameter& InImageSize,
 		const FLinearColor& InTint = FLinearColor(1.0f, 1.0f, 1.0f, 1.0f), 
 		ESlateBrushTileType::Type InTiling = ESlateBrushTileType::NoTile, 
 		ESlateBrushImageType::Type InImageType = ESlateBrushImageType::FullColor )
@@ -88,7 +89,7 @@ struct SLATECORE_API FSlateDynamicImageBrush
 	*/
 	static TSharedPtr<FSlateDynamicImageBrush> CreateWithImageData(
 		const FName InTextureName,
-		const FVector2D& InImageSize,
+		const UE::Slate::FDeprecateVector2DParameter& InImageSize,
 		const TArray<uint8>& InImageData,
 		const FLinearColor& InTint = FLinearColor(1.0f, 1.0f, 1.0f, 1.0f),
 		ESlateBrushTileType::Type InTiling = ESlateBrushTileType::NoTile,
diff --git a/Engine/Source/Runtime/SlateCore/Public/Brushes/SlateImageBrush.h b/Engine/Source/Runtime/SlateCore/Public/Brushes/SlateImageBrush.h
--- a/Engine/Source/Runtime/SlateCore/Public/Brushes/SlateImageBrush.h
+++ b/Engine/Source/Runtime/SlateCore/Public/Brushes/SlateImageBrush.h
@@ -5,6 +5,7 @@
 #include "CoreMinimal.h"
 #include "Layout/Margin.h"
 #include "Styling/SlateBrush.h"
+#include "Types/SlateVector2.h"
 
 /**
  * Ignores the Margin. Just renders the image. Can tile the image instead of stretching.
@@ -19,51 +20,51 @@ struct SLATECORE_API FSlateImageBrush
 	 * @param InTiling		How do we tile if at all?
 	 * @param InImageType	The type of image this this is
 	 */
-	FORCENOINLINE FSlateImageBrush( const FName& InImageName, const FVector2D& InImageSize, const FLinearColor& InTint = FLinearColor(1,1,1,1), ESlateBrushTileType::Type InTiling = ESlateBrushTileType::NoTile, ESlateBrushImageType::Type InImageType = ESlateBrushImageType::FullColor )
+	FORCENOINLINE FSlateImageBrush( const FName& InImageName, const UE::Slate::FDeprecateVector2DParameter& InImageSize, const FLinearColor& InTint = FLinearColor(1,1,1,1), ESlateBrushTileType::Type InTiling = ESlateBrushTileType::NoTile, ESlateBrushImageType::Type InImageType = ESlateBrushImageType::FullColor )
 		: FSlateBrush(ESlateBrushDrawType::Image, InImageName, FMargin(0), InTiling, InImageType, InImageSize, InTint)
 	{ }
 
-	FORCENOINLINE FSlateImageBrush( const FString& InImageName, const FVector2D& InImageSize, const FLinearColor& InTint = FLinearColor(1,1,1,1), ESlateBrushTileType::Type InTiling = ESlateBrushTileType::NoTile, ESlateBrushImageType::Type InImageType = ESlateBrushImageType::FullColor )
+	FORCENOINLINE FSlateImageBrush( const FString& InImageName, const UE::Slate::FDeprecateVector2DParameter& InImageSize, const FLinearColor& InTint = FLinearColor(1,1,1,1), ESlateBrushTileType::Type InTiling = ESlateBrushTileType::NoTile, ESlateBrushImageType::Type InImageType = ESlateBrushImageType::FullColor )
 		: FSlateBrush(ESlateBrushDrawType::Image, *InImageName, FMargin(0), InTiling, InImageType, InImageSize, InTint)
 	{ }
 
-	FORCENOINLINE FSlateImageBrush( const ANSICHAR* InImageName, const FVector2D& InImageSize, const FLinearColor& InTint = FLinearColor(1,1,1,1), ESlateBrushTileType::Type InTiling = ESlateBrushTileType::NoTile, ESlateBrushImageType::Type InImageType = ESlateBrushImageType::FullColor )
+	FORCENOINLINE FSlateImageBrush( const ANSICHAR* InImageName, const UE::Slate::FDeprecateVector2DParameter& InImageSize, const FLinearColor& InTint = FLinearColor(1,1,1,1), ESlateBrushTileType::Type InTiling = ESlateBrushTileType::NoTile, ESlateBrushImageType::Type InImageType = ESlateBrushImageType::FullColor )
 		: FSlateBrush(ESlateBrushDrawType::Image, InImageName, FMargin(0), InTiling, InImageType, InImageSize, InTint)
 	{ }
 
-	FORCENOINLINE FSlateImageBrush( const WIDECHAR* InImageName, const FVector2D& InImageSize, const FLinearColor& InTint = FLinearColor(1,1,1,1), ESlateBrushTileType::Type InTiling = ESlateBrushTileType::NoTile, ESlateBrushImageType::Type InImageType = ESlateBrushImageType::FullColor )
+	FORCENOINLINE FSlateImageBrush( const WIDECHAR* InImageName, const UE::Slate::FDeprecateVector2DParameter& InImageSize, const FLinearColor& InTint = FLinearColor(1,1,1,1), ESlateBrushTileType::Type InTiling = ESlateBrushTileType::NoTile, ESlateBrushImageType::Type InImageType = ESlateBrushImageType::FullColor )
 		: FSlateBrush(ESlateBrushDrawType::Image, InImageName, FMargin(0), InTiling, InImageType, InImageSize, InTint)
 	{ }
 
-	FORCENOINLINE FSlateImageBrush( const FName& InImageName, const FVector2D& InImageSize, const TSharedRef< FLinearColor >& InTint, ESlateBrushTileType::Type InTiling = ESlateBrushTileType::NoTile, ESlateBrushImageType::Type InImageType = ESlateBrushImageType::FullColor )
+	FORCENOINLINE FSlateImageBrush( const FName& InImageName, const UE::Slate::FDeprecateVector2DParameter& InImageSize, const TSharedRef< FLinearColor >& InTint, ESlateBrushTileType::Type InTiling = ESlateBrushTileType::NoTile, ESlateBrushImageType::Type InImageType = ESlateBrushImageType::FullColor )
 		: FSlateBrush(ESlateBrushDrawType::Image, InImageName, FMargin(0), InTiling, InImageType, InImageSize, InTint)
 	{ }
 
-	FORCENOINLINE FSlateImageBrush( const FString& InImageName, const FVector2D& InImageSize, const TSharedRef< FLinearColor >& InTint, ESlateBrushTileType::Type InTiling = ESlateBrushTileType::NoTile, ESlateBrushImageType::Type InImageType = ESlateBrushImageType::FullColor )
+	FORCENOINLINE FSlateImageBrush( const FString& InImageName, const UE::Slate::FDeprecateVector2DParameter& InImageSize, const TSharedRef< FLinearColor >& InTint, ESlateBrushTileType::Type InTiling = ESlateBrushTileType::NoTile, ESlateBrushImageType::Type InImageType = ESlateBrushImageType::FullColor )
 		: FSlateBrush(ESlateBrushDrawType::Image, *InImageName, FMargin(0), InTiling, InImageType, InImageSize, InTint)
 	{ }
 
-	FORCENOINLINE FSlateImageBrush( const ANSICHAR* InImageName, const FVector2D& InImageSize, const TSharedRef< FLinearColor >& InTint, ESlateBrushTileType::Type InTiling = ESlateBrushTileType::NoTile, ESlateBrushImageType::Type InImageType = ESlateBrushImageType::FullColor )
+	FORCENOINLINE FSlateImageBrush( const ANSICHAR* InImageName, const UE::Slate::FDeprecateVector2DParameter& InImageSize, const TSharedRef< FLinearColor >& InTint, ESlateBrushTileType::Type InTiling = ESlateBrushTileType::NoTile, ESlateBrushImageType::Type InImageType = ESlateBrushImageType::FullColor )
 		: FSlateBrush(ESlateBrushDrawType::Image, InImageName, FMargin(0), InTiling, InImageType, InImageSize, InTint)
 	{ }
 
-	FORCENOINLINE FSlateImageBrush( const TCHAR* InImageName, const FVector2D& InImageSize, const TSharedRef< FLinearColor >& InTint, ESlateBrushTileType::Type InTiling = ESlateBrushTileType::NoTile, ESlateBrushImageType::Type InImageType = ESlateBrushImageType::FullColor )
+	FORCENOINLINE FSlateImageBrush( const TCHAR* InImageName, const UE::Slate::FDeprecateVector2DParameter& InImageSize, const TSharedRef< FLinearColor >& InTint, ESlateBrushTileType::Type InTiling = ESlateBrushTileType::NoTile, ESlateBrushImageType::Type InImageType = ESlateBrushImageType::FullColor )
 		: FSlateBrush(ESlateBrushDrawType::Image, InImageName, FMargin(0), InTiling, InImageType, InImageSize, InTint)
 	{ }
 
-	FORCENOINLINE FSlateImageBrush( const FName& InImageName, const FVector2D& InImageSize, const FSlateColor& InTint, ESlateBrushTileType::Type InTiling = ESlateBrushTileType::NoTile, ESlateBrushImageType::Type InImageType = ESlateBrushImageType::FullColor )
+	FORCENOINLINE FSlateImageBrush( const FName& InImageName, const UE::Slate::FDeprecateVector2DParameter& InImageSize, const FSlateColor& InTint, ESlateBrushTileType::Type InTiling = ESlateBrushTileType::NoTile, ESlateBrushImageType::Type InImageType = ESlateBrushImageType::FullColor )
 		: FSlateBrush(ESlateBrushDrawType::Image, InImageName, FMargin(0), InTiling, InImageType, InImageSize, InTint)
 	{ }
 
-	FORCENOINLINE FSlateImageBrush( const FString& InImageName, const FVector2D& InImageSize, const FSlateColor& InTint, ESlateBrushTileType::Type InTiling = ESlateBrushTileType::NoTile, ESlateBrushImageType::Type InImageType = ESlateBrushImageType::FullColor )
+	FORCENOINLINE FSlateImageBrush( const FString& InImageName, const UE::Slate::FDeprecateVector2DParameter& InImageSize, const FSlateColor& InTint, ESlateBrushTileType::Type InTiling = ESlateBrushTileType::NoTile, ESlateBrushImageType::Type InImageType = ESlateBrushImageType::FullColor )
 		: FSlateBrush(ESlateBrushDrawType::Image, *InImageName, FMargin(0), InTiling, InImageType, InImageSize, InTint)
 	{ }
 
-	FORCENOINLINE FSlateImageBrush( const ANSICHAR* InImageName, const FVector2D& InImageSize, const FSlateColor& InTint, ESlateBrushTileType::Type InTiling = ESlateBrushTileType::NoTile, ESlateBrushImageType::Type InImageType = ESlateBrushImageType::FullColor )
+	FORCENOINLINE FSlateImageBrush( const ANSICHAR* InImageName, const UE::Slate::FDeprecateVector2DParameter& InImageSize, const FSlateColor& InTint, ESlateBrushTileType::Type InTiling = ESlateBrushTileType::NoTile, ESlateBrushImageType::Type InImageType = ESlateBrushImageType::FullColor )
 		: FSlateBrush(ESlateBrushDrawType::Image, InImageName, FMargin(0), InTiling, InImageType, InImageSize, InTint)
 	{ }
 
-	FORCENOINLINE FSlateImageBrush( const TCHAR* InImageName, const FVector2D& InImageSize, const FSlateColor& InTint, ESlateBrushTileType::Type InTiling = ESlateBrushTileType::NoTile, ESlateBrushImageType::Type InImageType = ESlateBrushImageType::FullColor )
+	FORCENOINLINE FSlateImageBrush( const TCHAR* InImageName, const UE::Slate::FDeprecateVector2DParameter& InImageSize, const FSlateColor& InTint, ESlateBrushTileType::Type InTiling = ESlateBrushTileType::NoTile, ESlateBrushImageType::Type InImageType = ESlateBrushImageType::FullColor )
 		: FSlateBrush(ESlateBrushDrawType::Image, InImageName, FMargin(0), InTiling, InImageType, InImageSize, InTint)
 	{ }
 
@@ -75,7 +76,7 @@ struct SLATECORE_API FSlateImageBrush
 	 * @param InImageType		The type of image this this is
 	 */
 
-	FORCENOINLINE FSlateImageBrush(UObject* InResourceObject, const FVector2D& InImageSize, const FSlateColor& InTint = FSlateColor(FLinearColor(1, 1, 1, 1)), ESlateBrushTileType::Type InTiling = ESlateBrushTileType::NoTile, ESlateBrushImageType::Type InImageType = ESlateBrushImageType::FullColor)
+	FORCENOINLINE FSlateImageBrush(UObject* InResourceObject, const UE::Slate::FDeprecateVector2DParameter& InImageSize, const FSlateColor& InTint = FSlateColor(FLinearColor(1, 1, 1, 1)), ESlateBrushTileType::Type InTiling = ESlateBrushTileType::NoTile, ESlateBrushImageType::Type InImageType = ESlateBrushImageType::FullColor)
 		: FSlateBrush(ESlateBrushDrawType::Image, NAME_None, FMargin(0), InTiling, InImageType, InImageSize, InTint, InResourceObject)
 	{ }
 
@@ -84,11 +85,11 @@ struct SLATECORE_API FSlateImageBrush
 struct SLATECORE_API FSlateVectorImageBrush
 	: public FSlateImageBrush
 {
-	FSlateVectorImageBrush(const FString& InImageName, const FVector2D& InImageSize, const FSlateColor& InTint, ESlateBrushTileType::Type InTiling = ESlateBrushTileType::NoTile)
+	FSlateVectorImageBrush(const FString& InImageName, const UE::Slate::FDeprecateVector2DParameter& InImageSize, const FSlateColor& InTint, ESlateBrushTileType::Type InTiling = ESlateBrushTileType::NoTile)
 		: FSlateImageBrush(InImageName, InImageSize, InTint, InTiling, ESlateBrushImageType::Vector)
 	{ }
 
-	FSlateVectorImageBrush(const FString& InImageName, const FVector2D& InImageSize, const FLinearColor& InTint = FLinearColor(1, 1, 1, 1), ESlateBrushTileType::Type InTiling = ESlateBrushTileType::NoTile)
+	FSlateVectorImageBrush(const FString& InImageName, const UE::Slate::FDeprecateVector2DParameter& InImageSize, const FLinearColor& InTint = FLinearColor(1, 1, 1, 1), ESlateBrushTileType::Type InTiling = ESlateBrushTileType::NoTile)
 		: FSlateImageBrush(InImageName, InImageSize, InTint, InTiling, ESlateBrushImageType::Vector)
 	{ }
 };
\ No newline at end of file
diff --git a/Engine/Source/Runtime/SlateCore/Public/Brushes/SlateNoResource.h b/Engine/Source/Runtime/SlateCore/Public/Brushes/SlateNoResource.h
--- a/Engine/Source/Runtime/SlateCore/Public/Brushes/SlateNoResource.h
+++ b/Engine/Source/Runtime/SlateCore/Public/Brushes/SlateNoResource.h
@@ -5,6 +5,7 @@
 #include "CoreMinimal.h"
 #include "Layout/Margin.h"
 #include "Styling/SlateBrush.h"
+#include "Types/SlateVector2.h"
 
 /**
  * A resource that has no appearance
@@ -17,7 +18,7 @@ struct SLATECORE_API FSlateNoResource
 	 *
 	 * @param InImageSize An optional image size (default is zero).
 	 */
-	FORCENOINLINE FSlateNoResource( const FVector2D& InImageSize = FVector2D::ZeroVector )
+	FORCENOINLINE FSlateNoResource( const UE::Slate::FDeprecateVector2DParameter& InImageSize = FVector2f::ZeroVector )
 		: FSlateBrush(ESlateBrushDrawType::NoDrawType, FName(NAME_None), FMargin(0), ESlateBrushTileType::NoTile, ESlateBrushImageType::NoImage, InImageSize)
 	{ }
 };
@@ -29,7 +30,7 @@ struct SLATECORE_API FSlateNoResource
 struct SLATECORE_API FSlateOptionalBrush
 	: public FSlateBrush
 {
-	FSlateOptionalBrush(const FVector2D& InImageSize = FVector2D::ZeroVector)
+	FSlateOptionalBrush(const UE::Slate::FDeprecateVector2DParameter& InImageSize = FVector2f::ZeroVector)
 		: FSlateBrush(ESlateBrushDrawType::NoDrawType, FName(NAME_None), FMargin(0), ESlateBrushTileType::NoTile, ESlateBrushImageType::NoImage, InImageSize)
 	{
 		bIsSet = false;
