Slate: Initial phase of FVector2D deprecation

diff --git a/Engine/Source/Runtime/SlateCore/Public/Brushes/SlateRoundedBoxBrush.h b/Engine/Source/Runtime/SlateCore/Public/Brushes/SlateRoundedBoxBrush.h
--- a/Engine/Source/Runtime/SlateCore/Public/Brushes/SlateRoundedBoxBrush.h
+++ b/Engine/Source/Runtime/SlateCore/Public/Brushes/SlateRoundedBoxBrush.h
@@ -4,6 +4,7 @@
 
 #include "CoreMinimal.h"
 #include "Styling/SlateBrush.h"
+#include "Types/SlateVector2.h"
 
 struct SLATECORE_API FSlateRoundedBoxBrush
 	: public FSlateBrush
@@ -15,7 +16,7 @@ struct SLATECORE_API FSlateRoundedBoxBrush
 	 * @param InColor 		Linear Fill Color 
 	 */
 	template<typename FillColorType>
-	FORCENOINLINE FSlateRoundedBoxBrush(const FillColorType& InColor, const FVector2D& InImageSize = FVector2D::ZeroVector)
+	FORCENOINLINE FSlateRoundedBoxBrush(const FillColorType& InColor, const UE::Slate::FDeprecateVector2DParameter& InImageSize = FVector2f::ZeroVector)
 		: FSlateBrush(ESlateBrushDrawType::RoundedBox, 
 					  NAME_None, 
 					  FMargin(0.0f), 
@@ -30,7 +31,23 @@ struct SLATECORE_API FSlateRoundedBoxBrush
 	}
 
 	template<typename FillColorType>
-	FORCENOINLINE FSlateRoundedBoxBrush(const FName InFillResourceName, const FillColorType& FillColorTint, const FVector2D& InImageSize = FVector2D::ZeroVector, ESlateBrushTileType::Type InTileType = ESlateBrushTileType::NoTile)
+	FORCENOINLINE FSlateRoundedBoxBrush(const FillColorType& InColor, const UE::Slate::FDeprecateVector2DResult& InImageSize)
+		: FSlateRoundedBoxBrush(InColor, UE::Slate::FDeprecateVector2DParameter(InImageSize))
+	{ 
+	}
+	template<typename FillColorType>
+	FORCENOINLINE FSlateRoundedBoxBrush(const FillColorType& InColor, const FVector2D& InImageSize)
+		: FSlateRoundedBoxBrush(InColor, UE::Slate::FDeprecateVector2DParameter(UE::Slate::CastToVector2f(InImageSize)))
+	{ 
+	}
+	template<typename FillColorType>
+	FORCENOINLINE FSlateRoundedBoxBrush(const FillColorType& InColor, const FVector2f& InImageSize)
+		: FSlateRoundedBoxBrush(InColor, UE::Slate::FDeprecateVector2DParameter(InImageSize))
+	{ 
+	}
+
+	template<typename FillColorType>
+	FORCENOINLINE FSlateRoundedBoxBrush(const FName InFillResourceName, const FillColorType& FillColorTint, const UE::Slate::FDeprecateVector2DParameter& InImageSize = FVector2f::ZeroVector, ESlateBrushTileType::Type InTileType = ESlateBrushTileType::NoTile)
 		: FSlateBrush(ESlateBrushDrawType::RoundedBox,
 			InFillResourceName,
 			FMargin(0.0f),
@@ -52,7 +69,7 @@ struct SLATECORE_API FSlateRoundedBoxBrush
 	 * @param InRadius      Corner Radius in Slate Units applied to the outline at each corner. X = Top Left, Y = Top Right, Z = Bottom Right, W = Bottom Left
 	 */
 	template<typename FillColorType, typename RadiusType>
-	FORCENOINLINE FSlateRoundedBoxBrush(const FillColorType& InColor, RadiusType InRadius, const FVector2D& InImageSize = FVector2D::ZeroVector)
+	FORCENOINLINE FSlateRoundedBoxBrush(const FillColorType& InColor, RadiusType InRadius, const UE::Slate::FDeprecateVector2DParameter& InImageSize = FVector2f::ZeroVector)
 		: FSlateBrush(ESlateBrushDrawType::RoundedBox, 
 					  NAME_None, 
 					  FMargin(0.0f), 
@@ -69,7 +86,7 @@ struct SLATECORE_API FSlateRoundedBoxBrush
 
 
 	template<typename FillColorType, typename RadiusType>
-	FORCENOINLINE FSlateRoundedBoxBrush(const FName InFillResourceName, const FillColorType& FillColorTint, RadiusType InRadius, const FVector2D& InImageSize = FVector2D::ZeroVector, ESlateBrushTileType::Type InTileType = ESlateBrushTileType::NoTile)
+	FORCENOINLINE FSlateRoundedBoxBrush(const FName InFillResourceName, const FillColorType& FillColorTint, RadiusType InRadius, const UE::Slate::FDeprecateVector2DParameter& InImageSize = FVector2f::ZeroVector, ESlateBrushTileType::Type InTileType = ESlateBrushTileType::NoTile)
 		: FSlateBrush(ESlateBrushDrawType::RoundedBox,
 			InFillResourceName,
 			FMargin(0.0f),
@@ -93,7 +110,7 @@ struct SLATECORE_API FSlateRoundedBoxBrush
 	 * @param InOutlineWidth Outline Width or Thickness
 	 */
 	template<typename FillColorType, typename OutlineColorType>
-	FORCENOINLINE FSlateRoundedBoxBrush(const FillColorType& InColor, const OutlineColorType& InOutlineColor, float InOutlineWidth, const FVector2D& InImageSize = FVector2D::ZeroVector)
+	FORCENOINLINE FSlateRoundedBoxBrush(const FillColorType& InColor, const OutlineColorType& InOutlineColor, float InOutlineWidth, const UE::Slate::FDeprecateVector2DParameter& InImageSize = FVector2f::ZeroVector)
 		: FSlateBrush(ESlateBrushDrawType::RoundedBox, 
 					  NAME_None, 
 					  FMargin(0.0f), 
@@ -116,7 +133,7 @@ struct SLATECORE_API FSlateRoundedBoxBrush
 	 * @param InOutlineWidth Outline Width or Thickness
 	 */
 	template<typename FillColorType, typename OutlineColorType>
-	FORCENOINLINE FSlateRoundedBoxBrush(const FName InFillResourceName, const FillColorType& FillColorTint, const OutlineColorType& InOutlineColor, float InOutlineWidth, const FVector2D& InImageSize = FVector2D::ZeroVector, ESlateBrushTileType::Type InTileType = ESlateBrushTileType::NoTile)
+	FORCENOINLINE FSlateRoundedBoxBrush(const FName InFillResourceName, const FillColorType& FillColorTint, const OutlineColorType& InOutlineColor, float InOutlineWidth, const UE::Slate::FDeprecateVector2DParameter& InImageSize = FVector2f::ZeroVector, ESlateBrushTileType::Type InTileType = ESlateBrushTileType::NoTile)
 		: FSlateBrush(ESlateBrushDrawType::RoundedBox,
 			InFillResourceName,
 			FMargin(0.0f),
@@ -142,7 +159,7 @@ struct SLATECORE_API FSlateRoundedBoxBrush
 	 * @param InOutlineWidth Outline Width or Thickness
 	 */
 	template<typename FillColorType, typename OutlineColorType, typename RadiusType>
-	FORCENOINLINE FSlateRoundedBoxBrush(const FillColorType& InColor, RadiusType InRadius, const OutlineColorType& InOutlineColor, float InOutlineWidth, const FVector2D& InImageSize = FVector2D::ZeroVector)
+	FORCENOINLINE FSlateRoundedBoxBrush(const FillColorType& InColor, RadiusType InRadius, const OutlineColorType& InOutlineColor, float InOutlineWidth, const UE::Slate::FDeprecateVector2DParameter& InImageSize = FVector2f::ZeroVector)
 		: FSlateBrush(ESlateBrushDrawType::RoundedBox, 
 					  NAME_None, 
 					  FMargin(0.0f), 
@@ -166,7 +183,7 @@ struct SLATECORE_API FSlateRoundedBoxBrush
 	 * @param InOutlineWidth Outline Width or Thickness
 	 */
 	template<typename FillColorType, typename OutlineColorType, typename RadiusType>
-	FORCENOINLINE FSlateRoundedBoxBrush(const FName InFillResourceName, const FillColorType& FillColorTint, RadiusType InRadius, const OutlineColorType& InOutlineColor, float InOutlineWidth, const FVector2D& InImageSize = FVector2D::ZeroVector, ESlateBrushTileType::Type InTileType = ESlateBrushTileType::NoTile)
+	FORCENOINLINE FSlateRoundedBoxBrush(const FName InFillResourceName, const FillColorType& FillColorTint, RadiusType InRadius, const OutlineColorType& InOutlineColor, float InOutlineWidth, const UE::Slate::FDeprecateVector2DParameter& InImageSize = FVector2f::ZeroVector, ESlateBrushTileType::Type InTileType = ESlateBrushTileType::NoTile)
 		: FSlateBrush(ESlateBrushDrawType::RoundedBox,
 					  InFillResourceName,
 					  FMargin(0.0f),
diff --git a/Engine/Source/Runtime/SlateCore/Public/FastUpdate/SlateInvalidationRoot.h b/Engine/Source/Runtime/SlateCore/Public/FastUpdate/SlateInvalidationRoot.h
--- a/Engine/Source/Runtime/SlateCore/Public/FastUpdate/SlateInvalidationRoot.h
+++ b/Engine/Source/Runtime/SlateCore/Public/FastUpdate/SlateInvalidationRoot.h
@@ -40,7 +40,7 @@ struct FSlateInvalidationContext
 	} 
 
 	FSlateRect CullingRect;
-	FVector2D ViewOffset;
+	UE::Slate::FDeprecateVector2DResult ViewOffset;
 	const FPaintArgs* PaintArgs;
 	const FWidgetStyle& WidgetStyle;
 	FSlateWindowElementList* WindowElementList;
@@ -66,7 +66,7 @@ struct FSlateInvalidationResult
 	{}
 
 	/** The view offset to use with the draw buffer */
-	FVector2D ViewOffset;
+	UE::Slate::FDeprecateVector2DResult ViewOffset;
 	/** The max layer id painted or cached */
 	int32 MaxLayerIdPainted;
 	/** If we had to repaint any widget */
@@ -172,7 +172,7 @@ private:
 	void InvalidateWidget(FWidgetProxy& Proxy, EInvalidateWidgetReason InvalidateReason);
 
 	void BuildFastPathWidgetList(const TSharedRef<SWidget>& RootWidget);
-	void AdjustWidgetsDesktopGeometry(FVector2D WindowToDesktopTransform);
+	void AdjustWidgetsDesktopGeometry(UE::Slate::FDeprecateVector2DParameter WindowToDesktopTransform);
 
 	/** Update child order and slate attribute registration */
 	void ProcessPreUpdate();
@@ -205,7 +205,7 @@ private:
 	/** Widgets that will be updated. */
 	TArray<FSlateInvalidationWidgetHeapElement> FinalUpdateList;
 
-	FVector2D CachedViewOffset;
+	FVector2f CachedViewOffset;
 
 	FSlateCachedElementData* CachedElementData;
 
diff --git a/Engine/Source/Runtime/SlateCore/Public/Input/Events.h b/Engine/Source/Runtime/SlateCore/Public/Input/Events.h
--- a/Engine/Source/Runtime/SlateCore/Public/Input/Events.h
+++ b/Engine/Source/Runtime/SlateCore/Public/Input/Events.h
@@ -10,6 +10,7 @@
 #include "GenericPlatform/GenericApplication.h"
 #include "GenericPlatform/GenericPlatformInputDeviceMapper.h"
 #include "Layout/Geometry.h"
+#include "Types/SlateVector2.h"
 #include "Events.generated.h"
 
 class FWidgetPath;
@@ -133,17 +134,17 @@ public:
 struct FVirtualPointerPosition
 {
 	FVirtualPointerPosition()
-		: CurrentCursorPosition(FVector2D::ZeroVector)
-		, LastCursorPosition(FVector2D::ZeroVector)
+		: CurrentCursorPosition(FVector2f::ZeroVector)
+		, LastCursorPosition(FVector2f::ZeroVector)
 	{}
 
-	FVirtualPointerPosition(const FVector2D& InCurrentCursorPosition, const FVector2D& InLastCursorPosition)
+	FVirtualPointerPosition(const UE::Slate::FDeprecateVector2DParameter& InCurrentCursorPosition, const UE::Slate::FDeprecateVector2DParameter& InLastCursorPosition)
 		: CurrentCursorPosition(InCurrentCursorPosition)
 		, LastCursorPosition(InLastCursorPosition)
 	{}
 
-	FVector2D CurrentCursorPosition;
-	FVector2D LastCursorPosition;
+	UE::Slate::FDeprecateVector2DResult CurrentCursorPosition;
+	UE::Slate::FDeprecateVector2DResult LastCursorPosition;
 };
 
 /**
@@ -690,9 +691,9 @@ public:
 	 * UStruct Constructor.  Not meant for normal usage.
 	 */
 	FPointerEvent()
-		: ScreenSpacePosition(FVector2D(0, 0))
-		, LastScreenSpacePosition(FVector2D(0, 0))
-		, CursorDelta(FVector2D(0, 0))
+		: ScreenSpacePosition(FVector2f(0.f, 0.f))
+		, LastScreenSpacePosition(FVector2f(0.f, 0.f))
+		, CursorDelta(FVector2f(0.f, 0.f))
 		, PressedButtons(&FTouchKeySet::EmptySet)
 		, EffectingButton()
 		, PointerIndex(0)
@@ -709,8 +710,8 @@ public:
 	/** Events are immutable once constructed. */
 	FPointerEvent(
 		uint32 InPointerIndex,
-		const FVector2D& InScreenSpacePosition,
-		const FVector2D& InLastScreenSpacePosition,
+		const UE::Slate::FDeprecateVector2DParameter& InScreenSpacePosition,
+		const UE::Slate::FDeprecateVector2DParameter& InLastScreenSpacePosition,
 		const TSet<FKey>& InPressedButtons,
 		FKey InEffectingButton,
 		float InWheelDelta,
@@ -719,7 +720,7 @@ public:
 		: FInputEvent(InModifierKeys, 0, false)
 		, ScreenSpacePosition(InScreenSpacePosition)
 		, LastScreenSpacePosition(InLastScreenSpacePosition)
-		, CursorDelta(InScreenSpacePosition - InLastScreenSpacePosition)
+		, CursorDelta(FVector2f(InScreenSpacePosition) - FVector2f(InLastScreenSpacePosition))
 		, PressedButtons(&InPressedButtons)
 		, EffectingButton(InEffectingButton)
 		, PointerIndex(InPointerIndex)
@@ -736,8 +737,8 @@ public:
 	FPointerEvent(
 		uint32 InUserIndex,
 		uint32 InPointerIndex,
-		const FVector2D& InScreenSpacePosition,
-		const FVector2D& InLastScreenSpacePosition,
+		const UE::Slate::FDeprecateVector2DParameter& InScreenSpacePosition,
+		const UE::Slate::FDeprecateVector2DParameter& InLastScreenSpacePosition,
 		const TSet<FKey>& InPressedButtons,
 		FKey InEffectingButton,
 		float InWheelDelta,
@@ -746,7 +747,7 @@ public:
 		: FInputEvent(InModifierKeys, InUserIndex, false)
 		, ScreenSpacePosition(InScreenSpacePosition)
 		, LastScreenSpacePosition(InLastScreenSpacePosition)
-		, CursorDelta(InScreenSpacePosition - InLastScreenSpacePosition)
+		, CursorDelta(FVector2f(InScreenSpacePosition) - FVector2f(InLastScreenSpacePosition))
 		, PressedButtons(&InPressedButtons)
 		, EffectingButton(InEffectingButton)
 		, PointerIndex(InPointerIndex)
@@ -763,8 +764,8 @@ public:
 	FPointerEvent(
 		FInputDeviceId InDeviceId,
 		uint32 InPointerIndex,
-		const FVector2D& InScreenSpacePosition,
-		const FVector2D& InLastScreenSpacePosition,
+		const UE::Slate::FDeprecateVector2DParameter& InScreenSpacePosition,
+		const UE::Slate::FDeprecateVector2DParameter& InLastScreenSpacePosition,
 		const TSet<FKey>& InPressedButtons,
 		FKey InEffectingButton,
 		float InWheelDelta,
@@ -774,7 +775,7 @@ public:
 		: FInputEvent(InModifierKeys, InDeviceId, false)
 		, ScreenSpacePosition(InScreenSpacePosition)
 		, LastScreenSpacePosition(InLastScreenSpacePosition)
-		, CursorDelta(InScreenSpacePosition - InLastScreenSpacePosition)
+		, CursorDelta(FVector2f(InScreenSpacePosition) - FVector2f(InLastScreenSpacePosition))
 		, PressedButtons(&InPressedButtons)
 		, EffectingButton(InEffectingButton)
 		, PointerIndex(InPointerIndex)
@@ -796,9 +797,9 @@ public:
 	FPointerEvent(
 		uint32 InUserIndex,
 		uint32 InPointerIndex,
-		const FVector2D& InScreenSpacePosition,
-		const FVector2D& InLastScreenSpacePosition,
-		const FVector2D& InDelta,
+		const UE::Slate::FDeprecateVector2DParameter& InScreenSpacePosition,
+		const UE::Slate::FDeprecateVector2DParameter& InLastScreenSpacePosition,
+		const UE::Slate::FDeprecateVector2DParameter& InDelta,
 		const TSet<FKey>& InPressedButtons,
 		const FModifierKeysState& InModifierKeys
 	)
@@ -821,9 +822,9 @@ public:
 	/** A constructor for raw mouse events */
 	FPointerEvent(
 		uint32 InPointerIndex,
-		const FVector2D& InScreenSpacePosition,
-		const FVector2D& InLastScreenSpacePosition,
-		const FVector2D& InDelta,
+		const UE::Slate::FDeprecateVector2DParameter& InScreenSpacePosition,
+		const UE::Slate::FDeprecateVector2DParameter& InLastScreenSpacePosition,
+		const UE::Slate::FDeprecateVector2DParameter& InDelta,
 		const TSet<FKey>& InPressedButtons,
 		const FModifierKeysState& InModifierKeys
 	)
@@ -843,25 +844,11 @@ public:
 		, bIsTouchFirstMove(false)
 	{ }
 
-	/** A constructor for touch events */
-	UE_DEPRECATED(4.20, "FPointerEvent constructor now takes a Force parameter")
 	FPointerEvent(
 		uint32 InUserIndex,
 		uint32 InPointerIndex,
-		const FVector2D& InScreenSpacePosition,
-		const FVector2D& InLastScreenSpacePosition,
-		bool bPressLeftMouseButton,
-		const FModifierKeysState& InModifierKeys = FModifierKeysState(),
-		uint32 InTouchpadIndex=0
-	)
-	: FPointerEvent(InUserIndex, InPointerIndex, InScreenSpacePosition, InLastScreenSpacePosition, 1.0f, bPressLeftMouseButton, false, false, InModifierKeys, InTouchpadIndex)
-	{ }
-
-	FPointerEvent(
-		uint32 InUserIndex,
-		uint32 InPointerIndex,
-		const FVector2D& InScreenSpacePosition,
-		const FVector2D& InLastScreenSpacePosition,
+		const UE::Slate::FDeprecateVector2DParameter& InScreenSpacePosition,
+		const UE::Slate::FDeprecateVector2DParameter& InLastScreenSpacePosition,
 		float InForce,
 		bool bPressLeftMouseButton,
 		bool bInIsForceChanged = false,
@@ -872,7 +859,7 @@ public:
 	: FInputEvent(InModifierKeys, InUserIndex, false)
 		, ScreenSpacePosition(InScreenSpacePosition)
 		, LastScreenSpacePosition(InLastScreenSpacePosition)
-		, CursorDelta(InScreenSpacePosition - InLastScreenSpacePosition)
+		, CursorDelta(FVector2f(InScreenSpacePosition) - FVector2f(InLastScreenSpacePosition))
 		, PressedButtons(bPressLeftMouseButton ? &FTouchKeySet::StandardSet : &FTouchKeySet::EmptySet)
 		, EffectingButton(EKeys::LeftMouseButton)
 		, PointerIndex(InPointerIndex)
@@ -889,8 +876,8 @@ public:
 	FPointerEvent(
 		FInputDeviceId InDeviceId,
 		uint32 InPointerIndex,
-		const FVector2D& InScreenSpacePosition,
-		const FVector2D& InLastScreenSpacePosition,
+		const UE::Slate::FDeprecateVector2DParameter& InScreenSpacePosition,
+		const UE::Slate::FDeprecateVector2DParameter& InLastScreenSpacePosition,
 		float InForce,
 		bool bPressLeftMouseButton,
 		bool bInIsForceChanged = false,
@@ -902,7 +889,7 @@ public:
 	: FInputEvent(InModifierKeys, InDeviceId, false)
 		, ScreenSpacePosition(InScreenSpacePosition)
 		, LastScreenSpacePosition(InLastScreenSpacePosition)
-		, CursorDelta(InScreenSpacePosition - InLastScreenSpacePosition)
+		, CursorDelta(FVector2f(InScreenSpacePosition) - FVector2f(InLastScreenSpacePosition))
 		, PressedButtons(bPressLeftMouseButton ? &FTouchKeySet::StandardSet : &FTouchKeySet::EmptySet)
 		, EffectingButton(EKeys::LeftMouseButton)
 		, PointerIndex(InPointerIndex)
@@ -923,18 +910,18 @@ public:
 
 	/** A constructor for gesture events */
 	FPointerEvent(
-		const FVector2D& InScreenSpacePosition,
-		const FVector2D& InLastScreenSpacePosition,
+		const UE::Slate::FDeprecateVector2DParameter& InScreenSpacePosition,
+		const UE::Slate::FDeprecateVector2DParameter& InLastScreenSpacePosition,
 		const TSet<FKey>& InPressedButtons,
 		const FModifierKeysState& InModifierKeys,
 		EGestureEvent InGestureType,
-		const FVector2D& InGestureDelta,
+		const UE::Slate::FDeprecateVector2DParameter& InGestureDelta,
 		bool bInIsDirectionInvertedFromDevice
 	)
 		: FInputEvent(InModifierKeys, 0, false)
 		, ScreenSpacePosition(InScreenSpacePosition)
 		, LastScreenSpacePosition(InLastScreenSpacePosition)
-		, CursorDelta(LastScreenSpacePosition - ScreenSpacePosition)
+		, CursorDelta(FVector2f(LastScreenSpacePosition) - FVector2f(ScreenSpacePosition))
 		, PressedButtons(&InPressedButtons)
 		, PointerIndex(0)
 		, Force(1.0f)
@@ -949,8 +936,8 @@ public:
 	/** A constructor to alter cursor positions */
 	FPointerEvent(
 		const FPointerEvent& Other,
-		const FVector2D& InScreenSpacePosition,
-		const FVector2D& InLastScreenSpacePosition)
+		const UE::Slate::FDeprecateVector2DParameter& InScreenSpacePosition,
+		const UE::Slate::FDeprecateVector2DParameter& InLastScreenSpacePosition)
 	{
 		*this = Other;
 		ScreenSpacePosition = InScreenSpacePosition;
@@ -960,13 +947,13 @@ public:
 public:
 
 	/** Returns The position of the cursor in screen space */
-	const FVector2D& GetScreenSpacePosition() const { return ScreenSpacePosition; }
+	const UE::Slate::FDeprecateVector2DResult& GetScreenSpacePosition() const { return ScreenSpacePosition; }
 
 	/** Returns the position of the cursor in screen space last time we handled an input event */
-	const FVector2D& GetLastScreenSpacePosition() const { return LastScreenSpacePosition; }
+	const UE::Slate::FDeprecateVector2DResult& GetLastScreenSpacePosition() const { return LastScreenSpacePosition; }
 
 	/** Returns the distance the mouse traveled since the last event was handled. */
-	const FVector2D& GetCursorDelta() const { return CursorDelta; }
+	const UE::Slate::FDeprecateVector2DResult& GetCursorDelta() const { return CursorDelta; }
 
 	/** Mouse buttons that are currently pressed */
 	bool IsMouseButtonDown( FKey MouseButton ) const { return PressedButtons->Contains( MouseButton ); }
@@ -1002,7 +989,7 @@ public:
 	EGestureEvent GetGestureType() const { return GestureType; }
 
 	/** Returns the change in gesture value since the last gesture event of the same type. */
-	const FVector2D& GetGestureDelta() const { return WheelOrGestureDelta; }
+	const UE::Slate::FDeprecateVector2DResult& GetGestureDelta() const { return WheelOrGestureDelta; }
 
 	/** Is the gesture delta inverted */
 	bool IsDirectionInvertedFromDevice() const { return bIsDirectionInvertedFromDevice; }
@@ -1049,9 +1036,9 @@ public:
 
 private:
 
-	FVector2D ScreenSpacePosition;
-	FVector2D LastScreenSpacePosition;
-	FVector2D CursorDelta;
+	UE::Slate::FDeprecateVector2DResult ScreenSpacePosition;
+	UE::Slate::FDeprecateVector2DResult LastScreenSpacePosition;
+	UE::Slate::FDeprecateVector2DResult CursorDelta;
 	const TSet<FKey>* PressedButtons;
 	FKey EffectingButton;
 	uint32 PointerIndex;
@@ -1059,7 +1046,7 @@ private:
 	float Force;
 	bool bIsTouchEvent;
 	EGestureEvent GestureType;
-	FVector2D WheelOrGestureDelta;
+	UE::Slate::FDeprecateVector2DResult WheelOrGestureDelta;
 	bool bIsDirectionInvertedFromDevice;
 	bool bIsTouchForceChanged;
 	bool bIsTouchFirstMove;
diff --git a/Engine/Source/Runtime/SlateCore/Public/Input/HittestGrid.h b/Engine/Source/Runtime/SlateCore/Public/Input/HittestGrid.h
--- a/Engine/Source/Runtime/SlateCore/Public/Input/HittestGrid.h
+++ b/Engine/Source/Runtime/SlateCore/Public/Input/HittestGrid.h
@@ -21,9 +21,6 @@ public:
 
 	virtual void ArrangeCustomHitTestChildren( FArrangedChildren& ArrangedChildren ) const = 0;
 
-	UE_DEPRECATED(5.0, "TranslateMouseCoordinateForCustomHitTestChild that returns a shared ptr is deprecated.")
-	virtual TSharedPtr<struct FVirtualPointerPosition> TranslateMouseCoordinateForCustomHitTestChild( const TSharedRef<SWidget>& ChildWidget, const FGeometry& ViewportGeometry, const FVector2D& ScreenSpaceMouseCoordinate, const FVector2D& LastScreenSpaceMouseCoordinate ) const { return TSharedPtr<struct FVirtualPointerPosition>(); }
-
 	virtual TOptional<FVirtualPointerPosition> TranslateMouseCoordinateForCustomHitTestChild(const SWidget& ChildWidget, const FGeometry& MyGeometry, const FVector2D ScreenSpaceMouseCoordinate, const FVector2D LastScreenSpaceMouseCoordinate) const = 0;
 };
 
@@ -37,7 +34,7 @@ public:
 	 * Given a Slate Units coordinate in virtual desktop space, perform a hittest
 	 * and return the path along which the corresponding event would be bubbled.
 	 */
-	TArray<FWidgetAndPointer> GetBubblePath(FVector2D DesktopSpaceCoordinate, float CursorRadius, bool bIgnoreEnabledStatus, int32 UserIndex = INDEX_NONE);
+	TArray<FWidgetAndPointer> GetBubblePath(UE::Slate::FDeprecateVector2DParameter DesktopSpaceCoordinate, float CursorRadius, bool bIgnoreEnabledStatus, int32 UserIndex = INDEX_NONE);
 
 	/**
 	 * Set the position and size of the hittest area in desktop coordinates
@@ -47,7 +44,7 @@ public:
 	 *
 	 * @return      Returns true if a clear of the hittest grid was required. 
 	 */
-	bool SetHittestArea(const FVector2D& HittestPositionInDesktop, const FVector2D& HittestDimensions, const FVector2D& HitestOffsetInWindow = FVector2D::ZeroVector);
+	bool SetHittestArea(const UE::Slate::FDeprecateVector2DParameter& HittestPositionInDesktop, const UE::Slate::FDeprecateVector2DParameter& HittestDimensions, const UE::Slate::FDeprecateVector2DParameter& HitestOffsetInWindow = FVector2f::ZeroVector);
 
 	/** Insert custom hit test data for a widget already in the grid */
 	UE_DEPRECATED(5.0, "Deprecated. Use the InsertCustomHitTestPath with a pointer.")
@@ -78,9 +75,9 @@ public:
 	 */
 	TSharedPtr<SWidget> FindNextFocusableWidget(const FArrangedWidget& StartingWidget, const EUINavigation Direction, const FNavigationReply& NavigationReply, const FArrangedWidget& RuleWidget, int32 UserIndex);
 
-	FVector2D GetGridSize() const { return GridSize; }
-	FVector2D GetGridOrigin() const { return GridOrigin; }
-	FVector2D GetGridWindowOrigin() const { return GridWindowOrigin; }
+	UE::Slate::FDeprecateVector2DResult GetGridSize() const { return GridSize; }
+	UE::Slate::FDeprecateVector2DResult GetGridOrigin() const { return GridOrigin; }
+	UE::Slate::FDeprecateVector2DResult GetGridWindowOrigin() const { return GridWindowOrigin; }
 
 	/** Clear the grid */
 	void Clear();
@@ -278,7 +275,7 @@ private:
 	TSharedPtr<SWidget> FindFocusableWidget(const FSlateRect WidgetRect, const FSlateRect SweptRect, int32 AxisIndex, int32 Increment, const EUINavigation Direction, const FNavigationReply& NavigationReply, TCompareFunc CompareFunc, TSourceSideFunc SourceSideFunc, TDestSideFunc DestSideFunc, int32 UserIndex, TArray<FDebuggingFindNextFocusableWidgetArgs::FWidgetResult>* IntermediatedResultPtr) const;
 
 	/** Constrains a float position into the grid coordinate. */
-	FIntPoint GetCellCoordinate(FVector2D Position) const;
+	FIntPoint GetCellCoordinate(UE::Slate::FDeprecateVector2DParameter Position) const;
 
 	/** Access a cell at coordinates X, Y. Coordinates are row and column indexes. */
 	FORCEINLINE_DEBUGGABLE FCell& CellAt(const int32 X, const int32 Y)
@@ -334,13 +331,13 @@ private:
 	FIntPoint NumCells;
 
 	/** Where the 0,0 of the upper-left-most cell corresponds to in desktop space. */
-	FVector2D GridOrigin;
+	FVector2f GridOrigin;
 
 	/** Where the 0,0 of the upper-left-most cell corresponds to in window space. */
-	FVector2D GridWindowOrigin;
+	FVector2f GridWindowOrigin;
 
 	/** The Size of the current grid. */
-	FVector2D GridSize;
+	FVector2f GridSize;
 
 	/** The current slate user index that should be associated with any added widgets */
 	int32 CurrentUserIndex;
diff --git a/Engine/Source/Runtime/SlateCore/Public/Layout/Clipping.h b/Engine/Source/Runtime/SlateCore/Public/Layout/Clipping.h
--- a/Engine/Source/Runtime/SlateCore/Public/Layout/Clipping.h
+++ b/Engine/Source/Runtime/SlateCore/Public/Layout/Clipping.h
@@ -70,7 +70,7 @@ public:
 	explicit FSlateClippingZone(const FSlateRect& AxisAlignedRect);
 	explicit FSlateClippingZone(const FGeometry& BoundingGeometry);
 	explicit FSlateClippingZone(const FPaintGeometry& PaintingGeometry);
-	FSlateClippingZone(const FVector2D& InTopLeft, const FVector2D& InTopRight, const FVector2D& InBottomLeft, const FVector2D& InBottomRight);
+	FSlateClippingZone(const UE::Slate::FDeprecateVector2DParameter& InTopLeft, const UE::Slate::FDeprecateVector2DParameter& InTopRight, const UE::Slate::FDeprecateVector2DParameter& InBottomLeft, const UE::Slate::FDeprecateVector2DParameter& InBottomRight);
 	FSlateClippingZone() {}
 
 	/**  */
@@ -119,7 +119,7 @@ public:
 	}
 
 	/** Is a point inside the clipping zone? */
-	bool IsPointInside(const FVector2D& Point) const;
+	bool IsPointInside(const UE::Slate::FDeprecateVector2DParameter& Point) const;
 
 	/**
 	 * Intersects two clipping zones and returns the new clipping zone that would need to be used.
@@ -151,9 +151,9 @@ public:
 		);
 	}
 
-	FSlateClippingZone ConvertRelativeToAbsolute(const FVector2D& WindowOffset) const
+	FSlateClippingZone ConvertRelativeToAbsolute(const UE::Slate::FDeprecateVector2DParameter& WindowOffset) const
 	{
-		FSlateClippingZone Absolute(FVector2D(TopLeft) + WindowOffset, FVector2D(TopRight) + WindowOffset, FVector2D(BottomLeft) + WindowOffset, FVector2D(BottomRight) + WindowOffset);
+		FSlateClippingZone Absolute(TopLeft + WindowOffset, TopRight + WindowOffset, BottomLeft + WindowOffset, BottomRight + WindowOffset);
 		Absolute.bIsAxisAligned = bIsAxisAligned;
 		Absolute.bIntersect = bIntersect;
 		Absolute.bAlwaysClip = bAlwaysClip;
@@ -161,7 +161,7 @@ public:
 		return Absolute;
 	}
 private:
-	void InitializeFromArbitraryPoints(const FVector2D& InTopLeft, const FVector2D& InTopRight, const FVector2D& InBottomLeft, const FVector2D& InBottomRight);
+	void InitializeFromArbitraryPoints(const UE::Slate::FDeprecateVector2DParameter& InTopLeft, const UE::Slate::FDeprecateVector2DParameter& InTopRight, const UE::Slate::FDeprecateVector2DParameter& InBottomLeft, const UE::Slate::FDeprecateVector2DParameter& InBottomRight);
 
 private:
 	/** Is the clipping zone axis aligned?  Axis aligned clipping zones are much cheaper. */
@@ -211,7 +211,7 @@ public:
 	FSlateClippingState(const FSlateClippingState& Other);
 
 	/** Is a point inside the clipping state? */
-	bool IsPointInside(const FVector2D& Point) const;
+	bool IsPointInside(const UE::Slate::FDeprecateVector2DParameter& Point) const;
 
 #if WITH_SLATE_DEBUGGING
 	/** Set the state index that this clipping state originated from.  We just do this for debugging purposes. */
diff --git a/Engine/Source/Runtime/SlateCore/Public/Layout/Geometry.h b/Engine/Source/Runtime/SlateCore/Public/Layout/Geometry.h
--- a/Engine/Source/Runtime/SlateCore/Public/Layout/Geometry.h
+++ b/Engine/Source/Runtime/SlateCore/Public/Layout/Geometry.h
@@ -17,6 +17,7 @@
 #include "Rendering/SlateRenderTransform.h"
 #include "Templates/SharedPointer.h"
 #include "UObject/ObjectMacros.h"
+#include "Types/SlateVector2.h"
 
 #include "Geometry.generated.h"
 
@@ -61,7 +62,7 @@ public:
 	 * @param InSize                   The size of this geometry.
 	 * @param InScale                  The scale of this geometry with respect to Normal Slate Coordinates.
 	 */
-	FGeometry( const FVector2D& OffsetFromParent, const FVector2D& ParentAbsolutePosition, const FVector2D& InLocalSize, float InScale )
+	FGeometry( const UE::Slate::FDeprecateVector2DParameter& OffsetFromParent, const UE::Slate::FDeprecateVector2DParameter& ParentAbsolutePosition, const UE::Slate::FDeprecateVector2DParameter& InLocalSize, float InScale )
 		: Size(InLocalSize)
 		, Scale(1.0f)
 		, AbsolutePosition(0.0f, 0.0f)
@@ -69,9 +70,9 @@ public:
 	{
 		// Since OffsetFromParent is given as a LocalSpaceOffset, we MUST convert this offset into the space of the parent to construct a valid layout transform.
 		// The extra TransformPoint below does this by converting the local offset to an offset in parent space.
-		FVector2D LayoutOffset = TransformPoint(InScale, OffsetFromParent);
+		FVector2f LayoutOffset = TransformPoint(InScale, UE::Slate::CastToVector2f(OffsetFromParent));
 
-		FSlateLayoutTransform ParentAccumulatedLayoutTransform(InScale, ParentAbsolutePosition);
+		FSlateLayoutTransform ParentAccumulatedLayoutTransform(InScale, UE::Slate::CastToVector2f(ParentAbsolutePosition));
 		FSlateLayoutTransform LocalLayoutTransform(LayoutOffset);
 		FSlateLayoutTransform AccumulatedLayoutTransform = Concatenate(LocalLayoutTransform, ParentAccumulatedLayoutTransform);
 		AccumulatedRenderTransform = TransformCast<FSlateRenderTransform>(AccumulatedLayoutTransform);
@@ -93,10 +94,10 @@ private:
 	 * @param ParentAccumulatedRenderTransform	The accumulated render transform of the parent widget. AccumulatedRenderTransform = Concat(LocalRenderTransform, LocalLayoutTransform, ParentAccumulatedRenderTransform).
 	 */
 	FGeometry( 
-		const FVector2D& InLocalSize, 
+		const UE::Slate::FDeprecateVector2DParameter& InLocalSize, 
 		const FSlateLayoutTransform& InLocalLayoutTransform, 
 		const FSlateRenderTransform& InLocalRenderTransform, 
-		const FVector2D& InLocalRenderTransformPivot, 
+		const UE::Slate::FDeprecateVector2DParameter& InLocalRenderTransformPivot, 
 		const FSlateLayoutTransform& ParentAccumulatedLayoutTransform, 
 		const FSlateRenderTransform& ParentAccumulatedRenderTransform)
 		: Size(InLocalSize)
@@ -105,11 +106,11 @@ private:
 		, AccumulatedRenderTransform(
 			Concatenate(
 				// convert the pivot to local space and make it the origin
-				Inverse(TransformPoint(FScale2D(InLocalSize), InLocalRenderTransformPivot)),
+				Inverse(TransformPoint(FScale2D(UE::Slate::CastToVector2f(InLocalSize)), UE::Slate::CastToVector2f(InLocalRenderTransformPivot))),
 				// apply the render transform in local space centered around the pivot
 				InLocalRenderTransform,
 				// translate the pivot point back.
-				TransformPoint(FScale2D(InLocalSize), InLocalRenderTransformPivot),
+				TransformPoint(FScale2D(UE::Slate::CastToVector2f(InLocalSize)), UE::Slate::CastToVector2f(InLocalRenderTransformPivot)),
 				// apply the layout transform next.
 				InLocalLayoutTransform,
 				// finally apply the parent accumulated transform, which takes us to the root.
@@ -133,12 +134,12 @@ private:
 	 * @param ParentAccumulatedRenderTransform	The accumulated render transform of the parent widget. AccumulatedRenderTransform = Concat(LocalRenderTransform, LocalLayoutTransform, ParentAccumulatedRenderTransform).
 	 */
 	FGeometry(
-		const FVector2D& InLocalSize,
+		const UE::Slate::FDeprecateVector2DParameter& InLocalSize,
 		const FSlateLayoutTransform& InLocalLayoutTransform,
 		const FSlateLayoutTransform& ParentAccumulatedLayoutTransform,
 		const FSlateRenderTransform& ParentAccumulatedRenderTransform,
 		bool bParentHasRenderTransform)
-		: Size(InLocalSize)
+		: Size(UE::Slate::CastToVector2f(InLocalSize))
 		, Scale(1.0f)
 		, AbsolutePosition(0.0f, 0.0f)
 		, AccumulatedRenderTransform(Concatenate(InLocalLayoutTransform, ParentAccumulatedRenderTransform))
@@ -188,11 +189,10 @@ public:
 	 * @param LayoutTransform	Layout transform of the geometry.
 	 * @return					The new root geometry
 	 */
-	FORCEINLINE_DEBUGGABLE static FGeometry MakeRoot(const FVector2D& InLocalSize, const FSlateLayoutTransform& LayoutTransform)
+	FORCEINLINE_DEBUGGABLE static FGeometry MakeRoot(const UE::Slate::FDeprecateVector2DParameter& InLocalSize, const FSlateLayoutTransform& LayoutTransform)
 	{
 		return FGeometry(InLocalSize, LayoutTransform, FSlateLayoutTransform(), FSlateRenderTransform(), false);
 	}
-
 	/**
 	 * Makes a new geometry that is essentially the root of a hierarchy (has no parent transforms to inherit).
 	 * For a root Widget, the LayoutTransform is often the window DPI scale + window offset.
@@ -201,11 +201,10 @@ public:
 	 * @param LayoutTransform	Layout transform of the geometry.
 	 * @return					The new root geometry
 	 */
-	FORCEINLINE_DEBUGGABLE static FGeometry MakeRoot(const FVector2D& InLocalSize, const FSlateLayoutTransform& LayoutTransform, const FSlateRenderTransform& RenderTransform)
+	FORCEINLINE_DEBUGGABLE static FGeometry MakeRoot(const UE::Slate::FDeprecateVector2DParameter& InLocalSize, const FSlateLayoutTransform& LayoutTransform, const FSlateRenderTransform& RenderTransform)
 	{
 		return FGeometry(InLocalSize, LayoutTransform, FSlateLayoutTransform(), FSlateRenderTransform(), !RenderTransform.IsIdentity());
 	}
-
 	/**
 	 * Create a child geometry relative to this one with a given local space size, layout transform, and render transform.
 	 * For example, a widget with a 5x5 margin will create a geometry for it's child contents having a LayoutTransform of Translate(5,5) and a LocalSize 10 units smaller 
@@ -218,7 +217,7 @@ public:
 	 *
 	 * @return					The new child geometry.
 	 */
-	FORCEINLINE_DEBUGGABLE FGeometry MakeChild(const FVector2D& InLocalSize, const FSlateLayoutTransform& LayoutTransform, const FSlateRenderTransform& RenderTransform, const FVector2D& RenderTransformPivot) const
+	FORCEINLINE_DEBUGGABLE FGeometry MakeChild(const UE::Slate::FDeprecateVector2DParameter& InLocalSize, const FSlateLayoutTransform& LayoutTransform, const FSlateRenderTransform& RenderTransform, const UE::Slate::FDeprecateVector2DParameter& RenderTransformPivot) const
 	{
 		return FGeometry(InLocalSize, LayoutTransform, RenderTransform, RenderTransformPivot, GetAccumulatedLayoutTransform(), GetAccumulatedRenderTransform());
 	}
@@ -233,16 +232,16 @@ public:
 	 *
 	 * @return					The new child geometry.
 	 */
-	FORCEINLINE_DEBUGGABLE FGeometry MakeChild(const FVector2D& InLocalSize, const FSlateLayoutTransform& LayoutTransform) const
+	FORCEINLINE_DEBUGGABLE FGeometry MakeChild(const UE::Slate::FDeprecateVector2DParameter& InLocalSize, const FSlateLayoutTransform& LayoutTransform) const
 	{
 		return FGeometry(InLocalSize, LayoutTransform, GetAccumulatedLayoutTransform(), GetAccumulatedRenderTransform(), bHasRenderTransform);
 	}
-
-	FORCEINLINE_DEBUGGABLE FGeometry MakeChild(const FSlateRenderTransform& RenderTransform, const FVector2D& RenderTransformPivot = FVector2D(0.5f, 0.5)) const
+	FORCEINLINE_DEBUGGABLE FGeometry MakeChild(const FSlateRenderTransform& RenderTransform, const UE::Slate::FDeprecateVector2DParameter& RenderTransformPivot = FVector2f(0.5f, 0.5f)) const
 	{
 		return FGeometry(GetLocalSize(), FSlateLayoutTransform(), RenderTransform, RenderTransformPivot, GetAccumulatedLayoutTransform(), GetAccumulatedRenderTransform());
 	}
 
+
 	/**
 	 * Create a child geometry+widget relative to this one using the given LayoutGeometry.
 	 *
@@ -265,7 +264,9 @@ public:
 	 *
 	 * @return					The new child geometry+widget.
 	 */
-	FArrangedWidget MakeChild(const TSharedRef<SWidget>& ChildWidget, const FVector2D& InLocalSize, const FSlateLayoutTransform& LayoutTransform) const;
+	FArrangedWidget MakeChild(const TSharedRef<SWidget>& ChildWidget, const UE::Slate::FDeprecateVector2DParameter& InLocalSize, const FSlateLayoutTransform& LayoutTransform) const;
+
+#if UE_ENABLE_SLATE_VECTOR_DEPRECATION_MECHANISMS
 
 	/**
 	 * !!! DEPRECATED FUNCTION !!! Use MakeChild taking a layout transform instead!
@@ -277,13 +278,16 @@ public:
 	 *
 	 * @return				The new child geometry.
 	 */
+	UE_DEPRECATED(5.2, "Use FGeometry MakeChild(const FVector2f& InLocalSize, const FSlateLayoutTransform& InLayoutTransform) instead.")
 	FORCEINLINE_DEBUGGABLE FGeometry MakeChild(const FVector2D& ChildOffset, const FVector2D& InLocalSize, float ChildScale = 1.0f) const
 	{
 		// Since ChildOffset is given as a LocalSpaceOffset, we MUST convert this offset into the space of the parent to construct a valid layout transform.
 		// The extra TransformPoint below does this by converting the local offset to an offset in parent space.
-		return FGeometry(InLocalSize, FSlateLayoutTransform(ChildScale, TransformPoint(ChildScale, ChildOffset)), GetAccumulatedLayoutTransform(), GetAccumulatedRenderTransform(), bHasRenderTransform);
+		return FGeometry(UE::Slate::CastToVector2f(InLocalSize), FSlateLayoutTransform(ChildScale, TransformPoint(ChildScale, UE::Slate::CastToVector2f(ChildOffset))), GetAccumulatedLayoutTransform(), GetAccumulatedRenderTransform(), bHasRenderTransform);
 	}
 
+#endif
+
 	/**
 	 * !!! DEPRECATED FUNCTION !!! Use MakeChild taking a layout transform instead!
 	 * Create a child geometry+widget relative to this one with a given local space size and layout transform given by a scale+offset.
@@ -296,7 +300,7 @@ public:
 	 *
 	 * @return				The new child geometry+widget.
 	 */
-	FArrangedWidget MakeChild(const TSharedRef<SWidget>& ChildWidget, const FVector2D& ChildOffset, const FVector2D& InLocalSize, float ChildScale = 1.0f) const;
+	FArrangedWidget MakeChild(const TSharedRef<SWidget>& ChildWidget, const UE::Slate::FDeprecateVector2DParameter& ChildOffset, const UE::Slate::FDeprecateVector2DParameter& InLocalSize, float ChildScale = 1.0f) const;
 
 	/**
 	 * Create a paint geometry that represents this geometry.
@@ -305,7 +309,7 @@ public:
 	 */
 	FORCEINLINE_DEBUGGABLE FPaintGeometry ToPaintGeometry() const
 	{
-		return FPaintGeometry(GetAccumulatedLayoutTransform(), GetAccumulatedRenderTransform(), Size, bHasRenderTransform);
+		return FPaintGeometry(GetAccumulatedLayoutTransform(), GetAccumulatedRenderTransform(), FVector2f(Size), bHasRenderTransform);
 	}
 
 	/**
@@ -317,10 +321,10 @@ public:
 	 *
 	 * @return					The new paint geometry derived from this one.
 	 */
-	FORCEINLINE_DEBUGGABLE FPaintGeometry ToPaintGeometry(const FVector2D& InLocalSize, const FSlateLayoutTransform& InLayoutTransform) const
+	FORCEINLINE_DEBUGGABLE FPaintGeometry ToPaintGeometry(const UE::Slate::FDeprecateVector2DParameter& InLocalSize, const FSlateLayoutTransform& InLayoutTransform) const
 	{
 		FSlateLayoutTransform NewAccumulatedLayoutTransform = Concatenate(InLayoutTransform, GetAccumulatedLayoutTransform());
-		return FPaintGeometry(NewAccumulatedLayoutTransform, Concatenate(InLayoutTransform, GetAccumulatedRenderTransform()), InLocalSize, bHasRenderTransform);
+		return FPaintGeometry(NewAccumulatedLayoutTransform, Concatenate(InLayoutTransform, GetAccumulatedRenderTransform()), UE::Slate::CastToVector2f(InLocalSize), bHasRenderTransform);
 	}
 
 	/**
@@ -332,7 +336,7 @@ public:
 	 *
 	 * @return					The new paint geometry derived from this one.
 	 */
-	FORCEINLINE_DEBUGGABLE FPaintGeometry ToPaintGeometry(const FVector2D& InLocalSize, const FSlateLayoutTransform& InLayoutTransform, const FSlateRenderTransform& RenderTransform, const FVector2D& RenderTransformPivot = FVector2D(0.5f, 0.5)) const
+	FORCEINLINE_DEBUGGABLE FPaintGeometry ToPaintGeometry(const UE::Slate::FDeprecateVector2DParameter& InLocalSize, const FSlateLayoutTransform& InLayoutTransform, const FSlateRenderTransform& RenderTransform, const UE::Slate::FDeprecateVector2DParameter& RenderTransformPivot = FVector2f(0.5f, 0.5f)) const
 	{
 		return MakeChild(InLocalSize, InLayoutTransform, RenderTransform, RenderTransformPivot).ToPaintGeometry();
 	}
@@ -347,9 +351,11 @@ public:
 	 */
 	FORCEINLINE_DEBUGGABLE FPaintGeometry ToPaintGeometry(const FSlateLayoutTransform& LayoutTransform) const
 	{
-		return ToPaintGeometry(Size, LayoutTransform);
+		return ToPaintGeometry(FVector2f(Size), LayoutTransform);
 	}
 
+#if UE_ENABLE_SLATE_VECTOR_DEPRECATION_MECHANISMS
+
 	/**
 	 * !!! DEPRECATED FUNCTION !!! Use ToPaintGeometry taking a layout transform instead!
 	 * Create a paint geometry relative to this one with a given local space size and layout transform given as a scale+offset.
@@ -361,13 +367,16 @@ public:
 	 * 
 	 * @return				The new paint geometry derived from this one.
 	 */
-	FORCEINLINE_DEBUGGABLE FPaintGeometry ToPaintGeometry(const FVector2D& InLocalOffset, const FVector2D& InLocalSize, float InLocalScale = 1.0f) const
+	UE_DEPRECATED(5.2, "Use FPaintGeometry ToPaintGeometry(const FVector2f& InLocalSize, const FSlateLayoutTransform& InLayoutTransform) instead.")
+	FORCEINLINE_DEBUGGABLE FPaintGeometry ToPaintGeometry(const UE::Slate::FDeprecateVector2DParameter& InLocalOffset, const UE::Slate::FDeprecateVector2DParameter& InLocalSize, float InLocalScale = 1.0f) const
 	{
 		// Since ChildOffset is given as a LocalSpaceOffset, we MUST convert this offset into the space of the parent to construct a valid layout transform.
 		// The extra TransformPoint below does this by converting the local offset to an offset in parent space.
-		return ToPaintGeometry(InLocalSize, FSlateLayoutTransform(InLocalScale, TransformPoint(InLocalScale, InLocalOffset)));
+		return ToPaintGeometry(InLocalSize, FSlateLayoutTransform(InLocalScale, TransformPoint(InLocalScale, UE::Slate::CastToVector2f(InLocalOffset))));
 	}
 
+#endif
+
 	/**
 	 * !!! DEPRECATED FUNCTION !!! Use ToPaintGeometry taking a layout transform instead!
 	 * Create a paint geometry relative to this one with the same size size and a given local space offset.
@@ -377,9 +386,9 @@ public:
 	 * 
 	 * @return				The new paint geometry derived from this one.
 	 */
-	FORCEINLINE_DEBUGGABLE FPaintGeometry ToOffsetPaintGeometry(const FVector2D& LocalOffset) const
+	FORCEINLINE_DEBUGGABLE FPaintGeometry ToOffsetPaintGeometry(const UE::Slate::FDeprecateVector2DParameter& LocalOffset) const
 	{
-		return ToPaintGeometry(FSlateLayoutTransform(LocalOffset));
+		return ToPaintGeometry(FSlateLayoutTransform(UE::Slate::CastToVector2f(LocalOffset)));
 	}
 
 	/**
@@ -390,12 +399,13 @@ public:
 	 * 
 	 * @return				The new paint geometry derived from this one.
 	 */
-	FORCEINLINE_DEBUGGABLE FPaintGeometry ToInflatedPaintGeometry(const FVector2D& InflateAmount) const
+	FORCEINLINE_DEBUGGABLE FPaintGeometry ToInflatedPaintGeometry(const UE::Slate::FDeprecateVector2DParameter& InflateAmount) const
 	{
+		FVector2f InflateAmount2f = UE::Slate::CastToVector2f(InflateAmount);
 		// This essentially adds (or subtracts) a border around the widget. We scale the size then offset by the border amount.
 		// Note this is not scaling child widgets, so the scale is not changing.
-		FVector2D NewSize = Size + InflateAmount * 2;
-		return ToPaintGeometry(NewSize, FSlateLayoutTransform(-InflateAmount));
+		FVector2f NewSize = FVector2f(Size) + InflateAmount2f* 2;
+		return ToPaintGeometry(NewSize, FSlateLayoutTransform(-InflateAmount2f));
 	}
 
 	/** 
@@ -403,10 +413,10 @@ public:
 	 * 
 	 * @return true if the provided location in absolute coordinates is within the bounds of this geometry. 
 	 */
-	FORCEINLINE_DEBUGGABLE bool IsUnderLocation(const FVector2D& AbsoluteCoordinate) const
+	FORCEINLINE_DEBUGGABLE bool IsUnderLocation(const UE::Slate::FDeprecateVector2DParameter& AbsoluteCoordinate) const
 	{
 		// this render transform invert is a little expensive. We might consider caching it?
-		FSlateRotatedRect Rect = TransformRect(GetAccumulatedRenderTransform(), FSlateRotatedRect(FSlateRect(FVector2D(0.0f, 0.0f), Size)));
+		FSlateRotatedRect Rect = TransformRect(GetAccumulatedRenderTransform(), FSlateRotatedRect(FSlateRect(FVector2f(0.0f, 0.0f), FVector2f(Size))));
 		return Rect.IsUnderLocation(AbsoluteCoordinate);
 	}
 
@@ -415,10 +425,10 @@ public:
 	 * 
 	 * @return Transforms AbsoluteCoordinate into the local space of this Geometry. 
 	 */
-	FORCEINLINE_DEBUGGABLE FVector2D AbsoluteToLocal(FVector2D AbsoluteCoordinate) const
+	FORCEINLINE_DEBUGGABLE UE::Slate::FDeprecateVector2DResult AbsoluteToLocal(UE::Slate::FDeprecateVector2DParameter AbsoluteCoordinate) const
 	{
 		// this render transform invert is a little expensive. We might consider caching it.
-		return TransformPoint(Inverse(GetAccumulatedRenderTransform()), AbsoluteCoordinate);
+		return UE::Slate::FDeprecateVector2DResult(TransformPoint(Inverse(GetAccumulatedRenderTransform()), UE::Slate::CastToVector2f(AbsoluteCoordinate)));
 	}
 
 	/**
@@ -428,9 +438,9 @@ public:
 	 * 
 	 * @return  Absolute coordinates
 	 */
-	FORCEINLINE_DEBUGGABLE FVector2D LocalToAbsolute(FVector2D LocalCoordinate) const
+	FORCEINLINE_DEBUGGABLE UE::Slate::FDeprecateVector2DResult LocalToAbsolute(UE::Slate::FDeprecateVector2DParameter LocalCoordinate) const
 	{
-		return TransformPoint(GetAccumulatedRenderTransform(), LocalCoordinate);
+		return UE::Slate::FDeprecateVector2DResult(TransformPoint(GetAccumulatedRenderTransform(), UE::Slate::CastToVector2f(LocalCoordinate)));
 	}
 
 	/**
@@ -439,22 +449,22 @@ public:
 	 * and not have the window start on a half pixel, which can cause the contents to jitter in relation to each other as the tooltip 
 	 * or popup moves around.
 	 */
-	FORCEINLINE_DEBUGGABLE FVector2D LocalToRoundedLocal(FVector2D LocalCoordinate) const
+	FORCEINLINE_DEBUGGABLE UE::Slate::FDeprecateVector2DResult LocalToRoundedLocal(UE::Slate::FDeprecateVector2DParameter LocalCoordinate) const
 	{
-		const FVector2D AbsoluteCoordinate = LocalToAbsolute(LocalCoordinate);
-		const FVector2D AbsoluteCoordinateRounded = FVector2D(FMath::RoundToFloat(AbsoluteCoordinate.X), FMath::RoundToFloat(AbsoluteCoordinate.Y));
+		const FVector2f AbsoluteCoordinate = LocalToAbsolute(UE::Slate::CastToVector2f(LocalCoordinate));
+		const FVector2f AbsoluteCoordinateRounded = FVector2f(FMath::RoundToFloat(AbsoluteCoordinate.X), FMath::RoundToFloat(AbsoluteCoordinate.Y));
 
 		return AbsoluteToLocal(AbsoluteCoordinateRounded);
 	}
 	
 	FORCEINLINE_DEBUGGABLE FSlateRect GetLayoutBoundingRect() const
 	{
-		return GetLayoutBoundingRect(FSlateRect(FVector2D(0.0f, 0.0f), Size));
+		return GetLayoutBoundingRect(FSlateRect(FVector2f(0.0f, 0.0f), FVector2f(Size)));
 	}
 
 	FORCEINLINE_DEBUGGABLE FSlateRect GetLayoutBoundingRect(const FMargin& LocalSpaceExtendBy) const
 	{
-		return GetLayoutBoundingRect(FSlateRect(FVector2D::ZeroVector, Size).ExtendBy(LocalSpaceExtendBy));
+		return GetLayoutBoundingRect(FSlateRect(FVector2f::ZeroVector, FVector2f(Size)).ExtendBy(LocalSpaceExtendBy));
 	}
 
 	FORCEINLINE_DEBUGGABLE FSlateRect GetLayoutBoundingRect(const FSlateRect& LocalSpaceRect) const
@@ -464,12 +474,12 @@ public:
 
 	FORCEINLINE_DEBUGGABLE FSlateRect GetRenderBoundingRect() const
 	{
-		return GetRenderBoundingRect(FSlateRect(FVector2D(0.0f, 0.0f), Size));
+		return GetRenderBoundingRect(FSlateRect(FVector2f(0.0f, 0.0f), FVector2f(Size)));
 	}
 
 	FORCEINLINE_DEBUGGABLE FSlateRect GetRenderBoundingRect(const FMargin& LocalSpaceExtendBy) const
 	{
-		return GetRenderBoundingRect(FSlateRect(FVector2D::ZeroVector, Size).ExtendBy(LocalSpaceExtendBy));
+		return GetRenderBoundingRect(FSlateRect(FVector2f::ZeroVector, FVector2f(Size)).ExtendBy(LocalSpaceExtendBy));
 	}
 
 	FORCEINLINE_DEBUGGABLE FSlateRect GetRenderBoundingRect(const FSlateRect& LocalSpaceRect) const
@@ -486,13 +496,13 @@ public:
 	 * Absolute coordinates could be either desktop or window space depending on what space the root of the widget hierarchy is in.
 	 *
 	 * @return the size of the geometry in absolute space */
-	FORCEINLINE_DEBUGGABLE FVector2D GetDrawSize() const
+	FORCEINLINE_DEBUGGABLE UE::Slate::FDeprecateVector2DResult GetDrawSize() const
 	{
-		return TransformVector(GetAccumulatedLayoutTransform(), Size);
+		return UE::Slate::FDeprecateVector2DResult(TransformVector(GetAccumulatedLayoutTransform(), FVector2f(Size)));
 	}
 
 	/** @return the size of the geometry in local space. */
-	FORCEINLINE const FVector2D& GetLocalSize() const { return Size; }
+	FORCEINLINE UE::Slate::FDeprecateVector2DResult GetLocalSize() const { return Size; }
 
 	/** @return the accumulated render transform. Shouldn't be needed in general. */
 	FORCEINLINE const FSlateRenderTransform& GetAccumulatedRenderTransform() const { return AccumulatedRenderTransform; }
@@ -500,7 +510,7 @@ public:
 	/** @return the accumulated layout transform. Shouldn't be needed in general. */
 	FORCEINLINE FSlateLayoutTransform GetAccumulatedLayoutTransform() const
 	{
-		return FSlateLayoutTransform(Scale, FVector2D(AbsolutePosition));
+		return FSlateLayoutTransform(Scale, AbsolutePosition);
 	}
 
 	/**
@@ -522,17 +532,17 @@ public:
 	/**
 	 * Get the absolute position in render space.
 	 */
-	FORCEINLINE FVector2D GetAbsolutePosition() const
+	FORCEINLINE UE::Slate::FDeprecateVector2DResult GetAbsolutePosition() const
 	{
-		return AccumulatedRenderTransform.TransformPoint(FVector2D::ZeroVector);
+		return UE::Slate::FDeprecateVector2DResult(AccumulatedRenderTransform.TransformPoint(FVector2f::ZeroVector));
 	}
 
 	/**
 	 * Get the absolute size of the geometry in render space.
 	 */
-	FORCEINLINE FVector2D GetAbsoluteSize() const
+	FORCEINLINE UE::Slate::FDeprecateVector2DResult GetAbsoluteSize() const
 	{
-		return AccumulatedRenderTransform.TransformVector(GetLocalSize());
+		return UE::Slate::FDeprecateVector2DResult(AccumulatedRenderTransform.TransformVector(FVector2f(GetLocalSize())));
 	}
 
 	/**
@@ -540,11 +550,11 @@ public:
 	 *   (0,0) - upper left
 	 *   (1,1) - bottom right
 	 *
-	 * Example: Say you wanted to know the center of the widget in absolute space, GetAbsolutePositionAtCoordinates(FVector2D(0.5f, 0.5f));
+	 * Example: Say you wanted to know the center of the widget in absolute space, GetAbsolutePositionAtCoordinates(FVector2f(0.5f, 0.5f));
 	 */
-	FORCEINLINE FVector2D GetAbsolutePositionAtCoordinates(const FVector2D& NormalCoordinates) const
+	FORCEINLINE UE::Slate::FDeprecateVector2DResult GetAbsolutePositionAtCoordinates(const UE::Slate::FDeprecateVector2DParameter& NormalCoordinates) const
 	{
-		return AccumulatedRenderTransform.TransformPoint(NormalCoordinates * GetLocalSize());
+		return UE::Slate::FDeprecateVector2DResult(AccumulatedRenderTransform.TransformPoint(FVector2f(NormalCoordinates) * FVector2f(GetLocalSize())));
 	}
 
 	/**
@@ -552,11 +562,11 @@ public:
 	 *   (0,0) - upper left
 	 *   (1,1) - bottom right
 	 *
-	 * Example: Say you wanted to know the center of the widget in local space, GetLocalPositionAtCoordinates(FVector2D(0.5f, 0.5f));
+	 * Example: Say you wanted to know the center of the widget in local space, GetLocalPositionAtCoordinates(FVector2f(0.5f, 0.5f));
 	 */
-	FORCEINLINE FVector2D GetLocalPositionAtCoordinates(const FVector2D& NormalCoordinates) const
+	FORCEINLINE UE::Slate::FDeprecateVector2DResult GetLocalPositionAtCoordinates(const UE::Slate::FDeprecateVector2DParameter& NormalCoordinates) const
 	{
-		return FVector2D(Position) + (NormalCoordinates * GetLocalSize());
+		return UE::Slate::FDeprecateVector2DResult(Position + (FVector2f(NormalCoordinates) * GetLocalSize()));
 	}
 
 	bool HasRenderTransform() const { return bHasRenderTransform; }
@@ -572,7 +582,7 @@ public:
 	 * 
 	 * Size of the geometry in local space.
 	 */
-	const FVector2D /*Local*/Size;
+	const FDeprecateSlateVector2D /*Local*/Size;
 
 	/** 
 	 * !!! DEPRECATED !!! These legacy public members should ideally not be referenced, as they do not account for the render transform.
diff --git a/Engine/Source/Runtime/SlateCore/Public/Layout/LayoutGeometry.h b/Engine/Source/Runtime/SlateCore/Public/Layout/LayoutGeometry.h
--- a/Engine/Source/Runtime/SlateCore/Public/Layout/LayoutGeometry.h
+++ b/Engine/Source/Runtime/SlateCore/Public/Layout/LayoutGeometry.h
@@ -5,6 +5,7 @@
 #include "CoreMinimal.h"
 #include "Layout/SlateRect.h"
 #include "Rendering/SlateLayoutTransform.h"
+#include "Types/SlateVector2.h"
 
 class FLayoutGeometry
 {
@@ -14,7 +15,7 @@ public:
 	{
 	}
 
-	explicit FLayoutGeometry(const FSlateLayoutTransform& InLocalToParent, const FVector2D& SizeInLocalSpace)
+	explicit FLayoutGeometry(const FSlateLayoutTransform& InLocalToParent, const UE::Slate::FDeprecateVector2DParameter& SizeInLocalSpace)
 		: LocalToParent(InLocalToParent)
 		, LocalSize(FVector2f(SizeInLocalSpace))
 	{
@@ -25,24 +26,24 @@ public:
 		return LocalToParent; 
 	}
 
-	FVector2D GetSizeInLocalSpace() const
+	UE::Slate::FDeprecateVector2DResult GetSizeInLocalSpace() const
 	{
-		return FVector2D(LocalSize);
+		return UE::Slate::FDeprecateVector2DResult(LocalSize);
 	}
 
-	FVector2D GetSizeInParentSpace() const
+	UE::Slate::FDeprecateVector2DResult GetSizeInParentSpace() const
 	{
-		return TransformVector(LocalToParent, FVector2D(LocalSize));
+		return UE::Slate::FDeprecateVector2DResult(TransformVector(LocalToParent, LocalSize));
 	}
 
-	FVector2D GetOffsetInParentSpace() const
+	UE::Slate::FDeprecateVector2DResult GetOffsetInParentSpace() const
 	{
-		return LocalToParent.GetTranslation();
+		return UE::Slate::FDeprecateVector2DResult(LocalToParent.GetTranslation());
 	}
 
 	FSlateRect GetRectInLocalSpace() const
 	{
-		return FSlateRect(FVector2D(0.0f, 0.0f), FVector2D(LocalSize));
+		return FSlateRect(FVector2f(0.0f, 0.0f), LocalSize);
 	}
 
 	FSlateRect GetRectInParentSpace() const
@@ -50,27 +51,6 @@ public:
 		return TransformRect(LocalToParent, GetRectInLocalSpace());
 	}
 
-public:
-	//FLayoutGeometry MakeChild(const FSlateLayoutTransform& ChildToThisTransform)
-	//{
-	//	return FLayoutGeometry(Concatenate(ChildToThisTransform, LocalToParent), LocalSize);
-	//}
-
-	//FLayoutGeometry MakeChild(const FSlateLayoutTransform& ChildToThisTransform, const FVector2D& SizeInChildSpace)
-	//{
-	//	return FLayoutGeometry(Concatenate(ChildToThisTransform, LocalToParent), SizeInChildSpace);
-	//}
-
-	//FLayoutGeometry MakeOffsetChild(const FVector2D& ChildOffset, const FVector2D& SizeInChildSpace)
-	//{
-	//	return FLayoutGeometry(Concatenate(ChildOffset, LocalToParent), SizeInChildSpace);
-	//}
-
-	//FLayoutGeometry MakeInflatedChild(const FVector2D& InflateAmount)
-	//{
-	//	return FLayoutGeometry(Concatenate(Inverse(InflateAmount), LocalToParent), LocalSize + 2.0f*InflateAmount);
-	//}
-
 private:
 	FSlateLayoutTransform LocalToParent;
 	FVector2f LocalSize;
diff --git a/Engine/Source/Runtime/SlateCore/Public/Layout/LayoutUtils.h b/Engine/Source/Runtime/SlateCore/Public/Layout/LayoutUtils.h
--- a/Engine/Source/Runtime/SlateCore/Public/Layout/LayoutUtils.h
+++ b/Engine/Source/Runtime/SlateCore/Public/Layout/LayoutUtils.h
@@ -212,8 +212,8 @@ static void ArrangeSingleChild(EFlowDirection InFlowDirection, const FGeometry&
 
 		ArrangedChildren.AddWidget(ChildVisibility, AllottedGeometry.MakeChild(
 			ChildSlot.GetWidget(),
-			FVector2D(XResult.Offset, YResult.Offset),
-			FVector2D(XResult.Size, YResult.Size)
+			FVector2f(XResult.Size, YResult.Size),
+			FSlateLayoutTransform(FVector2f(XResult.Offset, YResult.Offset))
 		));
 	}
 }
@@ -244,7 +244,7 @@ static void ArrangeChildrenInStack(EFlowDirection InLayoutFlow, const TPanelChil
 				// All widgets contribute their margin to the fixed space requirement
 				FixedTotal += CurChild.GetPadding().template GetTotalSpaceAlong<Orientation>();
 
-				FVector2D ChildDesiredSize = CurChild.GetWidget()->GetDesiredSize();
+				FVector2f ChildDesiredSize = CurChild.GetWidget()->GetDesiredSize();
 
 				// Auto-sized children contribute their desired size to the fixed space requirement
 				const float ChildSize = (Orientation == Orient_Vertical)
@@ -305,7 +305,7 @@ static void ArrangeChildrenInStack(EFlowDirection InLayoutFlow, const TPanelChil
 				}
 				else
 				{
-					const FVector2D ChildDesiredSize = CurChild.GetWidget()->GetDesiredSize();
+					const FVector2f ChildDesiredSize = CurChild.GetWidget()->GetDesiredSize();
 
 					// Auto-sized widgets get their desired-size value
 					ChildSize = (Orientation == Orient_Vertical)
@@ -323,19 +323,19 @@ static void ArrangeChildrenInStack(EFlowDirection InLayoutFlow, const TPanelChil
 
 			const FMargin SlotPadding(LayoutPaddingWithFlow(InLayoutFlow, CurChild.GetPadding()));
 
-			FVector2D SlotSize = (Orientation == Orient_Vertical)
-				? FVector2D(AllottedGeometry.GetLocalSize().X, ChildSize + SlotPadding.template GetTotalSpaceAlong<Orient_Vertical>())
-				: FVector2D(ChildSize + SlotPadding.template GetTotalSpaceAlong<Orient_Horizontal>(), AllottedGeometry.GetLocalSize().Y);
+			FVector2f SlotSize = (Orientation == Orient_Vertical)
+				? FVector2f(AllottedGeometry.GetLocalSize().X, ChildSize + SlotPadding.template GetTotalSpaceAlong<Orient_Vertical>())
+				: FVector2f(ChildSize + SlotPadding.template GetTotalSpaceAlong<Orient_Horizontal>(), AllottedGeometry.GetLocalSize().Y);
 
 			// Figure out the size and local position of the child within the slot			
 			AlignmentArrangeResult XAlignmentResult = AlignChild<Orient_Horizontal>(InLayoutFlow, SlotSize.X, CurChild, SlotPadding);
 			AlignmentArrangeResult YAlignmentResult = AlignChild<Orient_Vertical>(SlotSize.Y, CurChild, SlotPadding);
 
-			const FVector2D LocalPosition = (Orientation == Orient_Vertical)
-				? FVector2D(XAlignmentResult.Offset, PositionSoFar + YAlignmentResult.Offset + InOffset)
-				: FVector2D(PositionSoFar + XAlignmentResult.Offset + InOffset, YAlignmentResult.Offset);
+			const FVector2f LocalPosition = (Orientation == Orient_Vertical)
+				? FVector2f(XAlignmentResult.Offset, PositionSoFar + YAlignmentResult.Offset + InOffset)
+				: FVector2f(PositionSoFar + XAlignmentResult.Offset + InOffset, YAlignmentResult.Offset);
 
-			const FVector2D LocalSize = FVector2D(XAlignmentResult.Size, YAlignmentResult.Size);
+			const FVector2f LocalSize = FVector2f(XAlignmentResult.Size, YAlignmentResult.Size);
 
 			// Add the information about this child to the output list (ArrangedChildren)
 			ArrangedChildren.AddWidget(ChildVisibility, AllottedGeometry.MakeChild(
@@ -378,4 +378,4 @@ static FMargin LayoutPaddingWithFlow(EFlowDirection InLayoutFlow, const FMargin&
 *
 * @return A best position within the RectToFit such that none of the popup clips outside of the RectToFit.
 */
-SLATECORE_API FVector2D ComputePopupFitInRect(const FSlateRect& InAnchor, const FSlateRect& PopupRect, const EOrientation& Orientation, const FSlateRect& RectToFit);
+SLATECORE_API UE::Slate::FDeprecateVector2DResult ComputePopupFitInRect(const FSlateRect& InAnchor, const FSlateRect& PopupRect, const EOrientation& Orientation, const FSlateRect& RectToFit);
diff --git a/Engine/Source/Runtime/SlateCore/Public/Layout/Margin.h b/Engine/Source/Runtime/SlateCore/Public/Layout/Margin.h
--- a/Engine/Source/Runtime/SlateCore/Public/Layout/Margin.h
+++ b/Engine/Source/Runtime/SlateCore/Public/Layout/Margin.h
@@ -178,9 +178,9 @@ public:
 	 *
 	 * @return The top left offset.
 	 */
-	UE::Slate::FDeprecateVector2D GetTopLeft() const
+	UE::Slate::FDeprecateVector2DResult GetTopLeft() const
 	{
-		return UE::Slate::FDeprecateVector2D(Left, Top);
+		return UE::Slate::FDeprecateVector2DResult(Left, Top);
 	}
 	FVector2f GetTopLeft2f() const
 	{
@@ -192,9 +192,9 @@ public:
 	 *
 	 * @return Cumulative margin size.
 	 */
-	UE::Slate::FDeprecateVector2D GetDesiredSize() const
+	UE::Slate::FDeprecateVector2DResult GetDesiredSize() const
 	{
-		return UE::Slate::FDeprecateVector2D(GetDesiredSize2f());
+		return UE::Slate::FDeprecateVector2DResult(GetDesiredSize2f());
 	}
 	FVector2f GetDesiredSize2f() const
 	{
diff --git a/Engine/Source/Runtime/SlateCore/Public/Layout/PaintGeometry.h b/Engine/Source/Runtime/SlateCore/Public/Layout/PaintGeometry.h
--- a/Engine/Source/Runtime/SlateCore/Public/Layout/PaintGeometry.h
+++ b/Engine/Source/Runtime/SlateCore/Public/Layout/PaintGeometry.h
@@ -2,9 +2,10 @@
 
 #pragma once
 
-#include "CoreMinimal.h"
+#include "CoreTypes.h"
 #include "Rendering/SlateLayoutTransform.h"
 #include "Rendering/SlateRenderTransform.h"
+#include "Types/SlateVector2.h"
 
 /**
  * A Paint geometry contains the window-space (draw-space) info to draw an element on the screen.
@@ -40,7 +41,7 @@ struct SLATECORE_API FPaintGeometry
 	float DrawScale;
 
 	/** Get the Size of the geometry in local space. Must call CommitTransformsIfUsingLegacyConstructor() first if legacy ctor is used. */
-	const FVector2D& GetLocalSize() const { return LocalSize; }
+	UE::Slate::FDeprecateVector2DResult GetLocalSize() const { return UE::Slate::FDeprecateVector2DResult(LocalSize); }
 
 	/** Access the final render transform. Must call CommitTransformsIfUsingLegacyConstructor() first if legacy ctor is used. */
 	const FSlateRenderTransform& GetAccumulatedRenderTransform() const { return AccumulatedRenderTransform; }
@@ -54,9 +55,9 @@ struct SLATECORE_API FPaintGeometry
 	{ 
 		if (!bUsingLegacyConstructor) return;
 		
-		AccumulatedRenderTransform = FSlateRenderTransform(DrawScale, FVector2D(DrawPosition));
-		FSlateLayoutTransform AccumulatedLayoutTransform = FSlateLayoutTransform(DrawScale, FVector2D(DrawPosition));
-		LocalSize = TransformVector(Inverse(AccumulatedLayoutTransform), FVector2D(DrawSize));
+		AccumulatedRenderTransform = FSlateRenderTransform(DrawScale, DrawPosition);
+		FSlateLayoutTransform AccumulatedLayoutTransform = FSlateLayoutTransform(DrawScale, DrawPosition);
+		LocalSize = TransformVector(Inverse(AccumulatedLayoutTransform), DrawSize);
 	}
 
 	bool HasRenderTransform() const { return bHasRenderTransform; }
@@ -66,7 +67,7 @@ private:
 	mutable FVector2f DrawSize;
 
 	// Mutable to support legacy constructors.
-	mutable FVector2D LocalSize;
+	mutable FVector2f LocalSize;
 
 	// final render transform for drawing. Transforms from local space to window space for the draw element.
 	// Mutable to support legacy constructors.
@@ -96,7 +97,7 @@ public:
 	 * @param InAccumulatedLayoutTransform	The accumulated layout transform (from an FGeometry)
 	 * @param InAccumulatedRenderTransform	The accumulated render transform (from an FGeometry)
 	 */
-	FPaintGeometry( const FSlateLayoutTransform& InAccumulatedLayoutTransform, const FSlateRenderTransform& InAccumulatedRenderTransform, const FVector2D& InLocalSize, bool bInHasRenderTransform)
+	FPaintGeometry( const FSlateLayoutTransform& InAccumulatedLayoutTransform, const FSlateRenderTransform& InAccumulatedRenderTransform, const UE::Slate::FDeprecateVector2DParameter& InLocalSize, bool bInHasRenderTransform)
 		: DrawPosition(InAccumulatedLayoutTransform.GetTranslation())
 		, DrawScale(InAccumulatedLayoutTransform.GetScale())
 		, DrawSize(0.0f, 0.0f)
@@ -108,7 +109,7 @@ public:
 	}
 
 	// !!! DEPRECATED!!! This is legacy and should be removed!
-	FPaintGeometry( FVector2D InDrawPosition, FVector2D InDrawSize, float InDrawScale )
+	FPaintGeometry( UE::Slate::FDeprecateVector2DParameter InDrawPosition, UE::Slate::FDeprecateVector2DParameter InDrawSize, float InDrawScale )
 		: DrawPosition(InDrawPosition)
 		, DrawScale(InDrawScale)
 		, DrawSize(InDrawSize)
diff --git a/Engine/Source/Runtime/SlateCore/Public/Layout/SlateRect.h b/Engine/Source/Runtime/SlateCore/Public/Layout/SlateRect.h
--- a/Engine/Source/Runtime/SlateCore/Public/Layout/SlateRect.h
+++ b/Engine/Source/Runtime/SlateCore/Public/Layout/SlateRect.h
@@ -38,32 +38,21 @@ public:
 		, Bottom(InBottom)
 	{ }
 
-	FSlateRect(FVector2f InStartPos, FVector2f InEndPos)
+	FSlateRect( const UE::Slate::FDeprecateVector2DParameter& InStartPos, const UE::Slate::FDeprecateVector2DParameter& InEndPos )
 		: Left(InStartPos.X)
 		, Top(InStartPos.Y)
 		, Right(InEndPos.X)
 		, Bottom(InEndPos.Y)
 	{ }
 
-	FSlateRect( const FVector2d InStartPos, const FVector2d InEndPos )
-		: Left(UE_REAL_TO_FLOAT(InStartPos.X))
-		, Top(UE_REAL_TO_FLOAT(InStartPos.Y))
-		, Right(UE_REAL_TO_FLOAT(InEndPos.X))
-		, Bottom(UE_REAL_TO_FLOAT(InEndPos.Y))
-	{ }
-
 	/**
 	 * Creates a rect from a top left point and extent. Provided as a factory function to not conflict
 	 * with the TopLeft + BottomRight ctor.
 	 */
-	static FSlateRect FromPointAndExtent(const FVector2f TopLeft, const FVector2f Size)
+	static FSlateRect FromPointAndExtent(const UE::Slate::FDeprecateVector2DParameter& TopLeft, const UE::Slate::FDeprecateVector2DParameter& Size)
 	{
 		return FSlateRect(TopLeft, TopLeft + Size);
 	}
-	static FSlateRect FromPointAndExtent(const FVector2d TopLeft, const FVector2d Size)
-	{
-		return FromPointAndExtent(UE::Slate::CastToVector2f(TopLeft), UE::Slate::CastToVector2f(Size));
-	}
 	
 public:
 
@@ -88,9 +77,9 @@ public:
 	 *
 	 * @return The size as a vector.
 	 */
-	FORCEINLINE UE::Slate::FDeprecateVector2D GetSize() const
+	FORCEINLINE UE::Slate::FDeprecateVector2DResult GetSize() const
 	{
-		return UE::Slate::FDeprecateVector2D(GetSize2f());
+		return UE::Slate::FDeprecateVector2DResult(GetSize2f());
 	}
 	FORCEINLINE FVector2f GetSize2f() const
 	{
@@ -110,13 +99,13 @@ public:
 	 * 
 	 * @return The center point.
 	 */
-	FORCEINLINE UE::Slate::FDeprecateVector2D GetCenter() const
+	FORCEINLINE UE::Slate::FDeprecateVector2DResult GetCenter() const
 	{
-		return UE::Slate::FDeprecateVector2D(GetCenter2f());
+		return UE::Slate::FDeprecateVector2DResult(GetCenter2f());
 	}
 	FORCEINLINE FVector2f GetCenter2f() const
 	{
-		return FVector2f(Left, Top) + GetSize() * 0.5f;
+		return FVector2f(Left, Top) + GetSize2f() * 0.5f;
 	}
 
 	/**
@@ -124,9 +113,9 @@ public:
 	 * 
 	 * @return The top-left position.
 	 */
-	FORCEINLINE UE::Slate::FDeprecateVector2D GetTopLeft() const
+	FORCEINLINE UE::Slate::FDeprecateVector2DResult GetTopLeft() const
 	{
-		return UE::Slate::FDeprecateVector2D(GetTopLeft2f());
+		return UE::Slate::FDeprecateVector2DResult(GetTopLeft2f());
 	}
 	FORCEINLINE FVector2f GetTopLeft2f() const
 	{
@@ -138,9 +127,9 @@ public:
 	 *
 	 * @return The top-right position.
 	 */
-	FORCEINLINE UE::Slate::FDeprecateVector2D GetTopRight() const
+	FORCEINLINE UE::Slate::FDeprecateVector2DResult GetTopRight() const
 	{
-		return UE::Slate::FDeprecateVector2D(GetTopRight2f());
+		return UE::Slate::FDeprecateVector2DResult(GetTopRight2f());
 	}
 	FORCEINLINE FVector2f GetTopRight2f() const
 	{
@@ -152,9 +141,9 @@ public:
 	 * 
 	 * @return The bottom-right position.
 	 */
-	FORCEINLINE UE::Slate::FDeprecateVector2D GetBottomRight() const
+	FORCEINLINE UE::Slate::FDeprecateVector2DResult GetBottomRight() const
 	{
-		return UE::Slate::FDeprecateVector2D(GetBottomRight2f());
+		return UE::Slate::FDeprecateVector2DResult(GetBottomRight2f());
 	}
 	FORCEINLINE FVector2f GetBottomRight2f() const
 	{
@@ -166,9 +155,9 @@ public:
 	 * 
 	 * @return The bottom-left position.
 	 */
-	FORCEINLINE UE::Slate::FDeprecateVector2D GetBottomLeft() const
+	FORCEINLINE UE::Slate::FDeprecateVector2DResult GetBottomLeft() const
 	{
-		return UE::Slate::FDeprecateVector2D(GetBottomLeft2f());
+		return UE::Slate::FDeprecateVector2DResult(GetBottomLeft2f());
 	}
 	FORCEINLINE FVector2f GetBottomLeft2f() const
 	{
@@ -206,14 +195,10 @@ public:
 	 *
 	 * @return An offset rectangle.
 	 */
-	FORCEINLINE FSlateRect OffsetBy( const FVector2f OffsetAmount ) const
+	FORCEINLINE FSlateRect OffsetBy( const UE::Slate::FDeprecateVector2DParameter& OffsetAmount ) const
 	{
 		return FSlateRect(GetTopLeft2f() + OffsetAmount, GetBottomRight2f() + OffsetAmount);
 	}
-	FORCEINLINE FSlateRect OffsetBy( const FVector2d OffsetAmount ) const
-	{
-		return OffsetBy(UE::Slate::CastToVector2f(OffsetAmount));
-	}
 
 	/**
 	 * Return a rectangle that is scaled by the amount specified.
@@ -295,14 +280,10 @@ public:
 	 * @param Point	The point to check
 	 * @return True if the point is inside the rectangle
 	 */
-	FORCEINLINE bool ContainsPoint( const FVector2f Point ) const
+	FORCEINLINE bool ContainsPoint( const UE::Slate::FDeprecateVector2DParameter& Point ) const
 	{
 		return Point.X >= Left && Point.X <= Right && Point.Y >= Top && Point.Y <= Bottom;
 	}
-	FORCEINLINE bool ContainsPoint( const FVector2d Point ) const
-	{
-		return ContainsPoint(UE::Slate::CastToVector2f(Point));
-	}
 
 	bool operator==( const FSlateRect& Other ) const
 	{
diff --git a/Engine/Source/Runtime/SlateCore/Public/Layout/SlateRotatedRect.h b/Engine/Source/Runtime/SlateCore/Public/Layout/SlateRotatedRect.h
--- a/Engine/Source/Runtime/SlateCore/Public/Layout/SlateRotatedRect.h
+++ b/Engine/Source/Runtime/SlateCore/Public/Layout/SlateRotatedRect.h
@@ -31,21 +31,13 @@ public:
 	}
 
 	/** Per-element constructor. */
-	FSlateRotatedRect(const FVector2f InTopLeft, const FVector2f InExtentX, const FVector2f InExtentY)
+	FSlateRotatedRect(const UE::Slate::FDeprecateVector2DParameter& InTopLeft, const UE::Slate::FDeprecateVector2DParameter& InExtentX, const UE::Slate::FDeprecateVector2DParameter& InExtentY)
 		: TopLeft(InTopLeft)
 		, ExtentX(InExtentX)
 		, ExtentY(InExtentY)
 	{
 	}
 
-	/** Per-element constructor. */
-	FSlateRotatedRect(const FVector2d InTopLeft, const FVector2d InExtentX, const FVector2d InExtentY)
-		: TopLeft(UE::Slate::CastToVector2f(InTopLeft))
-		, ExtentX(UE::Slate::CastToVector2f(InExtentX))
-		, ExtentY(UE::Slate::CastToVector2f(InExtentY))
-	{
-	}
-
 public:
 
 	/** transformed Top-left corner. */
@@ -70,8 +62,7 @@ public:
 	FSlateRect ToBoundingRect() const;
 
 	/** Point-in-rect test. */
-	bool IsUnderLocation(const FVector2f Location) const;
-	bool IsUnderLocation(const FVector2d Location) const;
+	bool IsUnderLocation(const UE::Slate::FDeprecateVector2DParameter Location) const;
 
 	static FSlateRotatedRect MakeRotatedRect(const FSlateRect& ClipRectInLayoutWindowSpace, const FSlateLayoutTransform& InverseLayoutTransform, const FSlateRenderTransform& RenderTransform)
 	{
diff --git a/Engine/Source/Runtime/SlateCore/Public/Rendering/DrawElementPayloads.h b/Engine/Source/Runtime/SlateCore/Public/Rendering/DrawElementPayloads.h
--- a/Engine/Source/Runtime/SlateCore/Public/Rendering/DrawElementPayloads.h
+++ b/Engine/Source/Runtime/SlateCore/Public/Rendering/DrawElementPayloads.h
@@ -22,17 +22,12 @@ struct FSlateGradientStop
 						  A two stop gradient should go from (0,0), to (Width,Height).
 	 * @param InColor	- The color to lerp towards at this stop.
 	 */
-	FSlateGradientStop(const FVector2d InPosition, const FLinearColor InColor)
+	template<typename VectorType>
+	FSlateGradientStop(const VectorType& InPosition, const FLinearColor InColor)
 		: Position(UE::Slate::CastToVector2f(InPosition))
 		, Color(InColor)
 	{
 	}
-	FSlateGradientStop(const FVector2f InPosition, const FLinearColor InColor)
-		: Position(InPosition)
-		, Color(InColor)
-	{
-
-	}
 };
 template <> struct TIsPODType<FSlateGradientStop> { enum { Value = true }; };
 
@@ -66,12 +61,7 @@ struct FSlateBoxPayload : public FSlateDataPayload, public FSlateTintableElement
 	ESlateBrushDrawType::Type GetBrushDrawType() const { return DrawType; }
 	const FSlateShaderResourceProxy* GetResourceProxy() const { return ResourceProxy; }
 
-	void SetBrush(const FSlateBrush* InBrush, FVector2D LocalSize, float DrawScale)
-	{
-		SetBrush(InBrush, UE::Slate::CastToVector2f(LocalSize), DrawScale);
-	}
-
-	void SetBrush(const FSlateBrush* InBrush, FVector2f LocalSize, float DrawScale)
+	void SetBrush(const FSlateBrush* InBrush, UE::Slate::FDeprecateVector2DParameter LocalSize, float DrawScale)
 	{
 		check(InBrush);
 		ensureMsgf(InBrush->GetDrawType() != ESlateBrushDrawType::NoDrawType, TEXT("This should have been filtered out earlier in the Make... call."));
@@ -82,7 +72,7 @@ struct FSlateBoxPayload : public FSlateDataPayload, public FSlateTintableElement
 		Tiling = InBrush->GetTiling();
 		Mirroring = InBrush->GetMirroring();
 		DrawType = InBrush->GetDrawType();
-		FSlateResourceHandle Handle = InBrush->GetRenderingResource(FVector2D(LocalSize), DrawScale);
+		FSlateResourceHandle Handle = InBrush->GetRenderingResource(LocalSize, DrawScale);
 		if (Handle.IsValid())
 		{
 			ResourceProxy = Handle.GetResourceProxy();
@@ -222,18 +212,7 @@ struct FSlateSplinePayload : public FSlateDataPayload, public FSlateTintableElem
 	void SetThickness(float InThickness) { Thickness = InThickness; }
 	float GetThickness() const { return Thickness; }
 
-	void SetCubicBezier(const FVector2d InP0, const FVector2d InP1, const FVector2d InP2, const FVector2d InP3, float InThickness, const FLinearColor InTint)
-	{
-		SetCubicBezier(
-			UE::Slate::CastToVector2f(InP0),
-			UE::Slate::CastToVector2f(InP1),
-			UE::Slate::CastToVector2f(InP2),
-			UE::Slate::CastToVector2f(InP3),
-			InThickness,
-			InTint);
-	}
-
-	void SetCubicBezier(const FVector2f InP0, const FVector2f InP1, const FVector2f InP2, const FVector2f InP3, float InThickness, const FLinearColor InTint)
+	void SetCubicBezier(const UE::Slate::FDeprecateVector2DParameter InP0, const UE::Slate::FDeprecateVector2DParameter InP1, const UE::Slate::FDeprecateVector2DParameter InP2, const UE::Slate::FDeprecateVector2DParameter InP3, float InThickness, const FLinearColor InTint)
 	{
 		Tint = InTint;
 		P0 = InP0;
@@ -243,18 +222,7 @@ struct FSlateSplinePayload : public FSlateDataPayload, public FSlateTintableElem
 		Thickness = InThickness;
 	}
 
-	void SetHermiteSpline(const FVector2d InStart, const FVector2d InStartDir, const FVector2d InEnd, const FVector2d InEndDir, float InThickness, const FLinearColor InTint)
-	{
-		SetHermiteSpline(
-			UE::Slate::CastToVector2f(InStart),
-			UE::Slate::CastToVector2f(InStartDir),
-			UE::Slate::CastToVector2f(InEnd),
-			UE::Slate::CastToVector2f(InEndDir),
-			InThickness,
-			InTint);
-	}
-
-	void SetHermiteSpline(const FVector2f InStart, const FVector2f InStartDir, const FVector2f InEnd, const FVector2f InEndDir, float InThickness, const FLinearColor InTint)
+	void SetHermiteSpline(const UE::Slate::FDeprecateVector2DParameter InStart, const UE::Slate::FDeprecateVector2DParameter InStartDir, const UE::Slate::FDeprecateVector2DParameter InEnd, const UE::Slate::FDeprecateVector2DParameter InEndDir, float InThickness, const FLinearColor InTint)
 	{
 		Tint = InTint;
 		P0 = InStart;
@@ -264,18 +232,7 @@ struct FSlateSplinePayload : public FSlateDataPayload, public FSlateTintableElem
 		Thickness = InThickness;
 	}
 
-	void SetGradientHermiteSpline(const FVector2d InStart, const FVector2d InStartDir, const FVector2d InEnd, const FVector2d InEndDir, float InThickness, TArray<FSlateGradientStop> InGradientStops)
-	{
-		SetGradientHermiteSpline(
-			UE::Slate::CastToVector2f(InStart),
-			UE::Slate::CastToVector2f(InStartDir),
-			UE::Slate::CastToVector2f(InEnd),
-			UE::Slate::CastToVector2f(InEndDir),
-			InThickness,
-			MoveTemp(InGradientStops));
-	}
-
-	void SetGradientHermiteSpline(const FVector2f InStart, const FVector2f InStartDir, const FVector2f InEnd, const FVector2f InEndDir, float InThickness, TArray<FSlateGradientStop> InGradientStops)
+	void SetGradientHermiteSpline(const UE::Slate::FDeprecateVector2DParameter InStart, const UE::Slate::FDeprecateVector2DParameter InStartDir, const UE::Slate::FDeprecateVector2DParameter InEnd, const UE::Slate::FDeprecateVector2DParameter InEndDir, float InThickness, TArray<FSlateGradientStop> InGradientStops)
 	{
 		P0 = InStart;
 		P1 = InStart + InStartDir / 3.0f;
@@ -305,6 +262,7 @@ struct FSlateLinePayload : public FSlateDataPayload, public FSlateTintableElemen
 		Thickness = InThickness;
 	}
 
+#if UE_ENABLE_SLATE_VECTOR_DEPRECATION_MECHANISMS
 	void SetLines(const TArray<FVector2D>& InPoints, bool bInAntialias, const TArray<FLinearColor>* InPointColors = nullptr)
 	{
 		TArray<FVector2f> NewPoints;
@@ -322,7 +280,7 @@ struct FSlateLinePayload : public FSlateDataPayload, public FSlateTintableElemen
 			SetLines(MoveTemp(NewPoints), bInAntialias);
 		}
 	}
-
+#endif
 
 	void SetLines(TArray<FVector2f> InPoints, bool bInAntialias)
 	{
@@ -386,12 +344,7 @@ struct FSlateCachedBufferPayload : public FSlateDataPayload
 	FVector2f CachedRenderDataOffset;
 
 	// Cached Buffers
-	void SetCachedBuffer(FSlateRenderDataHandle* InRenderDataHandle, const FVector2d Offset)
-	{
-		SetCachedBuffer(InRenderDataHandle, UE::Slate::CastToVector2f(Offset));
-	}
-
-	void SetCachedBuffer(FSlateRenderDataHandle* InRenderDataHandle, const FVector2f Offset)
+	void SetCachedBuffer(FSlateRenderDataHandle* InRenderDataHandle, const UE::Slate::FDeprecateVector2DParameter Offset)
 	{
 		check(InRenderDataHandle);
 
diff --git a/Engine/Source/Runtime/SlateCore/Public/Rendering/DrawElements.h b/Engine/Source/Runtime/SlateCore/Public/Rendering/DrawElements.h
--- a/Engine/Source/Runtime/SlateCore/Public/Rendering/DrawElements.h
+++ b/Engine/Source/Runtime/SlateCore/Public/Rendering/DrawElements.h
@@ -19,6 +19,7 @@
 #include "Rendering/SlateRenderBatch.h"
 #include "DrawElementTextOverflowArgs.h"
 #include "ElementBatcher.h"
+#include "Types/SlateVector2.h"
 
 #if UE_ENABLE_INCLUDE_ORDER_DEPRECATED_IN_5_2
 #include "Fonts/FontCache.h"
@@ -117,17 +118,6 @@ public:
 		ESlateDrawEffect InDrawEffects = ESlateDrawEffect::None,
 		const FLinearColor& InTint = FLinearColor::White );
 
-
-	UE_DEPRECATED(4.20, "Storing and passing in a FSlateResourceHandle to MakeBox is no longer necessary.")
-	SLATECORE_API static void MakeBox(
-		FSlateWindowElementList& ElementList,
-		uint32 InLayer, 
-		const FPaintGeometry& PaintGeometry, 
-		const FSlateBrush* InBrush, 
-		const FSlateResourceHandle& InRenderingHandle, 
-		ESlateDrawEffect InDrawEffects = ESlateDrawEffect::None, 
-		const FLinearColor& InTint = FLinearColor::White );
-
 	SLATECORE_API static void MakeRotatedBox(
 		FSlateWindowElementList& ElementList,
 		uint32 InLayer, 
@@ -135,21 +125,10 @@ public:
 		const FSlateBrush* InBrush, 
 		ESlateDrawEffect,
 		float Angle,
-		TOptional<FVector2d> InRotationPoint,
+		UE::Slate::FDeprecateOptionalVector2DParameter InRotationPoint = TOptional<FVector2f>(),
 		ERotationSpace RotationSpace = RelativeToElement,
 		const FLinearColor& InTint = FLinearColor::White );
 
-	SLATECORE_API static void MakeRotatedBox(
-		FSlateWindowElementList& ElementList,
-		uint32 InLayer,
-		const FPaintGeometry& PaintGeometry,
-		const FSlateBrush* InBrush,
-		ESlateDrawEffect InDrawEffects = ESlateDrawEffect::None,
-		float Angle = 0.0f,
-		TOptional<FVector2f> InRotationPoint = TOptional<FVector2f>(),
-		ERotationSpace RotationSpace = RelativeToElement,
-		const FLinearColor& InTint = FLinearColor::White);
-
 	/**
 	 * Creates a text element which displays a string of a rendered in a certain font on the screen
 	 *
@@ -210,8 +189,7 @@ public:
 	 * @param InDrawEffects         Optional draw effects to apply
 	 * @param InTint                Color to tint the element
 	 */
-	SLATECORE_API static void MakeSpline(FSlateWindowElementList& ElementList, uint32 InLayer, const FPaintGeometry& PaintGeometry, const FVector2d InStart, const FVector2d InStartDir, const FVector2d InEnd, const FVector2d InEndDir, float InThickness = 0.0f, ESlateDrawEffect InDrawEffects = ESlateDrawEffect::None, const FLinearColor& InTint=FLinearColor::White);
-	SLATECORE_API static void MakeSpline(FSlateWindowElementList& ElementList, uint32 InLayer, const FPaintGeometry& PaintGeometry, const FVector2f InStart, const FVector2f InStartDir, const FVector2f InEnd, const FVector2f InEndDir, float InThickness = 0.0f, ESlateDrawEffect InDrawEffects = ESlateDrawEffect::None, const FLinearColor& InTint=FLinearColor::White);
+	SLATECORE_API static void MakeSpline(FSlateWindowElementList& ElementList, uint32 InLayer, const FPaintGeometry& PaintGeometry, const UE::Slate::FDeprecateVector2DParameter InStart, const UE::Slate::FDeprecateVector2DParameter InStartDir, const UE::Slate::FDeprecateVector2DParameter InEnd, const UE::Slate::FDeprecateVector2DParameter InEndDir, float InThickness = 0.0f, ESlateDrawEffect InDrawEffects = ESlateDrawEffect::None, const FLinearColor& InTint=FLinearColor::White);
 
 	/**
 	 * Creates a Bezier Spline element
@@ -226,19 +204,10 @@ public:
 	 * @param InDrawEffects         Optional draw effects to apply
 	 * @param InTint                Color to tint the element
 	 */
-	SLATECORE_API static void MakeCubicBezierSpline(FSlateWindowElementList& ElementList, uint32 InLayer, const FPaintGeometry& PaintGeometry, const FVector2d P0, const FVector2d P1, const FVector2d P2, const FVector2d P3, float InThickness = 0.0f, ESlateDrawEffect InDrawEffects = ESlateDrawEffect::None, const FLinearColor& InTint = FLinearColor::White);
-	SLATECORE_API static void MakeCubicBezierSpline(FSlateWindowElementList& ElementList, uint32 InLayer, const FPaintGeometry& PaintGeometry, const FVector2f P0, const FVector2f P1, const FVector2f P2, const FVector2f P3, float InThickness = 0.0f, ESlateDrawEffect InDrawEffects = ESlateDrawEffect::None, const FLinearColor& InTint = FLinearColor::White);
-
-	/** Just like MakeSpline but in draw-space coordinates. This is useful for connecting already-transformed widgets together. */
-	SLATECORE_API static void MakeDrawSpaceSpline(FSlateWindowElementList& ElementList, uint32 InLayer, const FVector2d InStart, const FVector2d InStartDir, const FVector2d InEnd, const FVector2d InEndDir, float InThickness = 0.0f, ESlateDrawEffect InDrawEffects = ESlateDrawEffect::None, const FLinearColor& InTint=FLinearColor::White);
-	SLATECORE_API static void MakeDrawSpaceSpline(FSlateWindowElementList& ElementList, uint32 InLayer, const FVector2f InStart, const FVector2f InStartDir, const FVector2f InEnd, const FVector2f InEndDir, float InThickness = 0.0f, ESlateDrawEffect InDrawEffects = ESlateDrawEffect::None, const FLinearColor& InTint=FLinearColor::White);
+	SLATECORE_API static void MakeCubicBezierSpline(FSlateWindowElementList& ElementList, uint32 InLayer, const FPaintGeometry& PaintGeometry, const UE::Slate::FDeprecateVector2DParameter P0, const UE::Slate::FDeprecateVector2DParameter P1, const UE::Slate::FDeprecateVector2DParameter P2, const UE::Slate::FDeprecateVector2DParameter P3, float InThickness = 0.0f, ESlateDrawEffect InDrawEffects = ESlateDrawEffect::None, const FLinearColor& InTint = FLinearColor::White);
 
 	/** Just like MakeSpline but in draw-space coordinates. This is useful for connecting already-transformed widgets together. */
-	UE_DEPRECATED(4.20, "Splines with color gradients will not be supported in the future.")
-	SLATECORE_API static void MakeDrawSpaceGradientSpline( FSlateWindowElementList& ElementList, uint32 InLayer, const FVector2D& InStart, const FVector2D& InStartDir, const FVector2D& InEnd, const FVector2D& InEndDir, const TArray<FSlateGradientStop>& InGradientStops, float InThickness = 0.0f, ESlateDrawEffect InDrawEffects = ESlateDrawEffect::None );
-
-	UE_DEPRECATED(4.20, "Splines with color gradients will not be supported in the future.")
-	static void MakeDrawSpaceGradientSpline(FSlateWindowElementList& ElementList, uint32 InLayer, const FVector2D& InStart, const FVector2D& InStartDir, const FVector2D& InEnd, const FVector2D& InEndDir, const FSlateRect InClippingRect, const TArray<FSlateGradientStop>& InGradientStops, float InThickness = 0.0f, ESlateDrawEffect InDrawEffects = ESlateDrawEffect::None);
+	SLATECORE_API static void MakeDrawSpaceSpline(FSlateWindowElementList& ElementList, uint32 InLayer, const UE::Slate::FDeprecateVector2DParameter InStart, const UE::Slate::FDeprecateVector2DParameter InStartDir, const UE::Slate::FDeprecateVector2DParameter InEnd, const UE::Slate::FDeprecateVector2DParameter InEndDir, float InThickness = 0.0f, ESlateDrawEffect InDrawEffects = ESlateDrawEffect::None, const FLinearColor& InTint=FLinearColor::White);
 
 	/**
 	 * Creates a line defined by the provided points
@@ -252,7 +221,9 @@ public:
 	 * @param bAntialias               Should antialiasing be applied to the line?
 	 * @param Thickness                The thickness of the line
 	 */
+#if UE_ENABLE_SLATE_VECTOR_DEPRECATION_MECHANISMS
 	SLATECORE_API static void MakeLines(FSlateWindowElementList& ElementList, uint32 InLayer, const FPaintGeometry& PaintGeometry, const TArray<FVector2d>& Points, ESlateDrawEffect InDrawEffects = ESlateDrawEffect::None, const FLinearColor& InTint=FLinearColor::White, bool bAntialias = true, float Thickness = 1.0f);
+#endif
 	SLATECORE_API static void MakeLines(FSlateWindowElementList& ElementList, uint32 InLayer, const FPaintGeometry& PaintGeometry, TArray<FVector2f> Points, ESlateDrawEffect InDrawEffects = ESlateDrawEffect::None, const FLinearColor& InTint = FLinearColor::White, bool bAntialias = true, float Thickness = 1.0f);
 
 
@@ -269,7 +240,9 @@ public:
 	 * @param bAntialias               Should antialiasing be applied to the line?
 	 * @param Thickness                The thickness of the line
 	 */
+#if UE_ENABLE_SLATE_VECTOR_DEPRECATION_MECHANISMS
 	SLATECORE_API static void MakeLines(FSlateWindowElementList& ElementList, uint32 InLayer, const FPaintGeometry& PaintGeometry, const TArray<FVector2d>& Points, const TArray<FLinearColor>& PointColors, ESlateDrawEffect InDrawEffects = ESlateDrawEffect::None, const FLinearColor& InTint=FLinearColor::White, bool bAntialias = true, float Thickness = 1.0f);
+#endif
 	SLATECORE_API static void MakeLines(FSlateWindowElementList& ElementList, uint32 InLayer, const FPaintGeometry& PaintGeometry, TArray<FVector2f> Points, TArray<FLinearColor> PointColors, ESlateDrawEffect InDrawEffects = ESlateDrawEffect::None, const FLinearColor& InTint=FLinearColor::White, bool bAntialias = true, float Thickness = 1.0f);
 
 	/**
@@ -315,24 +288,19 @@ public:
 	FORCEINLINE const FSlateRenderTransform& GetRenderTransform() const { return RenderTransform; }
 	FORCEINLINE void SetRenderTransform(const FSlateRenderTransform& InRenderTransform) { RenderTransform = InRenderTransform; }
 
-	FORCEINLINE UE::Slate::FDeprecateVector2D GetPosition() const
+	FORCEINLINE UE::Slate::FDeprecateVector2DResult GetPosition() const
 	{
-		return UE::Slate::FDeprecateVector2D(Position);
+		return UE::Slate::FDeprecateVector2DResult(Position);
 	}
 
-	UE_DEPRECATED(5.1, "Slate rendering uses float instead of double. Use SetPosition(FVector2f)")
-	FORCEINLINE void SetPosition(FVector2d InPosition)
-	{
-		Position = UE::Slate::CastToVector2f(InPosition);
-	}
-	FORCEINLINE void SetPosition(FVector2f InPosition)
+	FORCEINLINE void SetPosition(UE::Slate::FDeprecateVector2DParameter InPosition)
 	{
 		Position = InPosition;
 	}
 
-	FORCEINLINE UE::Slate::FDeprecateVector2D GetLocalSize() const
+	FORCEINLINE UE::Slate::FDeprecateVector2DResult GetLocalSize() const
 	{
-		return UE::Slate::FDeprecateVector2D(LocalSize);
+		return UE::Slate::FDeprecateVector2DResult(LocalSize);
 	}
 
 	FORCEINLINE float GetScale() const { return Scale; }
@@ -357,16 +325,13 @@ public:
 
 	void AddReferencedObjects(FReferenceCollector& Collector);
 
-	UE_DEPRECATED(5.1, "Slate rendering uses float instead of double. Use ApplyPositionOffset(FVector2f)")
-	void ApplyPositionOffset(FVector2d InOffset);
-
 	/**
 	 * Update element cached position with an arbitrary offset
 	 *
 	 * @param Element		   Element to update
 	 * @param InOffset         Absolute translation delta
 	 */
-	void ApplyPositionOffset(FVector2f InOffset);
+	void ApplyPositionOffset(UE::Slate::FDeprecateVector2DParameter InOffset);
 
 	UE_DEPRECATED(4.23, "GetClippingIndex has been deprecated.  If you were using this please use GetPrecachedClippingIndex instead.")
 	FORCEINLINE const int32 GetClippingIndex() const { return GetPrecachedClippingIndex(); }
@@ -564,9 +529,9 @@ public:
 	}
 
 	/** @return Get the window size that we will be painting */
-	UE::Slate::FDeprecateVector2D GetWindowSize() const
+	UE::Slate::FDeprecateVector2DResult GetWindowSize() const
 	{
-		return UE::Slate::FDeprecateVector2D(WindowSize);
+		return UE::Slate::FDeprecateVector2DResult(WindowSize);
 	}
 
 	template<typename PayloadType>
diff --git a/Engine/Source/Runtime/SlateCore/Public/Rendering/ShaderResourceManager.h b/Engine/Source/Runtime/SlateCore/Public/Rendering/ShaderResourceManager.h
--- a/Engine/Source/Runtime/SlateCore/Public/Rendering/ShaderResourceManager.h
+++ b/Engine/Source/Runtime/SlateCore/Public/Rendering/ShaderResourceManager.h
@@ -53,11 +53,6 @@ public:
 	/** 
 	 * Returns a texture associated with the passed in name.  Should return nullptr if not found 
 	 */
-	UE_DEPRECATED(5.1, "Slate rendering uses float instead of double. Use GetShaderResource(FName,FVector2f, float)")
-	virtual FSlateShaderResourceProxy* GetShaderResource(const FSlateBrush& Brush, FVector2d LocalSize, float DrawScale)
-	{
-		return GetShaderResource(Brush, UE::Slate::CastToVector2f(LocalSize), DrawScale);
-	}
 	virtual FSlateShaderResourceProxy* GetShaderResource(const FSlateBrush& Brush, FVector2f LocalSize, float DrawScale) = 0;
 
 	/**
@@ -71,11 +66,6 @@ public:
 	 * @param	DrawScale	The scale of the draw element using this brush.  This scale is used to rescale vector graphics only
 	 * @return	The created resource handle.  
 	 */
-	UE_DEPRECATED(5.1, "Slate rendering uses float instead of double. Use GetResourceHandle(FName,FVector2f, float)")
-	virtual FSlateResourceHandle GetResourceHandle(const FSlateBrush& Brush, FVector2d LocalSize, float DrawScale)
-	{
-		return GetResourceHandle(Brush, UE::Slate::CastToVector2f(LocalSize), DrawScale);
-	}
 	virtual FSlateResourceHandle GetResourceHandle(const FSlateBrush& Brush, FVector2f LocalSize, float DrawScale);
 	virtual FSlateResourceHandle GetResourceHandle(const FSlateBrush& Brush);
 
diff --git a/Engine/Source/Runtime/SlateCore/Public/Rendering/SlateLayoutTransform.h b/Engine/Source/Runtime/SlateCore/Public/Rendering/SlateLayoutTransform.h
--- a/Engine/Source/Runtime/SlateCore/Public/Rendering/SlateLayoutTransform.h
+++ b/Engine/Source/Runtime/SlateCore/Public/Rendering/SlateLayoutTransform.h
@@ -4,6 +4,7 @@
 #include "CoreMinimal.h"
 #include "Math/TransformCalculus.h"
 #include "Math/TransformCalculus2D.h"
+#include "Types/SlateVector2.h"
 
 /**
  * Represents a 2D transformation in the following order: scale then translate.
@@ -25,25 +26,22 @@ public:
 	}
 
 	/** Ctor from a scale followed by translate. Shortcut to Concatenate(InScale, InTranslation). */
-	template<typename VType>
-	explicit FSlateLayoutTransform(float InScale, const UE::Math::TVector2<VType>& InTranslation)
-		:Scale(InScale)
-		,Translation(FVector2f(InTranslation))
+	explicit FSlateLayoutTransform(float InScale, const UE::Slate::FDeprecateVector2DParameter& InTranslation)
+		: Scale(InScale)
+		, Translation(FVector2f(InTranslation))
 	{
 	}
 
-	/** Ctor from a 2D translation followed by a scale. Shortcut to Concatenate(InTranslation, InScale). While this is the opposite order we internally store them, we can represent this correctly. */
-	template<typename VType>
-	explicit FSlateLayoutTransform(const UE::Math::TVector2<VType>& InTranslation)
+	explicit FSlateLayoutTransform(const UE::Slate::FDeprecateVector2DParameter& InTranslation)
 		:Scale(1.0f)
 		,Translation(FVector2f(InTranslation))
 	{
 	}
 
 	/** Access to the 2D translation */
-	FVector2D GetTranslation() const
+	UE::Slate::FDeprecateVector2DResult GetTranslation() const
 	{
-		return FVector2D(Translation);
+		return UE::Slate::FDeprecateVector2DResult(Translation);
 	}
 
 	/** Access to the scale. */
@@ -56,22 +54,18 @@ public:
 	FMatrix ToMatrix() const
 	{
 		FMatrix Matrix = FScaleMatrix(GetScale());
-		Matrix.SetOrigin(FVector(GetTranslation(), 0.0f));
+		Matrix.SetOrigin(FVector(FVector2d(Translation), 0.0f));
 		return Matrix;
 	}
 
-	/** 2D transform support. */
-	template<typename VType>
-	UE::Math::TVector2<VType> TransformPoint(const UE::Math::TVector2<VType>& Point) const
+	FORCEINLINE UE::Slate::FDeprecateVector2DResult TransformPoint(const UE::Slate::FDeprecateVector2DParameter& Point) const
 	{
-		return ::TransformPoint((UE::Math::TVector2<VType>)Translation, ::TransformPoint(Scale, Point));
+		return UE::Slate::FDeprecateVector2DResult(::TransformPoint(Translation, ::TransformPoint(Scale, FVector2f(Point))));
 	}
 
-	/** 2D transform support. */
-	template<typename VType>
-	UE::Math::TVector2<VType> TransformVector(const UE::Math::TVector2<VType>& Vector) const
+	FORCEINLINE UE::Slate::FDeprecateVector2DResult TransformVector(const UE::Slate::FDeprecateVector2DParameter& Vector) const
 	{
-		return ::TransformVector((UE::Math::TVector2<VType>)Translation, ::TransformVector(Scale, Vector));
+		return UE::Slate::FDeprecateVector2DResult(::TransformVector(Translation, ::TransformVector(Scale, FVector2f(Vector))));
 	}
 
 	/**
@@ -169,6 +163,6 @@ template<> template<> inline TTransform2<float> TransformConverter<TTransform2<f
 
 template<> template<> inline TTransform2<double> TransformConverter<TTransform2<double>>::Convert<FSlateLayoutTransform>(const FSlateLayoutTransform& Transform)
 {
-	return TTransform2<double>(TScale2<double>(Transform.GetScale()), Transform.GetTranslation());
+	return TTransform2<double>(TScale2<double>(Transform.GetScale()), UE::Math::TVector2<double>(UE::Slate::CastToVector2f(Transform.GetTranslation())));
 }
 
diff --git a/Engine/Source/Runtime/SlateCore/Public/Rendering/SlateRenderTransform.h b/Engine/Source/Runtime/SlateCore/Public/Rendering/SlateRenderTransform.h
--- a/Engine/Source/Runtime/SlateCore/Public/Rendering/SlateRenderTransform.h
+++ b/Engine/Source/Runtime/SlateCore/Public/Rendering/SlateRenderTransform.h
@@ -3,7 +3,21 @@
 
 #include "CoreMinimal.h"
 #include "Math/TransformCalculus2D.h"
+#include "Types/SlateVector2.h"
 
 /** Typecast so it's more clear to the reader when we are dealing with a render transform vs. a layout transform. */
 typedef FTransform2D FSlateRenderTransform;
 
+
+#if UE_ENABLE_SLATE_VECTOR_DEPRECATION_MECHANISMS
+
+inline UE::Slate::FDeprecateVector2DResult TransformVector(const FSlateRenderTransform& Transform, const UE::Slate::FDeprecateVector2DParameter& Vector)
+{
+	return UE::Slate::FDeprecateVector2DResult(TransformVector(Transform, FVector2f(Vector)));
+}
+inline UE::Slate::FDeprecateVector2DResult TransformVector(const FSlateRenderTransform& Transform, const UE::Slate::FDeprecateVector2DResult& Vector)
+{
+	return UE::Slate::FDeprecateVector2DResult(TransformVector(Transform, FVector2f(Vector)));
+}
+
+#endif
\ No newline at end of file
diff --git a/Engine/Source/Runtime/SlateCore/Public/Rendering/SlateRenderer.h b/Engine/Source/Runtime/SlateCore/Public/Rendering/SlateRenderer.h
--- a/Engine/Source/Runtime/SlateCore/Public/Rendering/SlateRenderer.h
+++ b/Engine/Source/Runtime/SlateCore/Public/Rendering/SlateRenderer.h
@@ -326,22 +326,6 @@ public:
 
 	virtual bool GenerateDynamicImageResource(FName ResourceName, FSlateTextureDataRef TextureData) { return false; }
 
-	/**
-	 * Creates a resource handle with a specific size and scale (for brushes with dynamic resizing such as vector art)
-	 * A handle is used as fast path for looking up a rendering resource for a given brush when adding Slate draw elements
-	 * This can be cached and stored safely in code.  It will become invalid when a resource is destroyed
-	 * It is expensive to create a resource so do not do it in time sensitive areas
-	 *
-	 * @param	Brush		The brush to get a rendering resource handle
-	 * @param	LocalSize	The local size of the element using the brush (Vector graphics only)
-	 * @param	DrawScale	The draw scale of the element using the brush (Vector graphics only)
-	 * @return	The created resource handle.
-	 */
-	UE_DEPRECATED(5.1, "Slate rendering uses float instead of double. Use GetResourceHandle(FName,FVector2f, float)")
-	virtual FSlateResourceHandle GetResourceHandle(const FSlateBrush& Brush, FVector2d LocalSize, float DrawScale)
-	{
-		return GetResourceHandle(Brush, UE::Slate::CastToVector2f(LocalSize), DrawScale);
-	}
 	virtual FSlateResourceHandle GetResourceHandle(const FSlateBrush& Brush, FVector2f LocalSize, float DrawScale) = 0;
 
 	/**
diff --git a/Engine/Source/Runtime/SlateCore/Public/Rendering/SlateVectorGraphicsCache.h b/Engine/Source/Runtime/SlateCore/Public/Rendering/SlateVectorGraphicsCache.h
--- a/Engine/Source/Runtime/SlateCore/Public/Rendering/SlateVectorGraphicsCache.h
+++ b/Engine/Source/Runtime/SlateCore/Public/Rendering/SlateVectorGraphicsCache.h
@@ -16,12 +16,6 @@ class FSlateVectorGraphicsCache : public FSlateFlushableAtlasCache
 public:
 	SLATECORE_API FSlateVectorGraphicsCache(TSharedPtr<ISlateTextureAtlasFactory> InAtlasFactory);
 
-	UE_DEPRECATED(5.1, "Slate rendering uses float instead of double. Use GetShaderResource(FName,FVector2f, float)")
-	FSlateShaderResourceProxy* GetShaderResource(const FSlateBrush& Brush, FVector2d LocalSize, float DrawScale)
-	{
-		return GetShaderResource(Brush, UE::Slate::CastToVector2f(LocalSize), DrawScale);
-	}
-
 	SLATECORE_API FSlateShaderResourceProxy* GetShaderResource(const FSlateBrush& Brush, FVector2f LocalSize, float DrawScale);
 
 	SLATECORE_API void UpdateCache();
@@ -47,12 +41,6 @@ private:
 		FName BrushName;
 		FIntPoint PixelSize;
 
-		UE_DEPRECATED(5.1, "Slate rendering uses float instead of double. Use FVectorCacheKey(FName,FVector2f, float)")
-		FVectorCacheKey(FName InBrushName, FVector2d LocalSize, float DrawScale)
-			: FVectorCacheKey(InBrushName, UE::Slate::CastToVector2f(LocalSize), DrawScale)
-		{
-		}
-
 		FVectorCacheKey(FName InBrushName, FVector2f LocalSize, float DrawScale)
 			: BrushName(InBrushName)
 			, PixelSize((LocalSize*DrawScale).IntPoint())
@@ -80,11 +68,6 @@ private:
 		FVectorCacheKey Key;
 		TArray<uint8> PixelData;
 
-		UE_DEPRECATED(5.1, "Slate rendering uses float instead of double. Use FRasterRequest(FName,FVector2f, float)")
-		FRasterRequest(FName InBrushName, FVector2d InLocalSize, float InDrawScale)
-			: FRasterRequest(InBrushName, UE::Slate::CastToVector2f(InLocalSize), InDrawScale)
-		{}
-
 		FRasterRequest(FName InBrushName, FVector2f InLocalSize, float InDrawScale)
 			: Key(InBrushName, InLocalSize, InDrawScale)
 		{}
