Slate: Initial phase of FVector2D deprecation

diff --git a/Engine/Source/Runtime/SlateCore/Public/Styling/CoreStyle.h b/Engine/Source/Runtime/SlateCore/Public/Styling/CoreStyle.h
--- a/Engine/Source/Runtime/SlateCore/Public/Styling/CoreStyle.h
+++ b/Engine/Source/Runtime/SlateCore/Public/Styling/CoreStyle.h
@@ -82,25 +82,25 @@ namespace CoreStyleConstants
 {
 	// Note, these sizes are in Slate Units.
 	// Slate Units do NOT have to map to pixels.
-	inline const FVector2D Icon5x16(5.0f, 16.0f);
-	inline const FVector2D Icon8x4(8.0f, 4.0f);
-	inline const FVector2D Icon16x4(16.0f, 4.0f);
-	inline const FVector2D Icon8x8(8.0f, 8.0f);
-	inline const FVector2D Icon10x10(10.0f, 10.0f);
-	inline const FVector2D Icon12x12(12.0f, 12.0f);
-	inline const FVector2D Icon12x16(12.0f, 16.0f);
-	inline const FVector2D Icon14x14(14.0f, 14.0f);
-	inline const FVector2D Icon16x16(16.0f, 16.0f);
-	inline const FVector2D Icon18x18(18.0f, 18.0f);
-	inline const FVector2D Icon20x20(20.0f, 20.0f);
-	inline const FVector2D Icon22x22(22.0f, 22.0f);
-	inline const FVector2D Icon24x24(24.0f, 24.0f);
-	inline const FVector2D Icon25x25(25.0f, 25.0f);
-	inline const FVector2D Icon32x32(32.0f, 32.0f);
-	inline const FVector2D Icon40x40(40.0f, 40.0f);
-	inline const FVector2D Icon64x64(64.0f, 64.0f);
-	inline const FVector2D Icon36x24(36.0f, 24.0f);
-	inline const FVector2D Icon128x128(128.0f, 128.0f);
+	inline const UE::Slate::FDeprecateVector2DResult Icon5x16(5.0f, 16.0f);
+	inline const UE::Slate::FDeprecateVector2DResult Icon8x4(8.0f, 4.0f);
+	inline const UE::Slate::FDeprecateVector2DResult Icon16x4(16.0f, 4.0f);
+	inline const UE::Slate::FDeprecateVector2DResult Icon8x8(8.0f, 8.0f);
+	inline const UE::Slate::FDeprecateVector2DResult Icon10x10(10.0f, 10.0f);
+	inline const UE::Slate::FDeprecateVector2DResult Icon12x12(12.0f, 12.0f);
+	inline const UE::Slate::FDeprecateVector2DResult Icon12x16(12.0f, 16.0f);
+	inline const UE::Slate::FDeprecateVector2DResult Icon14x14(14.0f, 14.0f);
+	inline const UE::Slate::FDeprecateVector2DResult Icon16x16(16.0f, 16.0f);
+	inline const UE::Slate::FDeprecateVector2DResult Icon18x18(18.0f, 18.0f);
+	inline const UE::Slate::FDeprecateVector2DResult Icon20x20(20.0f, 20.0f);
+	inline const UE::Slate::FDeprecateVector2DResult Icon22x22(22.0f, 22.0f);
+	inline const UE::Slate::FDeprecateVector2DResult Icon24x24(24.0f, 24.0f);
+	inline const UE::Slate::FDeprecateVector2DResult Icon25x25(25.0f, 25.0f);
+	inline const UE::Slate::FDeprecateVector2DResult Icon32x32(32.0f, 32.0f);
+	inline const UE::Slate::FDeprecateVector2DResult Icon40x40(40.0f, 40.0f);
+	inline const UE::Slate::FDeprecateVector2DResult Icon64x64(64.0f, 64.0f);
+	inline const UE::Slate::FDeprecateVector2DResult Icon36x24(36.0f, 24.0f);
+	inline const UE::Slate::FDeprecateVector2DResult Icon128x128(128.0f, 128.0f);
 
 	// Common Margins
 	inline const FMargin DefaultMargins(8.f, 4.f);
diff --git a/Engine/Source/Runtime/SlateCore/Public/Styling/SlateBrush.h b/Engine/Source/Runtime/SlateCore/Public/Styling/SlateBrush.h
--- a/Engine/Source/Runtime/SlateCore/Public/Styling/SlateBrush.h
+++ b/Engine/Source/Runtime/SlateCore/Public/Styling/SlateBrush.h
@@ -252,7 +252,7 @@ public:
 
 	/** Size of the resource in Slate Units */
 	UPROPERTY(EditAnywhere, BlueprintReadWrite, Category=Brush)
-	FVector2D ImageSize;
+	FDeprecateSlateVector2D ImageSize;
 
 	/** The margin to use in Box and Border modes */
 	UPROPERTY(EditAnywhere, BlueprintReadWrite, Category=Brush, meta=( UVSpace="true" ))
@@ -284,8 +284,8 @@ public:
 
 public:
 
-	FVector2D GetImageSize() const { return ImageSize; }
-	void SetImageSize(FVector2D InImageSize) { ImageSize = InImageSize; }
+	UE::Slate::FDeprecateVector2DResult GetImageSize() const { return UE::Slate::FDeprecateVector2DResult(ImageSize); }
+	void SetImageSize(UE::Slate::FDeprecateVector2DParameter InImageSize) { ImageSize = InImageSize; }
 
 	const FMargin& GetMargin() const { return Margin; }
 
@@ -437,7 +437,7 @@ public:
 	 */
 	static const FString UTextureIdentifier( );
 	
-	const FSlateResourceHandle& GetRenderingResource(FVector2D LocalSize, float DrawScale) const
+	const FSlateResourceHandle& GetRenderingResource(UE::Slate::FDeprecateVector2DParameter LocalSize, float DrawScale) const
 	{
 		UpdateRenderingResource(LocalSize, DrawScale);
 
@@ -456,7 +456,7 @@ public:
 #endif
 
 private:
-	void UpdateRenderingResource(FVector2D LocalSize, float DrawScale) const;
+	void UpdateRenderingResource(FVector2f LocalSize, float DrawScale) const;
 	bool CanRenderResourceObject(UObject* InResourceObject) const;
 
 private:
@@ -497,11 +497,11 @@ protected:
 	 * @param InTint		  Tint to apply to the element.
 	 * @param InOutlineSettings Optional Outline Border Settings for RoundedBox mode
 	 */
-	 FORCENOINLINE FSlateBrush( ESlateBrushDrawType::Type InDrawType, const FName InResourceName, const FMargin& InMargin, ESlateBrushTileType::Type InTiling, ESlateBrushImageType::Type InImageType, const FVector2D& InImageSize, const FLinearColor& InTint = FLinearColor::White, UObject* InObjectResource = nullptr, bool bInDynamicallyLoaded = false);
+	 FORCENOINLINE FSlateBrush( ESlateBrushDrawType::Type InDrawType, const FName InResourceName, const FMargin& InMargin, ESlateBrushTileType::Type InTiling, ESlateBrushImageType::Type InImageType, const UE::Slate::FDeprecateVector2DParameter& InImageSize, const FLinearColor& InTint = FLinearColor::White, UObject* InObjectResource = nullptr, bool bInDynamicallyLoaded = false);
 
-	 FORCENOINLINE FSlateBrush( ESlateBrushDrawType::Type InDrawType, const FName InResourceName, const FMargin& InMargin, ESlateBrushTileType::Type InTiling, ESlateBrushImageType::Type InImageType, const FVector2D& InImageSize, const TSharedRef< FLinearColor >& InTint, UObject* InObjectResource = nullptr, bool bInDynamicallyLoaded = false);
+	 FORCENOINLINE FSlateBrush( ESlateBrushDrawType::Type InDrawType, const FName InResourceName, const FMargin& InMargin, ESlateBrushTileType::Type InTiling, ESlateBrushImageType::Type InImageType, const UE::Slate::FDeprecateVector2DParameter& InImageSize, const TSharedRef< FLinearColor >& InTint, UObject* InObjectResource = nullptr, bool bInDynamicallyLoaded = false);
 
-	 FORCENOINLINE FSlateBrush( ESlateBrushDrawType::Type InDrawType, const FName InResourceName, const FMargin& InMargin, ESlateBrushTileType::Type InTiling, ESlateBrushImageType::Type InImageType, const FVector2D& InImageSize, const FSlateColor& InTint, UObject* InObjectResource = nullptr, bool bInDynamicallyLoaded = false);
+	 FORCENOINLINE FSlateBrush( ESlateBrushDrawType::Type InDrawType, const FName InResourceName, const FMargin& InMargin, ESlateBrushTileType::Type InTiling, ESlateBrushImageType::Type InImageType, const UE::Slate::FDeprecateVector2DParameter& InImageSize, const FSlateColor& InTint, UObject* InObjectResource = nullptr, bool bInDynamicallyLoaded = false);
 
 };
 
diff --git a/Engine/Source/Runtime/SlateCore/Public/Styling/SlateStyle.h b/Engine/Source/Runtime/SlateCore/Public/Styling/SlateStyle.h
--- a/Engine/Source/Runtime/SlateCore/Public/Styling/SlateStyle.h
+++ b/Engine/Source/Runtime/SlateCore/Public/Styling/SlateStyle.h
@@ -102,7 +102,18 @@ public:
 	 * @param PropertyName - Name of the property to set.
 	 * @param InVector - The value to set.
 	 */
-	FORCENOINLINE void Set(const FName PropertyName, const FVector2D InVector)
+	UE_SLATE_VECTOR_DEPRECATED_DEFAULT()
+	FORCENOINLINE void Set(const FName PropertyName, const FVector2D& InVector)
+	{
+		LLM_SCOPE_BYTAG(UI_Style);
+		Vector2DValues.Add(PropertyName, UE::Slate::CastToVector2f(InVector));
+	}
+	FORCENOINLINE void Set(const FName PropertyName, const UE::Slate::FDeprecateVector2DResult& InVector)
+	{
+		LLM_SCOPE_BYTAG(UI_Style);
+		Vector2DValues.Add(PropertyName, InVector);
+	}
+	FORCENOINLINE void Set(const FName PropertyName, const FVector2f& InVector)
 	{
 		LLM_SCOPE_BYTAG(UI_Style);
 		Vector2DValues.Add(PropertyName, InVector);
@@ -219,7 +230,7 @@ protected:
 	TMap< FName, float > FloatValues;
 
 	/** FVector2D property storage. */
-	TMap< FName, FVector2D > Vector2DValues;
+	TMap< FName, FVector2f > Vector2DValues;
 
 	/** Color property storage. */
 	TMap< FName, FLinearColor > ColorValues;
diff --git a/Engine/Source/Runtime/SlateCore/Public/Styling/SlateTypes.h b/Engine/Source/Runtime/SlateCore/Public/Styling/SlateTypes.h
--- a/Engine/Source/Runtime/SlateCore/Public/Styling/SlateTypes.h
+++ b/Engine/Source/Runtime/SlateCore/Public/Styling/SlateTypes.h
@@ -14,6 +14,7 @@
 #include "Sound/SlateSound.h"
 #include "Styling/SlateBrush.h"
 #include "Styling/SlateWidgetStyle.h"
+#include "Types/SlateVector2.h"
 
 #include "SlateTypes.generated.h"
 
@@ -350,8 +351,8 @@ struct SLATECORE_API FTextBlockStyle : public FSlateWidgetStyle
 
 	/** How much should the shadow be offset? An offset of 0 implies no shadow. */
 	UPROPERTY(EditAnywhere, BlueprintReadWrite, Category=Appearance)
-	FVector2D ShadowOffset;
-	FTextBlockStyle& SetShadowOffset(const FVector2D& InShadowOffset) { ShadowOffset = InShadowOffset; return *this; }
+	FDeprecateSlateVector2D ShadowOffset;
+	FTextBlockStyle& SetShadowOffset(const UE::Slate::FDeprecateVector2DParameter& InShadowOffset) { ShadowOffset = InShadowOffset; return *this; }
 
 	/** The color and opacity of the shadow */
 	UPROPERTY(EditAnywhere, BlueprintReadWrite, Category=Appearance)
@@ -591,8 +592,8 @@ struct SLATECORE_API FComboButtonStyle : public FSlateWidgetStyle
 	  * An offset of 0 implies no shadow. 
 	  */
 	UPROPERTY(EditAnywhere, BlueprintReadWrite, Category=Appearance)
-	FVector2D ShadowOffset;
-	FComboButtonStyle& SetShadowOffset(const FVector2D& InShadowOffset) { ShadowOffset = InShadowOffset; return *this; }
+	FDeprecateSlateVector2D ShadowOffset;
+	FComboButtonStyle& SetShadowOffset(const UE::Slate::FDeprecateVector2DParameter& InShadowOffset) { ShadowOffset = InShadowOffset; return *this; }
 
 	/** 
 	  * The color and opacity of the shadow for the down arrow.
@@ -1883,8 +1884,8 @@ struct SLATECORE_API FDockTabStyle : public FSlateWidgetStyle
 
 	/** Icon size for icons in this tab */
 	UPROPERTY(EditAnywhere, Category = Appearance)
-	FVector2D IconSize;
-	FDockTabStyle& SetIconSize(const FVector2D& InIconSize) { IconSize = InIconSize; return *this; }
+	FDeprecateSlateVector2D IconSize;
+	FDockTabStyle& SetIconSize(const UE::Slate::FDeprecateVector2DParameter& InIconSize) { IconSize = InIconSize; return *this; }
 
 	/** The width that this tab will overlap with side-by-side tabs */
 	UPROPERTY(EditAnywhere, Category=Appearance)
diff --git a/Engine/Source/Runtime/SlateCore/Public/Styling/ToolBarStyle.h b/Engine/Source/Runtime/SlateCore/Public/Styling/ToolBarStyle.h
--- a/Engine/Source/Runtime/SlateCore/Public/Styling/ToolBarStyle.h
+++ b/Engine/Source/Runtime/SlateCore/Public/Styling/ToolBarStyle.h
@@ -106,8 +106,8 @@ struct SLATECORE_API FToolBarStyle : public FSlateWidgetStyle
 	FToolBarStyle& SetBackgroundPadding(const FMargin& InMargin) { BackgroundPadding = InMargin; return *this; }
 
 	UPROPERTY(EditAnywhere, BlueprintReadWrite, Category = Appearance)
-	FVector2D IconSize;
-	FToolBarStyle& SetIconSize(const FVector2D& InIconSize) { IconSize = InIconSize; return *this; }
+	FDeprecateSlateVector2D IconSize;
+	FToolBarStyle& SetIconSize(const UE::Slate::FDeprecateVector2DParameter& InIconSize) { IconSize = InIconSize; return *this; }
 
 	UPROPERTY(EditAnywhere, BlueprintReadWrite, Category=Appearance)
 	bool bShowLabels;
diff --git a/Engine/Source/Runtime/SlateCore/Public/Types/PaintArgs.h b/Engine/Source/Runtime/SlateCore/Public/Types/PaintArgs.h
--- a/Engine/Source/Runtime/SlateCore/Public/Types/PaintArgs.h
+++ b/Engine/Source/Runtime/SlateCore/Public/Types/PaintArgs.h
@@ -24,8 +24,8 @@ class SLATECORE_API FPaintArgs
 	friend class SInvalidationPanel;
 	friend class SRetainerWidget;
 public:
-	FPaintArgs(const SWidget* PaintParent, FHittestGrid& InRootHittestGrid, FHittestGrid& InCurrentHitTestGrid, FVector2D InWindowOffset, double InCurrentTime, float InDeltaTime);
-	FPaintArgs(const SWidget* PaintParent, FHittestGrid& InRootHittestGrid, FVector2D InWindowOffset, double InCurrentTime, float InDeltaTime);
+	FPaintArgs(const SWidget* PaintParent, FHittestGrid& InRootHittestGrid, FHittestGrid& InCurrentHitTestGrid, UE::Slate::FDeprecateVector2DParameter InWindowOffset, double InCurrentTime, float InDeltaTime);
+	FPaintArgs(const SWidget* PaintParent, FHittestGrid& InRootHittestGrid, UE::Slate::FDeprecateVector2DParameter InWindowOffset, double InCurrentTime, float InDeltaTime);
 
 	UE_NODISCARD FORCEINLINE_DEBUGGABLE FPaintArgs WithNewParent(const SWidget* PaintParent) const
 	{
@@ -63,9 +63,9 @@ public:
 		return PaintParentPtr; 
 	}
 
-	FVector2D GetWindowToDesktopTransform() const
+	UE::Slate::FDeprecateVector2DResult GetWindowToDesktopTransform() const
 	{
-		return WindowOffset;
+		return UE::Slate::FDeprecateVector2DResult(WindowOffset);
 	}
 
 	double GetCurrentTime() const
@@ -96,7 +96,7 @@ private:
 	/** The current hit test grid.  Its possible that there is more than one grid when there is nested invalidation panels.  This is what widgets should add to always */
 	FHittestGrid& CurrentGrid;
 
-	FVector2D WindowOffset;
+	FVector2f WindowOffset;
 	const SWidget* PaintParentPtr;
 
 	double CurrentTime;
diff --git a/Engine/Source/Runtime/SlateCore/Public/Types/SlateVector2.h b/Engine/Source/Runtime/SlateCore/Public/Types/SlateVector2.h
--- a/Engine/Source/Runtime/SlateCore/Public/Types/SlateVector2.h
+++ b/Engine/Source/Runtime/SlateCore/Public/Types/SlateVector2.h
@@ -2,219 +2,598 @@
 
 #pragma once
 
-#include "CoreMinimal.h"
+#include "CoreTypes.h"
+#include "Math/Vector2D.h"
+#include "Math/IntPoint.h"
 #include "Math/UnrealMathUtility.h"
+#include "Misc/Optional.h"
 
-namespace UE::Slate
-{
+#include "SlateVector2.generated.h"
 
-FORCEINLINE static FVector2f CastToVector2f(FVector2d InValue)
-{
-	const float X = UE_REAL_TO_FLOAT(InValue.X);
-	const float Y = UE_REAL_TO_FLOAT(InValue.Y);
-#if 0
-	ensureAlways(FMath::IsNearlyEqual((double)X, InValue.X));
-	ensureAlways(FMath::IsNearlyEqual((double)Y, InValue.Y));
+/**
+ * When disabled, all deprecation mechanisms will be disabled, and public slate APIs will only compile with single-precision
+ */
+#ifndef UE_ENABLE_SLATE_VECTOR_DEPRECATION_MECHANISMS
+	#define UE_ENABLE_SLATE_VECTOR_DEPRECATION_MECHANISMS 1
 #endif
-	return FVector2f(X, Y);
-}
+
+
+/**
+ * The subsequent deprecation macros can be enabled or disabled per-module
+ * Tell the IncludeTool static analyzer to ignore that divergence
+ */
+#define UE_INCLUDETOOL_IGNORE_INCONSISTENT_STATE
+
 
 /**
- * Structure for deprecating FVector2D to FVector2f
+ * When a module enables UE_REPORT_SLATE_VECTOR_DEPRECATION through PrivateDefines.Add("UE_REPORT_SLATE_VECTOR_DEPRECATION=1"),
+ * deprecation mechanisms will be enabled on that module for Slate double/single vector APIs
  */
-struct FDeprecateVector2D
+#ifndef UE_REPORT_SLATE_VECTOR_DEPRECATION
+	#define UE_REPORT_SLATE_VECTOR_DEPRECATION 0
+#endif
+
+static_assert(UE_ENABLE_SLATE_VECTOR_DEPRECATION_MECHANISMS || !UE_REPORT_SLATE_VECTOR_DEPRECATION, "UE_REPORT_SLATE_VECTOR_DEPRECATION should never be specified where UE_ENABLE_SLATE_VECTOR_DEPRECATION_MECHANISMS is disabled.");
+
+#ifndef UE_REPORT_SLATE_VECTOR_DEPRECATION_VERSION
+	#define UE_REPORT_SLATE_VECTOR_DEPRECATION_VERSION all
+#endif
+
+
+#if UE_REPORT_SLATE_VECTOR_DEPRECATION
+	#define UE_SLATE_VECTOR_DEPRECATED(Text) UE_DEPRECATED(UE_REPORT_SLATE_VECTOR_DEPRECATION_VERSION, Text)
+	#define UE_SLATE_VECTOR_DEPRECATED_DEFAULT() UE_DEPRECATED(UE_REPORT_SLATE_VECTOR_DEPRECATION_VERSION, "Use FVector2f or float directly; this Slate API no longer interfaces directly with double-precision scalars and vectors.")
+#else
+	#define UE_SLATE_VECTOR_DEPRECATED(Text)
+	#define UE_SLATE_VECTOR_DEPRECATED_DEFAULT()
+#endif
+
+
+#undef UE_INCLUDETOOL_IGNORE_INCONSISTENT_STATE
+
+
+struct FDeprecateSlateVector2D;
+
+namespace UE::Slate
 {
-public:
-	FDeprecateVector2D() = default;
-	explicit FDeprecateVector2D(FVector2f InValue)
-		: X(InValue.X)
-		, Y(InValue.Y)
-	{
-	}
-	explicit FDeprecateVector2D(float InX, float InY)
-		: X(InX)
-		, Y(InY)
-	{
-	}
 
-public:
-	union
+	/**
+	 * CastToVector2f allows supported types to be explicitly coerced to an FVector2f or const FVector2f&
+	 * This allows generic programming patterns that must operate on FVector2f given many input types
+	 */
+	const FVector2f& CastToVector2f(const FVector2f& InValue);
+		  FVector2f  CastToVector2f(const FVector2d& InValue);
+
+
+#if UE_ENABLE_SLATE_VECTOR_DEPRECATION_MECHANISMS
+
+	struct FDeprecateVector2DParameter;
+	struct FDeprecateOptionalVector2DParameter;
+
+		  FVector2f  CastToVector2f(const FIntPoint& InValue);
+	const FVector2f& CastToVector2f(const FDeprecateSlateVector2D& InValue);
+	const FVector2f& CastToVector2f(const FDeprecateVector2DParameter& InValue);
+
+	struct SLATECORE_API FDeprecateSlateVectorPtrVariant
 	{
-		struct
-		{
-			float X;
-			float Y;
-		};
+		explicit FDeprecateSlateVectorPtrVariant(FDeprecateSlateVector2D* InInstance);
+
+		UE_SLATE_VECTOR_DEPRECATED_DEFAULT()
+		operator const FVector2D*() const &;
+
+		operator const FVector2f*() const;
 
-		UE_DEPRECATED(all, "For internal use only")
-		float XY[2];
+		operator FVector2f*();
+
+		FDeprecateSlateVector2D* Instance;
+		mutable FVector2D DoubleVector;
 	};
 
-public:
-	operator FVector2d() const
+#else // UE_ENABLE_SLATE_VECTOR_DEPRECATION_MECHANISMS
+
+	using FDeprecateVector2DResult            = FVector2f;
+	using FDeprecateVector2DParameter         = FVector2f;
+	using FDeprecateOptionalVector2DParameter = TOptional<FVector2f>;
+
+#endif
+
+} // namespace UE::Slate
+
+
+/**
+ * Defines a return value or persistent member variable that was previously an FVector2D
+ * but is in the process of being deprecated in favor of FVector2f. Once deprecated completely, usages will
+ * be converted to using FVector2f directly. It is not (and will never be) constructible from double precision types.
+ *
+ * This type is implicitly convertible to FVector2D, and can be operated on as if it were an FVector2D,
+ * but will emit deprecation warnings under these situations where UE_REPORT_SLATE_VECTOR_DEPRECATION is enabled.
+ *
+ * This type should not be used as a parameter, where FDeprecateVector2DParameter is preferred due to its
+ * implicit conversion from many types, and separate deprecation mechanisms.
+ *
+ * NOTE: DO NOT USE THIS TYPE DIRECTLY IN CLIENT CODE - ALL USAGES SHOULD USE FVector2f or float
+ */
+USTRUCT(BlueprintInternalUseOnly, DisplayName="Vector2D", meta=(HiddenByDefault, ShortTooltip="Vector2D (single-precision)", ToolTip="Vector2D (single-precision)"))
+struct FDeprecateSlateVector2D : public FVector2f
+{
+	GENERATED_BODY()
+
+	FDeprecateSlateVector2D() = default;
+
+	/** Construction from another vector */
+	FDeprecateSlateVector2D(FVector2f InValue)
+		: FVector2f(InValue)
 	{
-		return FVector2d(X, Y);
 	}
 
-	operator FVector2f() const
+	/** Explicit construction from single-precision component values */
+	explicit FDeprecateSlateVector2D(float InX, float InY)
+		: FVector2f(InX, InY)
 	{
-		return AsVector2f();
 	}
 
-public:
-	/**  Get the length (magnitude) of this vector. */
-	float Size() const
+	/** Copy-assignment from another vector */
+	FDeprecateSlateVector2D& operator=(const FVector2f& InValue)
 	{
-		return AsVector2f().Size();
+		X = InValue.X;
+		Y = InValue.Y;
+		return *this;
 	}
 
-	/** Get the length (magnitude) of this vector. */
-	float Length() const
+	UE_SLATE_VECTOR_DEPRECATED_DEFAULT()
+	FDeprecateSlateVector2D& operator=(const FVector2d& InValue)
 	{
-		return AsVector2f().Length();
+		X = UE_REAL_TO_FLOAT(InValue.X);
+		Y = UE_REAL_TO_FLOAT(InValue.Y);
+		return *this;
 	}
 
-	/** Get the squared length of this vector. */
-	float SizeSquared() const
+
+	UE_SLATE_VECTOR_DEPRECATED_DEFAULT()
+	FDeprecateSlateVector2D& operator=(const FIntPoint& InValue)
 	{
-		return AsVector2f().SizeSquared();
+		X = UE_REAL_TO_FLOAT(InValue.X);
+		Y = UE_REAL_TO_FLOAT(InValue.Y);
+		return *this;
 	}
 
-	/**  Get the squared length of this vector. */
-	float SquaredLength() const
+	/**
+	 * Deprecated conversion to a double-precision vector for interoperability with legacy code
+	 */
+	UE_SLATE_VECTOR_DEPRECATED_DEFAULT()
+	operator FVector2d() const
 	{
-		return AsVector2f().SquaredLength();
+		return FVector2d(X, Y);
 	}
 
-	/** Checks whether all components of the vector are exactly zero. */
-	bool IsZero() const
+#if UE_ENABLE_SLATE_VECTOR_DEPRECATION_MECHANISMS
+
+	const UE::Slate::FDeprecateSlateVectorPtrVariant operator&() const
 	{
-		return AsVector2f().IsZero();
+		return UE::Slate::FDeprecateSlateVectorPtrVariant(const_cast<FDeprecateSlateVector2D*>(this));
 	}
 
-	/** Get this vector as an Int Point. */
-	FIntPoint IntPoint() const
+	UE::Slate::FDeprecateSlateVectorPtrVariant operator&()
 	{
-		return AsVector2f().IntPoint();
+		return UE::Slate::FDeprecateSlateVectorPtrVariant(this);
 	}
 
+#endif
+
+	SLATECORE_API bool Serialize(FStructuredArchive::FSlot Slot);
+	SLATECORE_API bool SerializeFromMismatchedTag(const struct FPropertyTag& Tag, FStructuredArchive::FSlot Slot);
+
+public:
+
+	/*~ ---------------------------------------
+	 * These functions hide the similarly named functions from FVector2f in order to provide FVector2d interop */
+
 	/** Get this vector as a vector where each component has been rounded to the nearest int. */
-	FDeprecateVector2D RoundToVector() const
+	FDeprecateSlateVector2D RoundToVector() const
 	{
-		return FDeprecateVector2D(AsVector2f().RoundToVector());
+		return FDeprecateSlateVector2D(FVector2f::RoundToVector());
 	}
 
-	/** Gets a specific component of the vector. */
-	float Component(int32 Index) const
+	bool Equals(const UE::Slate::FDeprecateVector2DParameter& V, float Tolerance=UE_KINDA_SMALL_NUMBER) const
 	{
-		return AsVector2f().Component(Index);
+		return FVector2f::Equals(UE::Slate::CastToVector2f(V), Tolerance);
 	}
 
-public:
-	FORCEINLINE FDeprecateVector2D operator+(const FVector2f V) const
+	FDeprecateSlateVector2D GetSignVector() const
 	{
-		return FDeprecateVector2D(AsVector2f() + V);
+		return FDeprecateSlateVector2D(FVector2f::GetSignVector());
 	}
-	FORCEINLINE FDeprecateVector2D operator+(const FVector2d V) const
-	{
-		return FDeprecateVector2D(AsVector2f() + CastToVector2f(V));
-	}
-	FORCEINLINE FDeprecateVector2D operator+(const FDeprecateVector2D V) const
+
+	FDeprecateSlateVector2D GetAbs() const
 	{
-		return FDeprecateVector2D(AsVector2f() + V.AsVector2f());
+		return FDeprecateSlateVector2D(FVector2f::GetAbs());
 	}
 
-	FORCEINLINE FDeprecateVector2D operator-(const FVector2f V) const
+	FDeprecateSlateVector2D ClampAxes(float MinAxisVal, float MaxAxisVal) const
 	{
-		return FDeprecateVector2D(AsVector2f() - V);
+		return FVector2f::ClampAxes(MinAxisVal, MaxAxisVal);
 	}
-	FORCEINLINE FDeprecateVector2D operator-(const FVector2d V) const
+
+	FDeprecateSlateVector2D GetRotated(float AngleDeg) const
 	{
-		return FDeprecateVector2D(AsVector2f() - CastToVector2f(V));
+		return FVector2f::GetRotated(AngleDeg);
 	}
-	FORCEINLINE FDeprecateVector2D operator-(const FDeprecateVector2D V) const
+
+	FDeprecateSlateVector2D GetSafeNormal(float Tolerance=UE_SMALL_NUMBER) const
 	{
-		return FDeprecateVector2D(AsVector2f() - V.AsVector2f());
+		return FVector2f::GetSafeNormal(Tolerance);
 	}
 
-	FORCEINLINE FDeprecateVector2D operator*(const FVector2f V) const
+	SLATECORE_API bool ComponentwiseAllLessThan(const UE::Slate::FDeprecateVector2DParameter& Other) const;
+
+	SLATECORE_API bool ComponentwiseAllGreaterThan(const UE::Slate::FDeprecateVector2DParameter& Other) const;
+
+	SLATECORE_API bool ComponentwiseAllLessOrEqual(const UE::Slate::FDeprecateVector2DParameter& Other) const;
+
+	SLATECORE_API bool ComponentwiseAllGreaterOrEqual(const UE::Slate::FDeprecateVector2DParameter& Other) const;
+
+	UE_SLATE_VECTOR_DEPRECATED_DEFAULT()
+	FDeprecateSlateVector2D ClampAxes(double MinAxisVal, double MaxAxisVal) const
 	{
-		return FDeprecateVector2D(AsVector2f() * V);
+		return ClampAxes(UE_REAL_TO_FLOAT(MinAxisVal), UE_REAL_TO_FLOAT(MaxAxisVal));
 	}
-	FORCEINLINE FDeprecateVector2D operator*(const FVector2d V) const
+
+	UE_SLATE_VECTOR_DEPRECATED_DEFAULT()
+	FDeprecateSlateVector2D GetRotated(double AngleDeg) const
 	{
-		return FDeprecateVector2D(AsVector2f() * CastToVector2f(V));
+		return GetRotated(UE_REAL_TO_FLOAT(AngleDeg));
 	}
-	FORCEINLINE FDeprecateVector2D operator*(const FDeprecateVector2D V) const
+	UE_SLATE_VECTOR_DEPRECATED_DEFAULT()
+	FDeprecateSlateVector2D GetSafeNormal(double Tolerance) const
 	{
-		return FDeprecateVector2D(AsVector2f() * V.AsVector2f());
+		return GetSafeNormal(UE_REAL_TO_FLOAT(Tolerance));
 	}
 
-	FORCEINLINE FDeprecateVector2D operator/(const FVector2f V) const
+	/*~ End interop
+	 * --------------------------------------- */
+
+public:
+
+#if UE_ENABLE_SLATE_VECTOR_DEPRECATION_MECHANISMS
+
+	/*~ Macros to assist in defining explicit binary mathematical operators.
+	 These are generally all necessary in order to resolve otherwise ambiguous overloaded operators
+	 if we were to define them in terms of FDeprecateVector2DParameter */
+	#define UE_SLATE_BINARY_VECTOR_OPERATORS_A(Type, ...)\
+		__VA_ARGS__\
+		friend FDeprecateSlateVector2D operator+(const Type& A, const FDeprecateSlateVector2D& B)\
+		{\
+			return FDeprecateSlateVector2D(UE::Slate::CastToVector2f(A) + UE::Slate::CastToVector2f(B));\
+		}\
+		__VA_ARGS__\
+		friend FDeprecateSlateVector2D operator-(const Type& A, const FDeprecateSlateVector2D& B)\
+		{\
+			return FDeprecateSlateVector2D(UE::Slate::CastToVector2f(A) - UE::Slate::CastToVector2f(B));\
+		}\
+		__VA_ARGS__\
+		friend FDeprecateSlateVector2D operator*(const Type& A, const FDeprecateSlateVector2D& B)\
+		{\
+			return FDeprecateSlateVector2D(UE::Slate::CastToVector2f(A) * UE::Slate::CastToVector2f(B));\
+		}\
+		__VA_ARGS__\
+		friend FDeprecateSlateVector2D operator/(const Type& A, const FDeprecateSlateVector2D& B)\
+		{\
+			return FDeprecateSlateVector2D(UE::Slate::CastToVector2f(A) / UE::Slate::CastToVector2f(B));\
+		}\
+		__VA_ARGS__\
+		friend bool operator==(const Type& A, const FDeprecateSlateVector2D& B)\
+		{\
+			return UE::Slate::CastToVector2f(A) == UE::Slate::CastToVector2f(B);\
+		}\
+		__VA_ARGS__\
+		friend bool operator!=(const Type& A, const FDeprecateSlateVector2D& B)\
+		{\
+			return UE::Slate::CastToVector2f(A) != UE::Slate::CastToVector2f(B);\
+		}
+
+	#define UE_SLATE_BINARY_ASSIGNMENT_VECTOR_OPERATORS_A(Type, ...)\
+		__VA_ARGS__\
+		friend Type& operator+=(Type& A, const FDeprecateSlateVector2D& B)\
+		{\
+			A = Type(FDeprecateSlateVector2D(UE::Slate::CastToVector2f(A) + UE::Slate::CastToVector2f(B)));\
+			return A;\
+		}\
+		__VA_ARGS__\
+		friend Type& operator-=(Type& A, const FDeprecateSlateVector2D& B)\
+		{\
+			A = Type(FDeprecateSlateVector2D(UE::Slate::CastToVector2f(A) - UE::Slate::CastToVector2f(B)));\
+			return A;\
+		}\
+		__VA_ARGS__\
+		friend Type& operator*=(Type& A, const FDeprecateSlateVector2D& B)\
+		{\
+			A = Type(FDeprecateSlateVector2D(UE::Slate::CastToVector2f(A) * UE::Slate::CastToVector2f(B)));\
+			return A;\
+		}\
+		__VA_ARGS__\
+		friend Type& operator/=(Type& A, const FDeprecateSlateVector2D& B)\
+		{\
+			A = Type(FDeprecateSlateVector2D(UE::Slate::CastToVector2f(A) / UE::Slate::CastToVector2f(B)));\
+			return A;\
+		}
+
+	#define UE_SLATE_BINARY_VECTOR_OPERATORS_B(Type, ...)\
+		__VA_ARGS__\
+		friend FDeprecateSlateVector2D operator+(const FDeprecateSlateVector2D& A, const Type& B)\
+		{\
+			return FDeprecateSlateVector2D(UE::Slate::CastToVector2f(A) + UE::Slate::CastToVector2f(B));\
+		}\
+		__VA_ARGS__\
+		friend FDeprecateSlateVector2D operator-(const FDeprecateSlateVector2D& A, const Type& B)\
+		{\
+			return FDeprecateSlateVector2D(UE::Slate::CastToVector2f(A) - UE::Slate::CastToVector2f(B));\
+		}\
+		__VA_ARGS__\
+		friend FDeprecateSlateVector2D operator*(const FDeprecateSlateVector2D& A, const Type& B)\
+		{\
+			return FDeprecateSlateVector2D(UE::Slate::CastToVector2f(A) * UE::Slate::CastToVector2f(B));\
+		}\
+		__VA_ARGS__\
+		friend FDeprecateSlateVector2D operator/(const FDeprecateSlateVector2D& A, const Type& B)\
+		{\
+			return FDeprecateSlateVector2D(UE::Slate::CastToVector2f(A) / UE::Slate::CastToVector2f(B));\
+		}\
+		__VA_ARGS__\
+		friend bool operator==(const FDeprecateSlateVector2D& A, const Type& B)\
+		{\
+			return UE::Slate::CastToVector2f(A) == UE::Slate::CastToVector2f(B);\
+		}\
+		__VA_ARGS__\
+		friend bool operator!=(const FDeprecateSlateVector2D& A, const Type& B)\
+		{\
+			return UE::Slate::CastToVector2f(A) != UE::Slate::CastToVector2f(B);\
+		}
+
+	#define UE_SLATE_BINARY_ASSIGNMENT_VECTOR_OPERATORS_B(Type, ...)\
+		__VA_ARGS__\
+		friend FDeprecateSlateVector2D& operator+=(FDeprecateSlateVector2D& A, const Type& B)\
+		{\
+			static_cast<FVector2f&>(A) += UE::Slate::CastToVector2f(B);\
+			return A;\
+		}\
+		__VA_ARGS__\
+		friend FDeprecateSlateVector2D& operator-=(FDeprecateSlateVector2D& A, const Type& B)\
+		{\
+			static_cast<FVector2f&>(A) -= UE::Slate::CastToVector2f(B);\
+			return A;\
+		}\
+		__VA_ARGS__\
+		friend FDeprecateSlateVector2D& operator*=(FDeprecateSlateVector2D& A, const Type& B)\
+		{\
+			static_cast<FVector2f&>(A) *= UE::Slate::CastToVector2f(B);\
+			return A;\
+		}\
+		__VA_ARGS__\
+		friend FDeprecateSlateVector2D& operator/=(FDeprecateSlateVector2D& A, const Type& B)\
+		{\
+			static_cast<FVector2f&>(A) /= UE::Slate::CastToVector2f(B);\
+			return A;\
+		}
+
+	/*~ Begin binary operator overloads */
+	UE_SLATE_BINARY_VECTOR_OPERATORS_A(FIntPoint, UE_SLATE_VECTOR_DEPRECATED_DEFAULT())
+	UE_SLATE_BINARY_VECTOR_OPERATORS_B(FIntPoint, UE_SLATE_VECTOR_DEPRECATED_DEFAULT())
+
+	UE_SLATE_BINARY_VECTOR_OPERATORS_A(FVector2d, UE_SLATE_VECTOR_DEPRECATED_DEFAULT())
+	UE_SLATE_BINARY_VECTOR_OPERATORS_B(FVector2d, UE_SLATE_VECTOR_DEPRECATED_DEFAULT())
+	UE_SLATE_BINARY_ASSIGNMENT_VECTOR_OPERATORS_A(FVector2d, UE_SLATE_VECTOR_DEPRECATED_DEFAULT())
+	UE_SLATE_BINARY_ASSIGNMENT_VECTOR_OPERATORS_B(FVector2d, UE_SLATE_VECTOR_DEPRECATED_DEFAULT())
+
+	UE_SLATE_BINARY_VECTOR_OPERATORS_A(FVector2f, )
+	UE_SLATE_BINARY_VECTOR_OPERATORS_B(FVector2f, )
+	UE_SLATE_BINARY_ASSIGNMENT_VECTOR_OPERATORS_A(FVector2f, )
+	UE_SLATE_BINARY_ASSIGNMENT_VECTOR_OPERATORS_B(FVector2f, )
+
+	UE_SLATE_BINARY_VECTOR_OPERATORS_A(FDeprecateSlateVector2D, )
+	UE_SLATE_BINARY_ASSIGNMENT_VECTOR_OPERATORS_A(FDeprecateSlateVector2D, )
+
+	friend FDeprecateSlateVector2D operator*(float A, const FDeprecateSlateVector2D& B)
 	{
-		return FDeprecateVector2D(AsVector2f() / V);
+		return FDeprecateSlateVector2D(A * UE::Slate::CastToVector2f(B));
 	}
-	FORCEINLINE FDeprecateVector2D operator/(const FVector2d V) const
+	friend FDeprecateSlateVector2D operator+(const FDeprecateSlateVector2D& A, float B)
 	{
-		return FDeprecateVector2D(AsVector2f() / CastToVector2f(V));
+		return FDeprecateSlateVector2D(UE::Slate::CastToVector2f(A) + B);
 	}
-	FORCEINLINE FDeprecateVector2D operator/(const FDeprecateVector2D V) const
+	friend FDeprecateSlateVector2D operator-(const FDeprecateSlateVector2D& A, float B)
 	{
-		return FDeprecateVector2D(AsVector2f() / V.AsVector2f());
+		return FDeprecateSlateVector2D(UE::Slate::CastToVector2f(A) - B);
 	}
-
-	FORCEINLINE FDeprecateVector2D operator+(float Scale) const
+	friend FDeprecateSlateVector2D operator*(const FDeprecateSlateVector2D& A, float B)
 	{
-		return FDeprecateVector2D(AsVector2f() + Scale);
+		return FDeprecateSlateVector2D(UE::Slate::CastToVector2f(A) * B);
 	}
-
-	FORCEINLINE FDeprecateVector2D operator-(float Scale) const
+	friend FDeprecateSlateVector2D operator/(const FDeprecateSlateVector2D& A, float B)
 	{
-		return FDeprecateVector2D(AsVector2f() - Scale);
+		return FDeprecateSlateVector2D(UE::Slate::CastToVector2f(A) / B);
 	}
 
-	FORCEINLINE FDeprecateVector2D operator*(float Scale) const
-	{
-		return FDeprecateVector2D(AsVector2f() * Scale);
-	}
+	/*~ End binary operator overloads */
 
-	FORCEINLINE FDeprecateVector2D operator/(float Scale) const
+	/*~ Unary operator overloads */
+	friend FDeprecateSlateVector2D operator-(const FDeprecateSlateVector2D& In)
 	{
-		return FDeprecateVector2D(AsVector2f() / Scale);
+		return FDeprecateSlateVector2D(-UE::Slate::CastToVector2f(In));
 	}
+	/*~ Unary operator overloads */
 
-public:
-	bool operator==(FVector2f V) const
-	{
-		return V == AsVector2f();
-	}
-	bool operator==(FVector2d V) const
-	{
-		return CastToVector2f(V) == AsVector2f();
-	}
-	bool operator==(FDeprecateVector2D V) const
+#undef UE_SLATE_BINARY_VECTOR_OPERATORS_A
+#undef UE_SLATE_BINARY_VECTOR_OPERATORS_B
+#undef UE_SLATE_BINARY_ASSIGNMENT_VECTOR_OPERATORS_A
+#undef UE_SLATE_BINARY_ASSIGNMENT_VECTOR_OPERATORS_B
+
+#endif // UE_ENABLE_SLATE_VECTOR_DEPRECATION_MECHANISMS
+};
+
+template<>
+struct TStructOpsTypeTraits<FDeprecateSlateVector2D> : public TStructOpsTypeTraitsBase2<FDeprecateSlateVector2D>
+{
+	enum
 	{
-		return X == V.X && Y == V.Y;
-	}
+		WithStructuredSerializer = true,
+		WithStructuredSerializeFromMismatchedTag = true,
+	};
+};
+
+
+#if UE_ENABLE_SLATE_VECTOR_DEPRECATION_MECHANISMS
+
+namespace UE::Slate
+{
+
+/**
+ * Defines a return value or persistent member variable that was previously an FVector2D
+ * but is in the process of being deprecated in favor of FVector2f. Once deprecated completely, usages will
+ * be converted to using FVector2f directly. It is not (and will never be) constructible from double precision types.
+ *
+ * This type is implicitly convertible to FVector2D, and can be operated on as if it were an FVector2D,
+ * but will emit deprecation warnings under these situations where UE_REPORT_SLATE_VECTOR_DEPRECATION is enabled.
+ *
+ * This type should not be used as a parameter, where FDeprecateVector2DParameter is preferred due to its
+ * implicit conversion from many types, and separate deprecation mechanisms.
+ *
+ * NOTE: DO NOT USE THIS TYPE DIRECTLY IN CLIENT CODE - ALL USAGES SHOULD USE FVector2f or float
+ */
+using FDeprecateVector2DResult = FDeprecateSlateVector2D;
+
+/**
+ * Defines a Slate vector used as a function parameter that was previously an FVector2D but is in the process
+ * of being deprecated in favor of FVector2f. Once deprecated completely, usages will be converted to using FVector2f directly.
+ *
+ * This type is implicitly constructible from both single and double precision vectors and scalars but will emit
+ * deprecation warnings for double-precision types where UE_REPORT_SLATE_VECTOR_DEPRECATION is enabled. It is not (and will never be)
+ * convertible to double precision types.
+ *
+ * This type should only be used as a parameter (not a return type), where FDeprecateVector2DResult is preferred so it can
+ * interop with legacy code that still requires conversion to FVector2d.
+ *
+ * NOTE: DO NOT USE THIS TYPE DIRECTLY IN CLIENT CODE - ALL USAGES SHOULD USE FVector2f or float
+ */
+struct FDeprecateVector2DParameter : FVector2f
+{
+	FDeprecateVector2DParameter(float InX, float InY)
+		: FVector2f(InX, InY)
+	{}
 
-	bool operator!=(FVector2f V) const
+	FDeprecateVector2DParameter(const FVector2f& InValue)
+		: FVector2f(InValue)
+	{}
+
+	FDeprecateVector2DParameter(const FDeprecateVector2DResult& InValue)
+		: FVector2f(InValue)
+	{}
+
+	UE_SLATE_VECTOR_DEPRECATED_DEFAULT()
+	FDeprecateVector2DParameter(double InX, double InY)
+		: FVector2f(UE_REAL_TO_FLOAT(InX), UE_REAL_TO_FLOAT(InY))
+	{}
+
+	UE_SLATE_VECTOR_DEPRECATED_DEFAULT()
+	FDeprecateVector2DParameter(const FVector2d& InValue)
+		: FVector2f(UE_REAL_TO_FLOAT(InValue.X), UE_REAL_TO_FLOAT(InValue.Y))
+	{}
+
+	UE_SLATE_VECTOR_DEPRECATED_DEFAULT()
+	FDeprecateVector2DParameter(const FIntPoint& InValue)
+		: FVector2f(FloatCastChecked<float, double>(InValue.X, 1.0), FloatCastChecked<float, double>(InValue.Y, 1.0))
+	{}
+};
+
+/**
+ * Defines an optional Slate vector parameter in a similar vein to FDeprecateVector2DParameter
+ * NOTE: DO NOT USE THIS TYPE DIRECTLY IN CLIENT CODE - ALL USAGES SHOULD USE FVector2f or float
+ */
+struct FDeprecateOptionalVector2DParameter : TOptional<FVector2f>
+{
+	FDeprecateOptionalVector2DParameter()
+	{}
+
+	FDeprecateOptionalVector2DParameter(const FVector2f& In)
+		: TOptional<FVector2f>(In)
+	{}
+
+	FDeprecateOptionalVector2DParameter(const TOptional<FVector2f>& In)
+		: TOptional<FVector2f>(In)
+	{}
+
+	UE_SLATE_VECTOR_DEPRECATED_DEFAULT()
+	FDeprecateOptionalVector2DParameter(const FVector2d& In)
+		: TOptional<FVector2f>(CastToVector2f(In))
 	{
-		return V != AsVector2f();
 	}
-	bool operator!=(FVector2d V) const
+
+	UE_SLATE_VECTOR_DEPRECATED_DEFAULT()
+	FDeprecateOptionalVector2DParameter(const TOptional<FVector2d>& In)
 	{
-		return CastToVector2f(V) != AsVector2f();
+		if (In.IsSet())
+		{
+			this->Emplace(CastToVector2f(In.GetValue()));
+		}
 	}
-	bool operator!=(FDeprecateVector2D V) const
+};
+
+inline FVector2f CastToVector2f(const FIntPoint& InValue)
+{
+	return InValue;
+}
+
+inline const FVector2f& CastToVector2f(const FDeprecateSlateVector2D& InValue)
+{
+	return InValue;
+}
+
+inline const FVector2f& CastToVector2f(const FDeprecateVector2DParameter& InValue)
+{
+	return InValue;
+}
+
+} // namespace UE::Slate
+
+inline UE::Slate::FDeprecateVector2DResult TransformPoint(float Transform, const UE::Slate::FDeprecateVector2DResult& Point)
+{
+	return UE::Slate::FDeprecateVector2DResult(Transform * FVector2f(Point));
+}
+
+inline UE::Slate::FDeprecateVector2DResult TransformVector(float Transform, const UE::Slate::FDeprecateVector2DResult& Vector)
+{
+	return UE::Slate::FDeprecateVector2DResult(Transform * FVector2f(Vector));
+}
+
+template <typename TransformType>
+inline UE::Slate::FDeprecateVector2DResult TransformPoint(const TransformType& Transform, const UE::Slate::FDeprecateVector2DResult& Point)
+{
+	return UE::Slate::FDeprecateVector2DResult(Transform.TransformPoint(Point));
+}
+
+template <typename TransformType>
+inline UE::Slate::FDeprecateVector2DResult TransformVector(const TransformType& Transform, const UE::Slate::FDeprecateVector2DResult& Point)
+{
+	return UE::Slate::FDeprecateVector2DResult(Transform.TransformVector(Point));
+}
+
+#endif // UE_ENABLE_SLATE_VECTOR_DEPRECATION_MECHANISMS
+
+namespace UE::Slate
+{
+
+	inline const FVector2f& CastToVector2f(const FVector2f& InValue)
 	{
-		return X != V.X || Y != V.Y;
+		return InValue;
 	}
-
-private:
-	FVector2f AsVector2f() const
+	inline FVector2f CastToVector2f(const FVector2d& InValue)
 	{
+		const float X = UE_REAL_TO_FLOAT(InValue.X);
+		const float Y = UE_REAL_TO_FLOAT(InValue.Y);
+	#if 0
+		ensureAlways(FMath::IsNearlyEqual((double)X, InValue.X));
+		ensureAlways(FMath::IsNearlyEqual((double)Y, InValue.Y));
+	#endif
 		return FVector2f(X, Y);
 	}
-};
 
-} // namespace
+} // namespace UE::Slate
\ No newline at end of file
diff --git a/Engine/Source/Runtime/SlateCore/Public/Widgets/SBoxPanel.h b/Engine/Source/Runtime/SlateCore/Public/Widgets/SBoxPanel.h
--- a/Engine/Source/Runtime/SlateCore/Public/Widgets/SBoxPanel.h
+++ b/Engine/Source/Runtime/SlateCore/Public/Widgets/SBoxPanel.h
@@ -622,9 +622,9 @@ private:
 	FSlateBrush DropIndicator_Below;
 
 	/** This is required for the paint to access which item we're hovering */
-	FVector2D CurrentDragOperationScreenSpaceLocation;
+	FVector2f CurrentDragOperationScreenSpaceLocation;
 	int32 CurrentDragOverSlotIndex;
 
 	/** @return the zone (above, below) based on where the user is hovering over */
-	EItemDropZone ZoneFromPointerPosition(FVector2D LocalPointerPos, const FGeometry& CurrentGeometry, const FGeometry& StartGeometry) const;
+	EItemDropZone ZoneFromPointerPosition(FVector2f LocalPointerPos, const FGeometry& CurrentGeometry, const FGeometry& StartGeometry) const;
 };
diff --git a/Engine/Source/Runtime/SlateCore/Public/Widgets/SWidget.h b/Engine/Source/Runtime/SlateCore/Public/Widgets/SWidget.h
--- a/Engine/Source/Runtime/SlateCore/Public/Widgets/SWidget.h
+++ b/Engine/Source/Runtime/SlateCore/Public/Widgets/SWidget.h
@@ -27,6 +27,7 @@
 #include "SlateGlobals.h"
 #include "Types/PaintArgs.h"
 #include "Types/SlateAttribute.h"
+#include "Types/SlateVector2.h"
 #include "FastUpdate/WidgetProxy.h"
 #include "InvalidateWidgetReason.h"
 #include "Widgets/SlateControlledConstruction.h"
@@ -641,16 +642,6 @@ public:
 	 */
 	virtual FPopupMethodReply OnQueryPopupMethod() const;
 
-	UE_DEPRECATED(4.26, "Renaming to TranslateMouseCoordinateForCustomHitTestChild")
-	TSharedPtr<FVirtualPointerPosition> TranslateMouseCoordinateFor3DChild(const TSharedRef<SWidget>& ChildWidget, const FGeometry& MyGeometry, const FVector2D& ScreenSpaceMouseCoordinate, const FVector2D& LastScreenSpaceMouseCoordinate) const
-	{
-		TOptional<FVirtualPointerPosition> Result = TranslateMouseCoordinateForCustomHitTestChild(ChildWidget.Get(), MyGeometry, ScreenSpaceMouseCoordinate, LastScreenSpaceMouseCoordinate);
-		return Result.IsSet() ? MakeShared<FVirtualPointerPosition>(Result.GetValue()) : TSharedPtr<FVirtualPointerPosition>();
-	}
-
-	UE_DEPRECATED(5.0, "TranslateMouseCoordinateForCustomHitTestChild that returns a shared ptr is deprecated.")
-	virtual TSharedPtr<FVirtualPointerPosition> TranslateMouseCoordinateForCustomHitTestChild(const TSharedRef<SWidget>& ChildWidget, const FGeometry& MyGeometry, const FVector2D& ScreenSpaceMouseCoordinate, const FVector2D& LastScreenSpaceMouseCoordinate) const;
-
 	virtual TOptional<FVirtualPointerPosition> TranslateMouseCoordinateForCustomHitTestChild(const SWidget& ChildWidget, const FGeometry& MyGeometry, const FVector2D ScreenSpaceMouseCoordinate, const FVector2D LastScreenSpaceMouseCoordinate) const;
 
 	/**
@@ -714,7 +705,7 @@ public:
 	const FWidgetProxyHandle GetProxyHandle() const { return FastPathProxyHandle; }
 
 	/** @return the DesiredSize that was computed the last time CacheDesiredSize() was called. */
-	FVector2D GetDesiredSize() const;
+	UE::Slate::FDeprecateVector2DResult GetDesiredSize() const;
 
 	void AssignParentWidget(TSharedPtr<SWidget> InParent);
 	bool ConditionallyDetatchParentWidget(SWidget* InExpectedParent);
@@ -1224,16 +1215,16 @@ public:
 		}
 	}
 
-	FORCEINLINE FVector2D GetRenderTransformPivotWithRespectToFlowDirection() const
+	FORCEINLINE UE::Slate::FDeprecateVector2DResult GetRenderTransformPivotWithRespectToFlowDirection() const
 	{
+		FVector2f TransformPivot = UE::Slate::CastToVector2f(RenderTransformPivotAttribute.Get());
 		if (LIKELY(GSlateFlowDirection == EFlowDirection::LeftToRight))
 		{
-			return RenderTransformPivotAttribute.Get();
+			return TransformPivot;
 		}
 		else
 		{
 			// If we're going right to left, flip the X's pivot mirrored about 0.5.
-			FVector2D TransformPivot = RenderTransformPivotAttribute.Get();
 			TransformPivot.X = 0.5f + (0.5f - TransformPivot.X);
 			return TransformPivot;
 		}
@@ -1246,9 +1237,10 @@ public:
 	}
 
 	/** @return the pivot point of the render transform. */
-	FORCEINLINE FVector2D GetRenderTransformPivot() const
+	FORCEINLINE UE::Slate::FDeprecateVector2DResult GetRenderTransformPivot() const
 	{
-		return RenderTransformPivotAttribute.Get();
+		FVector2f TransformPivot = UE::Slate::CastToVector2f(RenderTransformPivotAttribute.Get());
+		return TransformPivot;
 	}
 
 	/** @param InTransformPivot Sets the pivot point of the widget's render transform (in normalized local space). */
@@ -1575,7 +1567,7 @@ protected:
 	static int32 FindChildUnderMouse( const FArrangedChildren& Children, const FPointerEvent& MouseEvent );
 
 	/** @return The index of the child that is under the specified position */
-	static int32 FindChildUnderPosition(const FArrangedChildren& Children, const FVector2D& ArrangedSpacePosition);
+	static int32 FindChildUnderPosition(const FArrangedChildren& Children, const UE::Slate::FDeprecateVector2DParameter& ArrangedSpacePosition);
 
 	/** 
 	 * Determines if this widget should be enabled.
@@ -1927,14 +1919,14 @@ private:
 // FGeometry Arranged Widget Inlined Functions
 //=================================================================
 
-FORCEINLINE_DEBUGGABLE FArrangedWidget FGeometry::MakeChild(const TSharedRef<SWidget>& ChildWidget, const FVector2D& InLocalSize, const FSlateLayoutTransform& LayoutTransform) const
+FORCEINLINE_DEBUGGABLE FArrangedWidget FGeometry::MakeChild(const TSharedRef<SWidget>& ChildWidget, const UE::Slate::FDeprecateVector2DParameter& InLocalSize, const FSlateLayoutTransform& LayoutTransform) const
 {
 	// If there is no render transform set, use the simpler MakeChild call that doesn't bother concatenating the render transforms.
 	// This saves a significant amount of overhead since every widget does this, and most children don't have a render transform.
 	const TOptional<FSlateRenderTransform> RenderTransform = ChildWidget->GetRenderTransformWithRespectToFlowDirection();
 	if (RenderTransform.IsSet() )
 	{
-		const FVector2D RenderTransformPivot = ChildWidget->GetRenderTransformPivotWithRespectToFlowDirection();
+		const FVector2f RenderTransformPivot = UE::Slate::CastToVector2f(ChildWidget->GetRenderTransformPivotWithRespectToFlowDirection());
 		return FArrangedWidget(ChildWidget, MakeChild(InLocalSize, LayoutTransform, RenderTransform.GetValue(), RenderTransformPivot));
 	}
 	else
@@ -1945,17 +1937,16 @@ FORCEINLINE_DEBUGGABLE FArrangedWidget FGeometry::MakeChild(const TSharedRef<SWi
 
 FORCEINLINE_DEBUGGABLE FArrangedWidget FGeometry::MakeChild(const TSharedRef<SWidget>& ChildWidget, const FLayoutGeometry& LayoutGeometry) const
 {
-	return MakeChild(ChildWidget, LayoutGeometry.GetSizeInLocalSpace(), LayoutGeometry.GetLocalToParentTransform());
+	return MakeChild(ChildWidget, FVector2f(LayoutGeometry.GetSizeInLocalSpace()), LayoutGeometry.GetLocalToParentTransform());
 }
 
-FORCEINLINE_DEBUGGABLE FArrangedWidget FGeometry::MakeChild(const TSharedRef<SWidget>& ChildWidget, const FVector2D& ChildOffset, const FVector2D& InLocalSize, float ChildScale) const
+FORCEINLINE_DEBUGGABLE FArrangedWidget FGeometry::MakeChild(const TSharedRef<SWidget>& ChildWidget, const UE::Slate::FDeprecateVector2DParameter& ChildOffset, const UE::Slate::FDeprecateVector2DParameter& InLocalSize, float ChildScale) const
 {
 	// Since ChildOffset is given as a LocalSpaceOffset, we MUST convert this offset into the space of the parent to construct a valid layout transform.
 	// The extra TransformPoint below does this by converting the local offset to an offset in parent space.
-	return MakeChild(ChildWidget, InLocalSize, FSlateLayoutTransform(ChildScale, TransformPoint(ChildScale, ChildOffset)));
+	return MakeChild(ChildWidget, UE::Slate::CastToVector2f(InLocalSize), FSlateLayoutTransform(ChildScale, TransformPoint(ChildScale, UE::Slate::CastToVector2f(ChildOffset))));
 }
 
-
 template<typename TargetValueType, typename SourceValueType>
 bool SetWidgetAttribute(SWidget& ThisWidget, TAttribute<TargetValueType>& TargetValue, const TAttribute<SourceValueType>& SourceValue, EInvalidateWidgetReason BaseInvalidationReason)
 {
diff --git a/Engine/Source/Runtime/SlateCore/Public/Widgets/SWindow.h b/Engine/Source/Runtime/SlateCore/Public/Widgets/SWindow.h
--- a/Engine/Source/Runtime/SlateCore/Public/Widgets/SWindow.h
+++ b/Engine/Source/Runtime/SlateCore/Public/Widgets/SWindow.h
@@ -250,8 +250,8 @@ public:
 		, _Title()
 		, _bDragAnywhere( false )
 		, _AutoCenter( EAutoCenter::PreferredWorkArea )
-		, _ScreenPosition( FVector2D::ZeroVector )
-		, _ClientSize( FVector2D::ZeroVector )
+		, _ScreenPosition( FVector2f::ZeroVector )
+		, _ClientSize( FVector2f::ZeroVector )
 		, _AdjustInitialSizeAndPositionForDPIScale(true)
 		, _SupportsTransparency( EWindowTransparency::None )
 		, _InitialOpacity( 1.0f )
@@ -293,10 +293,10 @@ public:
 		SLATE_ARGUMENT( EAutoCenter, AutoCenter )
 
 		/** Screen-space position where the window should be initially located. */
-		SLATE_ARGUMENT( FVector2D, ScreenPosition )
+		SLATE_ARGUMENT( UE::Slate::FDeprecateVector2DParameter, ScreenPosition )
 
 		/** What the initial size of the window should be. */
-		SLATE_ARGUMENT( FVector2D, ClientSize )
+		SLATE_ARGUMENT( UE::Slate::FDeprecateVector2DParameter, ClientSize )
 
 		/** If the initial ClientSize and ScreenPosition arguments should be automatically adjusted to account for DPI scale */
 		SLATE_ARGUMENT( bool, AdjustInitialSizeAndPositionForDPIScale )
@@ -420,7 +420,7 @@ public:
 	 * @param ContentSize      The size of content that we want to accommodate 
 	 *
 	 * @return The size of the window necessary to accommodate the given content */
-	static FVector2D ComputeWindowSizeForContent( FVector2D ContentSize );
+	static UE::Slate::FDeprecateVector2DResult ComputeWindowSizeForContent( UE::Slate::FDeprecateVector2DParameter ContentSize );
 
 	/**
 	 * Grabs the window type
@@ -467,13 +467,13 @@ public:
 	FOptionalSize GetTitleBarSize() const;
 
 	/** @return the desired size in desktop pixels */
-	FVector2D GetDesiredSizeDesktopPixels() const;
+	UE::Slate::FDeprecateVector2DResult GetDesiredSizeDesktopPixels() const;
 
 	/**	@return The initially desired screen position of the slate window */
-	FVector2D GetInitialDesiredSizeInScreen() const;
+	UE::Slate::FDeprecateVector2DResult GetInitialDesiredSizeInScreen() const;
 
 	/**	@return The initially desired size of the slate window */
-	FVector2D GetInitialDesiredPositionInScreen() const;
+	UE::Slate::FDeprecateVector2DResult GetInitialDesiredPositionInScreen() const;
 
 	/** Get the Geometry that describes this window. Windows in Slate are unique in that they know their own geometry. */
 	FGeometry GetWindowGeometryInScreen() const;
@@ -488,10 +488,10 @@ public:
 	FSlateLayoutTransform GetLocalToWindowTransform() const;
 
 	/** @return The position of the window in screen space */
-	FVector2D GetPositionInScreen() const;
+	UE::Slate::FDeprecateVector2DResult GetPositionInScreen() const;
 
 	/** @return the size of the window in screen pixels */
-	FVector2D GetSizeInScreen() const;
+	UE::Slate::FDeprecateVector2DResult GetSizeInScreen() const;
 
 	/** @return the rectangle of the window for its non-maximized state */
 	FSlateRect GetNonMaximizedRectInScreen() const;
@@ -503,7 +503,7 @@ public:
 	FSlateRect GetClientRectInScreen() const;
 
 	/** @return the size of the window's usable client area. */
-	FVector2D GetClientSizeInScreen() const;
+	UE::Slate::FDeprecateVector2DResult GetClientSizeInScreen() const;
 
 	/** @return a clipping rectangle that represents this window in Window Space (i.e. always starts at 0,0) */
 	FSlateRect GetClippingRectangleInWindow() const;
@@ -515,16 +515,16 @@ public:
 	FMargin GetNonMaximizedWindowBorderSize() const;
 
 	/** Relocate the window to a screenspace position specified by NewPosition */
-	void MoveWindowTo( FVector2D NewPosition );
+	void MoveWindowTo( UE::Slate::FDeprecateVector2DParameter NewPosition );
 	/** Relocate the window to a screenspace position specified by NewPosition and resize it to NewSize */
-	void ReshapeWindow( FVector2D NewPosition, FVector2D NewSize );
+	void ReshapeWindow( UE::Slate::FDeprecateVector2DParameter NewPosition, UE::Slate::FDeprecateVector2DParameter NewSize );
 	void ReshapeWindow( const FSlateRect& InNewShape );
 	/**
 	 * Resize the window to be dpi scaled NewClientSize immediately
 	 *
 	 * @param NewClientSize: Client size with DPI scaling already applied that does not include border or title bars.
 	 */
-	void Resize( FVector2D NewClientSize );
+	void Resize( UE::Slate::FDeprecateVector2DParameter NewClientSize );
 
 	/** Returns the rectangle of the screen the window is associated with */
 	FSlateRect GetFullScreenInfo() const;
@@ -534,15 +534,15 @@ public:
 	/** @return Returns true if the window is currently morphing and is morphing by size */
 	bool IsMorphingSize() const;
 	/** Animate the window to TargetOpacity and TargetPosition over a short period of time */
-	void MorphToPosition( const FCurveSequence& Sequence, const float TargetOpacity, const FVector2D& TargetPosition );
+	void MorphToPosition( const FCurveSequence& Sequence, const float TargetOpacity, const UE::Slate::FDeprecateVector2DParameter& TargetPosition );
 	/** Animate the window to TargetOpacity and TargetShape over a short period of time */
 	void MorphToShape( const FCurveSequence& Sequence, const float TargetOpacity, const FSlateRect& TargetShape );
 	/** Set a new morph shape and force the morph to run for at least one frame in order to reach that target */
 	void UpdateMorphTargetShape( const FSlateRect& TargetShape );
 	/** Set a new morph position and force the morph to run for at least one frame in order to reach that target */
-	void UpdateMorphTargetPosition( const FVector2D& TargetPosition );
+	void UpdateMorphTargetPosition( const UE::Slate::FDeprecateVector2DParameter& TargetPosition );
 	/** @return Returns the currently set morph target position */
-	FVector2D GetMorphTargetPosition() const;
+	UE::Slate::FDeprecateVector2DResult GetMorphTargetPosition() const;
 	/** @return Returns the currently set morph target shape */
 	FSlateRect GetMorphTargetShape() const;
 
@@ -571,10 +571,10 @@ public:
 	void HACK_ForceToFront();
 
 	/** Sets the actual screen position of the window. THIS SHOULD ONLY BE CALLED BY THE OS */
-	void SetCachedScreenPosition(FVector2D NewPosition);
+	void SetCachedScreenPosition(UE::Slate::FDeprecateVector2DParameter NewPosition);
 
 	/**	Sets the actual size of the window. THIS SHOULD ONLY BE CALLED BY THE OS */
-	void SetCachedSize(FVector2D NewSize);
+	void SetCachedSize(UE::Slate::FDeprecateVector2DParameter NewSize);
 
 	TSharedPtr<FGenericWindow> GetNativeWindow();
 	TSharedPtr<const FGenericWindow> GetNativeWindow() const ;
@@ -838,7 +838,7 @@ public:
 	bool HasOSWindowBorder() const { return bHasOSWindowBorder; }
 
 	/** @return true if mouse coordinates is within this window */
-	bool IsScreenspaceMouseWithin(FVector2D ScreenspaceMouseCoordinate) const;
+	bool IsScreenspaceMouseWithin(UE::Slate::FDeprecateVector2DParameter ScreenspaceMouseCoordinate) const;
 
 	/** @return true if this is a user-sized window with a thick edge */
 	bool HasSizingFrame() const;
@@ -907,11 +907,11 @@ public:
 	void Minimize();
 
 	/** Gets the current Window Zone that mouse position is over. */
-	EWindowZone::Type GetCurrentWindowZone(FVector2D LocalMousePosition);
+	EWindowZone::Type GetCurrentWindowZone(UE::Slate::FDeprecateVector2DParameter LocalMousePosition);
 
 	/** Used to store the zone where the mouse down event occurred during move / drag */
 	EWindowZone::Type MoveResizeZone;
-	FVector2D MoveResizeStart;
+	UE::Slate::FDeprecateVector2DResult MoveResizeStart;
 	FSlateRect MoveResizeRect;
 
 	/** @return Gets the radius of the corner rounding of the window. */
@@ -942,7 +942,7 @@ private:
 	virtual bool ComputeVolatility() const override;
 
 	/** Resize using already dpi scaled window size including borders/title bar */
-	void ResizeWindowSize( FVector2D NewWindowSize );
+	void ResizeWindowSize( FVector2f NewWindowSize );
 
 	void OnGlobalInvalidationToggled(bool bGlobalInvalidationEnabled);
 public:
@@ -952,7 +952,7 @@ public:
 	 * @param InClientSize: Client size with DPI scaling already applied
 	 * @param DPIScale: Scale that will be applied for border and title. When not supplied detects DPIScale using native or initial position.
 	 */
-	FVector2D GetWindowSizeFromClientSize(FVector2D InClientSize, TOptional<float> DPIScale = TOptional<float>());
+	UE::Slate::FDeprecateVector2DResult GetWindowSizeFromClientSize(UE::Slate::FDeprecateVector2DParameter InClientSize, TOptional<float> DPIScale = TOptional<float>());
 
 	/** @return true if this window will be focused when it is first shown */
 	inline bool IsFocusedInitially() const
@@ -1005,7 +1005,7 @@ public:
 	{
 		if (bDrivenByWindow)
 		{
-			ViewportSize = FVector2D::ZeroVector;
+			ViewportSize = FVector2f::ZeroVector;
 		}
 		else
 		{
@@ -1024,7 +1024,7 @@ public:
 	/**
 	 * Returns the viewport size, taking into consideration if the window size should drive the viewport size 
 	 */
-	inline FVector2D GetViewportSize() const
+	inline UE::Slate::FDeprecateVector2DResult GetViewportSize() const
 	{
 		return (ViewportSize.X != 0) ? ViewportSize : Size;
 	}
@@ -1032,7 +1032,7 @@ public:
 	/**
 	 * Sets the viewport size independently of the window size, if non-zero.
 	 */
-	inline void SetIndependentViewportSize(const FVector2D& VP) 
+	inline void SetIndependentViewportSize(const UE::Slate::FDeprecateVector2DParameter& VP) 
 	{
 		ViewportSize = VP;
 	}
@@ -1168,22 +1168,22 @@ protected:
 	EWindowActivationPolicy WindowActivationPolicy;
 
 	/** Initial desired position of the window's content in screen space */
-	FVector2D InitialDesiredScreenPosition;
+	UE::Slate::FDeprecateVector2DResult InitialDesiredScreenPosition;
 
 	/** Initial desired size of the window's content in screen space */
-	FVector2D InitialDesiredSize;
+	UE::Slate::FDeprecateVector2DResult InitialDesiredSize;
 
 	/** Position of the window's content in screen space */
-	FVector2D ScreenPosition;
+	UE::Slate::FDeprecateVector2DResult ScreenPosition;
 
 	/** The position of the window before entering fullscreen */
-	FVector2D PreFullscreenPosition;
+	UE::Slate::FDeprecateVector2DResult PreFullscreenPosition;
 
 	/** Size of the window's content area in screen space */
-	FVector2D Size;
+	UE::Slate::FDeprecateVector2DResult Size;
 
 	/** Size of the viewport. If (0,0) then it is equal to Size */
-	FVector2D ViewportSize;
+	UE::Slate::FDeprecateVector2DResult ViewportSize;
 
 	/** Pointer to the viewport registered with this window if any */
 	TWeakPtr<ISlateViewport> Viewport;
diff --git a/Engine/Source/Runtime/SlateCore/SlateCore.Build.cs b/Engine/Source/Runtime/SlateCore/SlateCore.Build.cs
--- a/Engine/Source/Runtime/SlateCore/SlateCore.Build.cs
+++ b/Engine/Source/Runtime/SlateCore/SlateCore.Build.cs
@@ -6,6 +6,8 @@ public class SlateCore : ModuleRules
 {
 	public SlateCore(ReadOnlyTargetRules Target) : base(Target)
 	{
+		PrivateDefinitions.Add("UE_REPORT_SLATE_VECTOR_DEPRECATION=1");
+
 		PublicDependencyModuleNames.AddRange(
 			new string[] {
 				"Core",
diff --git a/Engine/Source/Runtime/UMG/Private/WidgetBlueprintLibrary.cpp b/Engine/Source/Runtime/UMG/Private/WidgetBlueprintLibrary.cpp
--- a/Engine/Source/Runtime/UMG/Private/WidgetBlueprintLibrary.cpp
+++ b/Engine/Source/Runtime/UMG/Private/WidgetBlueprintLibrary.cpp
@@ -168,7 +168,7 @@ void UWidgetBlueprintLibrary::DrawBox(FPaintContext& Context, FVector2D Position
 		FSlateDrawElement::MakeBox(
 			Context.OutDrawElements,
 			Context.MaxLayer,
-			Context.AllottedGeometry.ToPaintGeometry(Position, Size),
+			Context.AllottedGeometry.ToPaintGeometry(Size, FSlateLayoutTransform(Position)),
 			&Brush->Brush,
 			ESlateDrawEffect::None,
 			Tint);
