PR #9929: Fix ToolTips not following cursor while Slate GI is enabled (Contributed by [REDACTED_PERSON])

diff --git a/Engine/Source/Runtime/Slate/Private/Widgets/LayerManager/STooltipPresenter.cpp b/Engine/Source/Runtime/Slate/Private/Widgets/LayerManager/STooltipPresenter.cpp
--- a/Engine/Source/Runtime/Slate/Private/Widgets/LayerManager/STooltipPresenter.cpp
+++ b/Engine/Source/Runtime/Slate/Private/Widgets/LayerManager/STooltipPresenter.cpp
@@ -20,6 +20,12 @@ void STooltipPresenter::SetContent(TSharedPtr<SWidget> InWidget)
 	{
 		ChildSlot.DetachWidget();
 	}
+	Invalidate(EInvalidateWidgetReason::Volatility);
+}
+
+bool STooltipPresenter::ComputeVolatility() const
+{
+	return SPanel::ComputeVolatility() || ChildSlot.Num() > 0;
 }
 
 void STooltipPresenter::OnArrangeChildren(const FGeometry& AllottedGeometry, FArrangedChildren& ArrangedChildren) const
diff --git a/Engine/Source/Runtime/Slate/Public/Widgets/LayerManager/STooltipPresenter.h b/Engine/Source/Runtime/Slate/Public/Widgets/LayerManager/STooltipPresenter.h
--- a/Engine/Source/Runtime/Slate/Public/Widgets/LayerManager/STooltipPresenter.h
+++ b/Engine/Source/Runtime/Slate/Public/Widgets/LayerManager/STooltipPresenter.h
@@ -36,6 +36,7 @@ public:
 	void SetContent(TSharedPtr<SWidget> InWidget);
 
 private:
+	virtual bool ComputeVolatility() const override;
 	virtual void OnArrangeChildren(const FGeometry& AllottedGeometry, FArrangedChildren& ArrangedChildren) const override;
 
 	virtual FVector2D ComputeDesiredSize(float) const override;
diff --git a/Engine/Source/Runtime/SlateCore/Public/Layout/Children.h b/Engine/Source/Runtime/SlateCore/Public/Layout/Children.h
--- a/Engine/Source/Runtime/SlateCore/Public/Layout/Children.h
+++ b/Engine/Source/Runtime/SlateCore/Public/Layout/Children.h
@@ -312,12 +312,15 @@ private:
 public:
 	void AttachWidget(const TSharedPtr<SWidget>& InWidget)
 	{
-		WidgetPtr = InWidget;
-		GetOwner().Invalidate(EInvalidateWidgetReason::ChildOrder);
-
-		if (InWidget.IsValid() && InWidget != SNullWidget::NullWidget)
+		if (WidgetPtr != InWidget)
 		{
-			InWidget->AssignParentWidget(GetOwner().AsShared());
+			WidgetPtr = InWidget;
+			GetOwner().Invalidate(EInvalidateWidgetReason::ChildOrder);
+
+			if (InWidget.IsValid() && InWidget != SNullWidget::NullWidget)
+			{
+				InWidget->AssignParentWidget(GetOwner().AsShared());
+			}
 		}
 	}
 
