Fix TObservableArray::RemoveAtSwap ::begin and other fixes reported in review.

diff --git a/Engine/Source/Runtime/SlateCore/Private/Test/ObservableArrayTest.cpp b/Engine/Source/Runtime/SlateCore/Private/Test/ObservableArrayTest.cpp
--- a/Engine/Source/Runtime/SlateCore/Private/Test/ObservableArrayTest.cpp
+++ b/Engine/Source/Runtime/SlateCore/Private/Test/ObservableArrayTest.cpp
@@ -22,17 +22,15 @@
 	int32 Counter = 0;
 	TArray<int32> ExpectedArray;
 	FObservableArrayTest* Self = this;
-	auto TestExpectedAction = [Self , &ExpectedAction, &Counter, &ExpectedArray](::UE::Slate::Containers::TObservableArray<int32>::ObservableArrayChangedArgsType Args)
+	auto TestExpectedAction = [Self, &ExpectedAction, &Counter, &ExpectedArray](::UE::Slate::Containers::TObservableArray<int32>::ObservableArrayChangedArgsType Args)
 	{
 		Self->AddErrorIfFalse(Args.GetAction() == ExpectedAction, TEXT("The notification occurs with the wrong action"));
-		if (Args.GetAction() == Slate::Containers::EObservableArrayChangedAction::Add)
+		if (Args.GetAction() == Slate::Containers::EObservableArrayChangedAction::Add
+			|| Args.GetAction() == Slate::Containers::EObservableArrayChangedAction::Remove
+			|| Args.GetAction() == Slate::Containers::EObservableArrayChangedAction::RemoveSwap)
 		{
 			Self->AddErrorIfFalse(ExpectedArray == Args.GetItems(), TEXT("The notification occurs with the GetItems"));
 		}
-		else if (Args.GetAction() == Slate::Containers::EObservableArrayChangedAction::Remove)
-		{
-			Self->AddErrorIfFalse(ExpectedArray == Args.GetItems(), TEXT("The notification occurs with the wrong GetItmes"));
-		}
 		++Counter;
 	};
 
@@ -53,7 +51,7 @@
 
 					AddErrorIfFalse(ArrayValues == ObservableValues, TEXT("ObservableValues == ArrayValues"));
 					++ExpectedCounter;
-					AddErrorIfFalse(Counter == ExpectedCounter, "The delegate was not executed");
+					AddErrorIfFalse(Counter == ExpectedCounter, TEXT("The delegate was not executed"));
 				}
 				{
 					ExpectedArray.Reset();
@@ -64,7 +62,7 @@
 
 					AddErrorIfFalse(ArrayValues == ObservableValues, TEXT("ObservableValues == ArrayValues"));
 					++ExpectedCounter;
-					AddErrorIfFalse(Counter == ExpectedCounter, "The delegate was not executed");
+					AddErrorIfFalse(Counter == ExpectedCounter, TEXT("The delegate was not executed"));
 				}
 				{
 					ExpectedArray.Reset();
@@ -75,7 +73,7 @@
 
 					AddErrorIfFalse(ArrayValues == ObservableValues, TEXT("ObservableValues == ArrayValues"));
 					++ExpectedCounter;
-					AddErrorIfFalse(Counter == ExpectedCounter, "The delegate was not executed");
+					AddErrorIfFalse(Counter == ExpectedCounter, TEXT("The delegate was not executed"));
 				}
 				{
 					ExpectedArray.Reset();
@@ -86,7 +84,7 @@
 
 					AddErrorIfFalse(ArrayValues == ObservableValues, TEXT("ObservableValues == ArrayValues"));
 					++ExpectedCounter;
-					AddErrorIfFalse(Counter == ExpectedCounter, "The delegate was not executed");
+					AddErrorIfFalse(Counter == ExpectedCounter, TEXT("The delegate was not executed"));
 				}
 				{
 					ExpectedArray.Reset();
@@ -97,7 +95,7 @@
 
 					AddErrorIfFalse(ArrayValues == ObservableValues, TEXT("ObservableValues == ArrayValues"));
 					++ExpectedCounter;
-					AddErrorIfFalse(Counter == ExpectedCounter, "The delegate was not executed");
+					AddErrorIfFalse(Counter == ExpectedCounter, TEXT("The delegate was not executed"));
 				}
 				{
 					ExpectedArray.Reset();
@@ -108,7 +106,7 @@
 
 					AddErrorIfFalse(ArrayValues == ObservableValues, TEXT("ObservableValues == ArrayValues"));
 					++ExpectedCounter;
-					AddErrorIfFalse(Counter == ExpectedCounter, "The delegate was not executed");
+					AddErrorIfFalse(Counter == ExpectedCounter, TEXT("The delegate was not executed"));
 				}
 				{
 					ExpectedArray.Reset();
@@ -119,18 +117,54 @@
 
 					AddErrorIfFalse(ArrayValues == ObservableValues, TEXT("ObservableValues == ArrayValues"));
 					++ExpectedCounter;
-					AddErrorIfFalse(Counter == ExpectedCounter, "The delegate was not executed");
+					AddErrorIfFalse(Counter == ExpectedCounter, TEXT("The delegate was not executed"));
 				}
 			}
 			{
+				int32 ExpectedFirstIndex = INDEX_NONE;
+				int32 ExpectedSecondIndex = INDEX_NONE;
+				TArray<int32> PreviousObservableArray;
+				auto TestExpectedSwapAction = [Self, &ExpectedFirstIndex, &ExpectedSecondIndex, &ObservableValues, &PreviousObservableArray](::UE::Slate::Containers::TObservableArray<int32>::ObservableArrayChangedArgsType Args)
+				{
+					if (Args.GetAction() == Slate::Containers::EObservableArrayChangedAction::Swap)
+					{
+						int32 FirstIndex = Args.GetSwapIndex().FirstIndex;
+						int32 SecondIndex = Args.GetSwapIndex().SecondIndex;
+
+						Self->AddErrorIfFalse(FirstIndex != SecondIndex, TEXT("SecondIndex != FirstIndex."));
+						Self->AddErrorIfFalse(ObservableValues.IsValidIndex(FirstIndex) && ObservableValues.IsValidIndex(SecondIndex), TEXT("The swap index are invalid."));
+						Self->AddErrorIfFalse(ObservableValues[FirstIndex] == PreviousObservableArray[SecondIndex], TEXT("The swap of FirstIndex do not matches"));
+						Self->AddErrorIfFalse(ObservableValues[SecondIndex] == PreviousObservableArray[FirstIndex], TEXT("The swap of SecondIndex do not matches"));
+					}
+				};
+
 				ExpectedAction = ::UE::Slate::Containers::EObservableArrayChangedAction::Swap;
-				ExpectedArray.Reset();
+				FDelegateHandle SwapHandle = ObservableValues.OnArrayChanged().AddLambda(TestExpectedSwapAction);
+				{
+					ExpectedArray.Reset();
+					PreviousObservableArray = TArray<int32>(ObservableValues.GetData(), ObservableValues.Num());
+					ExpectedFirstIndex = 1;
+					ExpectedSecondIndex = 2;
 
-				ArrayValues.Swap(1, 2);
-				ObservableValues.Swap(1, 2);
-				AddErrorIfFalse(ArrayValues == ObservableValues, TEXT("ObservableValues == ArrayValues"));
-				++ExpectedCounter;
-				AddErrorIfFalse(Counter == ExpectedCounter, "The delegate was not executed");
+					ArrayValues.Swap(ExpectedFirstIndex, ExpectedSecondIndex);
+					ObservableValues.Swap(ExpectedFirstIndex, ExpectedSecondIndex);
+					AddErrorIfFalse(ArrayValues == ObservableValues, TEXT("ObservableValues == ArrayValues"));
+					++ExpectedCounter;
+					AddErrorIfFalse(Counter == ExpectedCounter, "The delegate was not executed");
+				}
+				{
+					ExpectedArray.Reset();
+					PreviousObservableArray = TArray<int32>(ObservableValues.GetData(), ObservableValues.Num());
+					ExpectedFirstIndex = 3;
+					ExpectedSecondIndex = 1;
+
+					ArrayValues.Swap(ExpectedFirstIndex, ExpectedSecondIndex);
+					ObservableValues.Swap(ExpectedFirstIndex, ExpectedSecondIndex);
+					AddErrorIfFalse(ArrayValues == ObservableValues, TEXT("ObservableValues == ArrayValues"));
+					++ExpectedCounter;
+					AddErrorIfFalse(Counter == ExpectedCounter, "The delegate was not executed");
+				}
+				ObservableValues.OnArrayChanged().Remove(SwapHandle);
 			}
 			{
 				ExpectedAction = ::UE::Slate::Containers::EObservableArrayChangedAction::Remove;
@@ -144,7 +178,7 @@
 
 					AddErrorIfFalse(ArrayValues == ObservableValues, TEXT("ObservableValues == ArrayValues"));
 					++ExpectedCounter;
-					AddErrorIfFalse(Counter == ExpectedCounter, "The delegate was not executed");
+					AddErrorIfFalse(Counter == ExpectedCounter, TEXT("The delegate was not executed"));
 				}
 				{
 					bool bContains = ObservableValues.Contains(2);
@@ -161,18 +195,7 @@
 					}
 					AddErrorIfFalse(bWasRemoved == bContains, TEXT("ObservableValues.Contains == ObservableValues.RemoveSingle"));
 					AddErrorIfFalse(ArrayValues == ObservableValues, TEXT("ObservableValues == ArrayValues"));
-					AddErrorIfFalse(Counter == ExpectedCounter, "The delegate was not executed");
+					AddErrorIfFalse(Counter == ExpectedCounter, TEXT("The delegate was not executed"));
 				}
 				{
 					ExpectedArray.Reset();
-					ExpectedArray.Add(3);
-
-					ArrayValues.RemoveSingleSwap(3);
-					ObservableValues.RemoveSingleSwap(3);
-
-					AddErrorIfFalse(ArrayValues == ObservableValues, TEXT("ObservableValues == ArrayValues"));
-					++ExpectedCounter;
-					AddErrorIfFalse(Counter == ExpectedCounter, "The delegate was not executed");
-				}
-				{
-					ExpectedArray.Reset();
@@ -183,7 +206,7 @@
 
 					AddErrorIfFalse(ArrayValues == ObservableValues, TEXT("ObservableValues == ArrayValues"));
 					++ExpectedCounter;
-					AddErrorIfFalse(Counter == ExpectedCounter, "The delegate was not executed");
+					AddErrorIfFalse(Counter == ExpectedCounter, TEXT("The delegate was not executed"));
 				}
 				{
 					ExpectedArray.Reset();
@@ -202,48 +225,7 @@
 
 					AddErrorIfFalse(ArrayValues == ObservableValues, TEXT("ObservableValues == ArrayValues"));
 					++ExpectedCounter;
-					AddErrorIfFalse(Counter == ExpectedCounter, "The delegate was not executed");
-				}
-				{
-					ExpectedArray.Reset();
-					ExpectedArray.Add(ObservableValues[1]);
-
-					ArrayValues.RemoveAtSwap(1);
-					ObservableValues.RemoveAtSwap(1);
-
-					AddErrorIfFalse(ArrayValues == ObservableValues, TEXT("ObservableValues == ArrayValues"));
-					++ExpectedCounter;
-					AddErrorIfFalse(Counter == ExpectedCounter, "The delegate was not executed");
-				}
-				{
-					ExpectedArray.Reset();
-					ExpectedArray.Add(ObservableValues[0]);
-
-					ArrayValues.RemoveAtSwap(0);
-					ObservableValues.RemoveAtSwap(0);
-
-					AddErrorIfFalse(ArrayValues == ObservableValues, TEXT("ObservableValues == ArrayValues"));
-					++ExpectedCounter;
-					AddErrorIfFalse(Counter == ExpectedCounter, "The delegate was not executed");
-				}
-				{
-					ExpectedArray.Reset();
-					int32 ToRemoveIndex = 1;
-					int32 NumberToRemove = 4;
-					for (int32 Index = ToRemoveIndex; Index < ToRemoveIndex + NumberToRemove; ++Index)
-					{
-						if (ObservableValues.IsValidIndex(Index))
-						{
-							ExpectedArray.Add(ObservableValues[Index]);
-						}
-					}
-
-					ArrayValues.RemoveAtSwap(ToRemoveIndex, NumberToRemove);
-					ObservableValues.RemoveAtSwap(ToRemoveIndex, NumberToRemove);
-
-					AddErrorIfFalse(ArrayValues == ObservableValues, TEXT("ObservableValues == ArrayValues"));
-					++ExpectedCounter;
-					AddErrorIfFalse(Counter == ExpectedCounter, "The delegate was not executed");
+					AddErrorIfFalse(Counter == ExpectedCounter, TEXT("The delegate was not executed"));
 				}
 			}
 			{
@@ -260,15 +242,124 @@
 
 				{
 					ExpectedArray.Reset();
-					ExpectedArray.Add(2);
+					ExpectedArray.Append({ 6, 6, 7, 7, 7, 8, 9, 10, 11, 12, 12 });
 
-					ArrayValues.Add(2);
-					ObservableValues.Add(2);
+					ArrayValues.Append(ExpectedArray);
+					ObservableValues.Append(ExpectedArray);
 
 					AddErrorIfFalse(ArrayValues == ObservableValues, TEXT("ObservableValues == ArrayValues"));
 					++ExpectedCounter;
-					AddErrorIfFalse(Counter == ExpectedCounter, "The delegate was not executed");
+					AddErrorIfFalse(Counter == ExpectedCounter, TEXT("The delegate was not executed"));
 				}
 			}
+			{
+				ExpectedAction = ::UE::Slate::Containers::EObservableArrayChangedAction::RemoveSwap;
+
+				bool bExpectPreviousMovedElemenentIndexToBeValid = true;
+				int32 ExpectedRemoveIndex = INDEX_NONE;
+				TArray<int32> ExpectedBeforeSwapArray;
+				auto TestExpectedSwapAction = [Self, &bExpectPreviousMovedElemenentIndexToBeValid , &ExpectedRemoveIndex, &ExpectedBeforeSwapArray, &ObservableValues](::UE::Slate::Containers::TObservableArray<int32>::ObservableArrayChangedArgsType Args)
+				{
+					if (Args.GetAction() == Slate::Containers::EObservableArrayChangedAction::RemoveSwap)
+					{
+						int32 RemoveIndex = Args.GetRemovedSwapIndexes().RemoveIndex;
+						int32 PreviousMovedElmenentIndex = Args.GetRemovedSwapIndexes().PreviousMovedElmenentIndex;
+
+						Self->AddErrorIfFalse(ExpectedRemoveIndex == RemoveIndex, TEXT("The RemoveIndex is not the expected value."));
+						if (ObservableValues.Num() > 0)
+						{
+							if (bExpectPreviousMovedElemenentIndexToBeValid)
+							{
+								Self->AddErrorIfFalse(ExpectedBeforeSwapArray.IsValidIndex(PreviousMovedElmenentIndex), TEXT("PreviousMovedElmenentIndex is invalid."));
+							}
+							else
+							{
+								Self->AddErrorIfFalse(!ExpectedBeforeSwapArray.IsValidIndex(PreviousMovedElmenentIndex), TEXT("PreviousMovedElmenentIndex is invalid."));
+							}
+
+							if (ObservableValues.IsValidIndex(RemoveIndex)
+								&& ExpectedBeforeSwapArray.IsValidIndex(PreviousMovedElmenentIndex))
+							{
+								Self->AddErrorIfFalse(ObservableValues[RemoveIndex] == ExpectedBeforeSwapArray[PreviousMovedElmenentIndex], TEXT("The swap index doesn't matches"));
+							}
+						}
+					}
+				};
+				FDelegateHandle RemoveSwapHandle = ObservableValues.OnArrayChanged().AddLambda(TestExpectedSwapAction);
+
+				{
+					ExpectedArray.Reset();
+					ExpectedArray.Append({7});
+					ExpectedBeforeSwapArray = TArray<int32>(ObservableValues.GetData(), ObservableValues.Num());
+					ExpectedRemoveIndex = ArrayValues.Find(7);
+					bExpectPreviousMovedElemenentIndexToBeValid = true;
+
+					ArrayValues.RemoveSingleSwap(7);
+					ObservableValues.RemoveSingleSwap(7);
+
+					AddErrorIfFalse(ArrayValues == ObservableValues, TEXT("ObservableValues == ArrayValues"));
+					++ExpectedCounter;
+					AddErrorIfFalse(Counter == ExpectedCounter, TEXT("The delegate was not executed"));
+				}
+				{
+					ExpectedArray.Reset();
+					ExpectedBeforeSwapArray = TArray<int32>(ObservableValues.GetData(), ObservableValues.Num());
+					ExpectedRemoveIndex = 0;
+					ExpectedArray.Add(ArrayValues[ExpectedRemoveIndex]);
+					bExpectPreviousMovedElemenentIndexToBeValid = true;
+
+					ArrayValues.RemoveAtSwap(ExpectedRemoveIndex);
+					ObservableValues.RemoveAtSwap(ExpectedRemoveIndex);
+
+					AddErrorIfFalse(ArrayValues == ObservableValues, TEXT("ObservableValues == ArrayValues"));
+					++ExpectedCounter;
+					AddErrorIfFalse(Counter == ExpectedCounter, TEXT("The delegate was not executed"));
+				}
+				{
+					ExpectedArray.Reset();
+					ExpectedBeforeSwapArray = TArray<int32>(ObservableValues.GetData(), ObservableValues.Num());
+					ExpectedRemoveIndex = 1;
+					for (int32 Index = 0; Index < 4; ++Index)
+					{
+						ExpectedArray.Add(ArrayValues[Index+ ExpectedRemoveIndex]);
+					}
+					bExpectPreviousMovedElemenentIndexToBeValid = true;
+
+					ArrayValues.RemoveAtSwap(ExpectedRemoveIndex, 4);
+					ObservableValues.RemoveAtSwap(ExpectedRemoveIndex, 4);
+
+					AddErrorIfFalse(ArrayValues == ObservableValues, TEXT("ObservableValues == ArrayValues"));
+					++ExpectedCounter;
+					AddErrorIfFalse(Counter == ExpectedCounter, TEXT("The delegate was not executed"));
+				}
+				{
+					ExpectedArray.Reset();
+					ExpectedBeforeSwapArray = TArray<int32>(ObservableValues.GetData(), ObservableValues.Num());
+					ExpectedRemoveIndex = ArrayValues.Num() - 2;
+					for (int32 Index = 0; Index < 2; ++Index)
+					{
+						ExpectedArray.Add(ArrayValues[Index+ ExpectedRemoveIndex]);
+					}
+					bExpectPreviousMovedElemenentIndexToBeValid = false;
+
+					ArrayValues.RemoveAtSwap(ExpectedRemoveIndex, 2);
+					ObservableValues.RemoveAtSwap(ExpectedRemoveIndex, 2);
+
+					AddErrorIfFalse(ArrayValues == ObservableValues, TEXT("ObservableValues == ArrayValues"));
+					++ExpectedCounter;
+					AddErrorIfFalse(Counter == ExpectedCounter, TEXT("The delegate was not executed"));
+				}
+
+				ObservableValues.OnArrayChanged().Remove(RemoveSwapHandle);
+			}
+			{
+				TArray<int32> ItteratingArrayValues;
+				ItteratingArrayValues.Reserve(ObservableValues.Num());
+				for (int32 Value : ObservableValues)
+				{
+					ItteratingArrayValues.Add(Value);
+				}
+				AddErrorIfFalse(ItteratingArrayValues == ArrayValues, TEXT("ItteratingArrayValues == ArrayValues after itterator."));
+			}
 		}
 	}
diff --git a/Engine/Source/Runtime/SlateCore/Public/Containers/ObservableArray.h b/Engine/Source/Runtime/SlateCore/Public/Containers/ObservableArray.h
--- a/Engine/Source/Runtime/SlateCore/Public/Containers/ObservableArray.h
+++ b/Engine/Source/Runtime/SlateCore/Public/Containers/ObservableArray.h
@@ -61,20 +61,20 @@ private:
 		return Result;
 	}
 
-	static TObservableArrayChangedArgs MakeRemoveAction(const ArrayViewType InRemoveItems, int32 InRemoveStartedIndex)
+	static TObservableArrayChangedArgs MakeRemoveAction(const ArrayViewType InRemovedItems, int32 InRemoveStartedIndex)
 	{
 		TObservableArrayChangedArgs Result;
-		Result.Array = InRemoveItems;
+		Result.Array = InRemovedItems;
 		Result.StartIndex = InRemoveStartedIndex;
 		Result.Action = EObservableArrayChangedAction::Remove;
 		check(Result.Array.Num() > 0 && Result.StartIndex >= 0);
 		return Result;
 	}
 	
-	static TObservableArrayChangedArgs MakeRemoveSwapAction(const ArrayViewType InRemoveItems, int32 InRemoveStartedIndex, int32 InPreviousMovedItemLocation)
+	static TObservableArrayChangedArgs MakeRemoveSwapAction(const ArrayViewType InRemovedItems, int32 InRemoveStartedIndex, int32 InPreviousMovedItemLocation)
 	{
 		TObservableArrayChangedArgs Result;
-		Result.Array = InRemoveItems;
+		Result.Array = InRemovedItems;
 		Result.StartIndex = InRemoveStartedIndex;
 		Result.MoveIndex = InPreviousMovedItemLocation;
 		Result.Action = EObservableArrayChangedAction::RemoveSwap;
@@ -122,13 +122,13 @@ public:
 			
 		}
 		/**
-		 * The removed elements index.
-		 * The removed index is still valid if the array is not empty.
-		 * The moved elements are now at that location (if any).
+		 * The index of the removed elements before their removal.
+		 * The index can now be out of bound.
+		 * The swapped elements (if any) are now at that location.
 		 */
 		SizeType RemoveIndex;
 		/**
-		 * The previous location of the elements (if any) before they moved to the new location.
+		 * The previous location of the swapped elements (if any) before the swap.
 		 * The index is not valid anymore.
 		 * Set to INDEX_NONE, if no element was moved.
 		 */
@@ -189,6 +189,8 @@ public:
 	using SizeType = typename ArrayType::SizeType;
 	using ElementType = typename ArrayType::ElementType;
 	using AllocatorType = typename ArrayType::AllocatorType;
+	using RangedForIteratorType = typename ArrayType::RangedForIteratorType;
+	using RangedForConstIteratorType = typename ArrayType::RangedForConstIteratorType;
 	using ObservableArrayChangedArgsType = TObservableArrayChangedArgs<InElementType>;
 	
 	DECLARE_MULTICAST_DELEGATE_OneParam(FArrayChangedDelegate, ObservableArrayChangedArgsType);
@@ -238,14 +240,14 @@ public:
 	}
 
 public:
-	int32 Add(const ElementType& Item)
+	SizeType Add(const ElementType& Item)
 	{
-		int32 NewIndex = Array.Add(Item);
+		SizeType NewIndex = Array.Add(Item);
 		ArrayChangedDelegate.Broadcast(ObservableArrayChangedArgsType::MakeAddAction({ GetData() + NewIndex, 1 }, NewIndex));
 		return NewIndex;
 	}
 
-	int32 Add(ElementType&& Item)
+	SizeType Add(ElementType&& Item)
 	{
 		int32 NewIndex = Array.Add(MoveTempIfPossible(Item));
 		ArrayChangedDelegate.Broadcast(ObservableArrayChangedArgsType::MakeAddAction({ GetData() + NewIndex, 1 }, NewIndex));
@@ -253,9 +255,9 @@ public:
 	}
 
 	template <typename... InArgTypes>
-	int32 Emplace(InArgTypes&&... Args)
+	SizeType Emplace(InArgTypes&&... Args)
 	{
-		int32 NewIndex = Array.Emplace(Forward<InArgTypes>(Args)...);
+		SizeType NewIndex = Array.Emplace(Forward<InArgTypes>(Args)...);
 		ArrayChangedDelegate.Broadcast(ObservableArrayChangedArgsType::MakeAddAction({ GetData() + NewIndex, 1 }, NewIndex));
 		return NewIndex;
 	}
@@ -276,9 +278,9 @@ public:
 	template <typename OtherElementType, typename OtherAllocatorType>
 	void Append(const TArray<OtherElementType, OtherAllocatorType>& Source)
 	{
-		SizeType PreviousNum = Array.Num();
-		if (PreviousNum > 0)
+		if (Source.Num() > 0)
 		{
+			SizeType PreviousNum = Array.Num();
 			Array.Append(Source);
 			SizeType NewNum = Array.Num();
 			ArrayChangedDelegate.Broadcast(ObservableArrayChangedArgsType::MakeAddAction({ GetData() + PreviousNum, NewNum - PreviousNum }, PreviousNum));
@@ -288,9 +290,9 @@ public:
 	template <typename OtherElementType, typename OtherAllocator>
 	void Append(TArray<OtherElementType, OtherAllocator>&& Source)
 	{
-		SizeType PreviousNum = Array.Num();
-		if (PreviousNum > 0)
+		if (Source.Num() > 0)
 		{
+			SizeType PreviousNum = Array.Num();
 			Array.Append(MoveTempIfPossible(Source));
 			SizeType NewNum = Array.Num();
 			ArrayChangedDelegate.Broadcast(ObservableArrayChangedArgsType::MakeAddAction({ GetData() + PreviousNum, NewNum - PreviousNum }, PreviousNum));
@@ -343,13 +345,13 @@ public:
 				FMemMark Mark(FMemStack::Get());
 				TArray<ElementType, TMemStackAllocator<>> RemovedElements;
 				RemovedElements.Reserve(NumToRemove);
-				for (int32 RemoveIndex = 0; RemoveIndex < NumToRemove; ++RemoveIndex)
+				for (SizeType RemoveIndex = 0; RemoveIndex < NumToRemove; ++RemoveIndex)
 				{
 					RemovedElements.Add(MoveTemp(Array[RemoveIndex + Index]));
 				}
 				
 				Array.RemoveAt(Index, NumToRemove, bAllowShrinking);
-				ArrayChangedDelegate.Broadcast(ObservableArrayChangedArgsType::MakeRemoveAction({ (ElementType*)RemovedElements.GetData(), NumToRemove }, Index));
+				ArrayChangedDelegate.Broadcast(ObservableArrayChangedArgsType::MakeRemoveAction({ RemovedElements.GetData(), NumToRemove }, Index));
 			}
 		}
 	}
@@ -366,12 +368,14 @@ public:
 		check((NumToRemove > 0) & (Index >= 0) & (Index + NumToRemove <= Num()));
 		if (NumToRemove > 0)
 		{
+			SizeType PreviousNum = Array.Num();
+			SizeType SwapAmount = FPlatformMath::Min(NumToRemove, PreviousNum - (Index + NumToRemove));
 			if (NumToRemove == 1)
 			{
 				ElementType RemovedElement = MoveTemp(Array[Index]);
 				
 				Array.RemoveAtSwap(Index, NumToRemove, bAllowShrinking);
-				ArrayChangedDelegate.Broadcast(ObservableArrayChangedArgsType::MakeRemoveAction({ &RemovedElement, 1 }, Index));
+				ArrayChangedDelegate.Broadcast(ObservableArrayChangedArgsType::MakeRemoveSwapAction({ &RemovedElement, 1 }, Index, PreviousNum - SwapAmount));
 			}
 			else
 			{
@@ -379,13 +383,13 @@ public:
 				FMemMark Mark(FMemStack::Get());
 				TArray<ElementType, TMemStackAllocator<>> RemovedElements;
 				RemovedElements.Reserve(NumToRemove);
-				for (int32 RemoveIndex = 0; RemoveIndex < NumToRemove; ++RemoveIndex)
+				for (SizeType RemoveIndex = 0; RemoveIndex < NumToRemove; ++RemoveIndex)
 				{
 					RemovedElements.Add(MoveTemp(Array[RemoveIndex + Index]));
 				}
-				
+
 				Array.RemoveAtSwap(Index, NumToRemove, bAllowShrinking);
-				ArrayChangedDelegate.Broadcast(ObservableArrayChangedArgsType::MakeRemoveAction({ RemovedElements.GetData(), NumToRemove }, Index));
+				ArrayChangedDelegate.Broadcast(ObservableArrayChangedArgsType::MakeRemoveSwapAction({ RemovedElements.GetData(), NumToRemove }, Index, PreviousNum - SwapAmount));
 			}
 		}
 	}
@@ -402,9 +406,10 @@ public:
 
 	void Reset(SizeType NewSize = 0)
 	{
-		if (Num())
+		SizeType PreviousNum = Array.Num();
+		Array.Reset(NewSize);
+		if (PreviousNum > 0)
 		{
-			Array.Reset(NewSize);
 			ArrayChangedDelegate.Broadcast(ObservableArrayChangedArgsType::MakeResetAction());
 		}
 	}
@@ -439,6 +444,16 @@ public:
 		return GetData()[Index];
 	}
 
+	ElementType* GetData()
+	{
+		return Array.GetData();
+	}
+
+	const ElementType* GetData() const
+	{
+		return Array.GetData();
+	}
+
 	template <typename ComparisonType>
 	bool Contains(const ComparisonType& Item) const
 	{
@@ -474,17 +489,6 @@ public:
 		return Array.IndexByPredicate(Pred);
 	}
 
-private:
-	ElementType* GetData()
-	{
-		return Array.GetData();
-	}
-
-	const ElementType* GetData() const
-	{
-		return Array.GetData();
-	}
-
 public:
 	template <typename InOtherAllocatorType>
 	bool operator==(TArray<ElementType, InOtherAllocatorType>& OtherArray) const
@@ -504,24 +508,24 @@ public:
 	}
 
 public:
-	ElementType* begin(TObservableArray& Arr)
+	RangedForIteratorType begin()
 	{
-		return Arr.Array.begin();
+		return Array.begin();
 	}
 
-	const ElementType* begin(const TObservableArray& Arr) const
+	RangedForConstIteratorType begin() const
 	{
-		return Arr.Array.begin();
+		return Array.begin();
 	}
 
-	ElementType* end(TObservableArray& Arr)
+	RangedForIteratorType end()
 	{
-		return Arr.Array.end();
+		return Array.end();
 	}
 
-	const ElementType* end(const TObservableArray& Arr) const
+	RangedForConstIteratorType end() const
 	{
-		return Arr.Array.end();
+		return Array.end();
 	}
 	
 private:
