Slate: Add debugging code to detect if we access a SWidget or a FChildren that is already destroyed.

diff --git a/Engine/Source/Runtime/SlateCore/Private/Widgets/SWidget.cpp b/Engine/Source/Runtime/SlateCore/Private/Widgets/SWidget.cpp
--- a/Engine/Source/Runtime/SlateCore/Private/Widgets/SWidget.cpp
+++ b/Engine/Source/Runtime/SlateCore/Private/Widgets/SWidget.cpp
@@ -245,6 +245,11 @@ PRAGMA_ENABLE_DEPRECATION_WARNINGS
 PRAGMA_DISABLE_DEPRECATION_WARNINGS
 SWidget::~SWidget()
 {
+#if WITH_SLATE_DEBUGGING
+	UE_CLOG(Debug_DestroyedTag != 0xDC, LogSlate, Fatal, TEXT("The widget is already destroyed."));
+	Debug_DestroyedTag = 0xA3;
+#endif
+
 #if UE_WITH_SLATE_DEBUG_FIND_WIDGET_REFLECTION_METADATA
 	if (FindWidgetMetaData::FoundWidget == this)
 	{
@@ -1543,7 +1548,7 @@ int32 SWidget::Paint(const FPaintArgs& Args, const FGeometry& AllottedGeometry,
 	// Detect children that should have been painted, but were skipped during the paint process.
 	// this will result in geometry being left on screen and not cleared, because it's visible, yet wasn't painted.
 #if WITH_SLATE_DEBUGGING
-	if (GSlateIsOnFastUpdatePath && GSlateEnsureAllVisibleWidgetsPaint)
+	if (GSlateIsOnFastUpdatePath && GSlateEnsureAllVisibleWidgetsPaint && !GIntraFrameDebuggingGameThread)
 	{
 		for (TWeakPtr<const SWidget>& DebugChildThatShouldHaveBeenPaintedPtr : DebugChildWidgetsToPaint)
 		{
@@ -1566,7 +1571,7 @@ int32 SWidget::Paint(const FPaintArgs& Args, const FGeometry& AllottedGeometry,
 		}
 	}
 
-	if (GSlateEnsureOutgoingLayerId)
+	if (GSlateEnsureOutgoingLayerId && !GIntraFrameDebuggingGameThread)
 	{
 		for (TWeakPtr<const SWidget>& DebugChildWeakPtr : DebugChildWidgetsToPaint)
 		{
diff --git a/Engine/Source/Runtime/SlateCore/Public/Layout/ChildrenBase.h b/Engine/Source/Runtime/SlateCore/Public/Layout/ChildrenBase.h
--- a/Engine/Source/Runtime/SlateCore/Public/Layout/ChildrenBase.h
+++ b/Engine/Source/Runtime/SlateCore/Public/Layout/ChildrenBase.h
@@ -5,9 +5,13 @@
 #include "CoreMinimal.h"
 #include "SlotBase.h"
 #include "SlateGlobals.h"
-#include "Debugging/SlateDebugging.h"
 #include "Types/ReflectionMetadata.h"
 
+#ifndef UE_WITH_SLATE_CHILDREN_DEBUGGING
+#define UE_WITH_SLATE_CHILDREN_DEBUGGING !(UE_BUILD_SHIPPING)
+#endif
+
+
 class SWidget;
 class FSlotBase;
 
@@ -51,14 +55,19 @@ public:
 	virtual TSharedRef<const SWidget> GetChildAt( int32 Index ) const = 0;
 
 	/** @return the SWidget that own the FChildren */
-	SWidget& GetOwner() const { return *Owner; }
+	SWidget& GetOwner() const
+	{
+		Debug_TestDestroyTag();
+		return *Owner;
+	}
 
 	/** Applies the predicate to all the widgets contained by the FChildren. */
 	template<typename Predicate>
 	void ForEachWidget(Predicate Pred)
 	{
-#if WITH_SLATE_DEBUGGING 
-		TGuardValue<bool> IteratingGuard(bIsIteratingChildren, true);
+#if UE_WITH_SLATE_CHILDREN_DEBUGGING 
+		Debug_TestDestroyTag();
+		TGuardValue<bool> IteratingGuard(Debug_bIsIteratingChildren, true);
 #endif
 
 		int32 WidgetCount = Num();
@@ -73,8 +82,9 @@ public:
 	template<typename Predicate>
 	void ForEachWidget(Predicate Pred) const
 	{
-#if WITH_SLATE_DEBUGGING 
-		TGuardValue<bool> IteratingGuard(bIsIteratingChildren, true);
+#if UE_WITH_SLATE_CHILDREN_DEBUGGING 
+		Debug_TestDestroyTag();
+		TGuardValue<bool> IteratingGuard(Debug_bIsIteratingChildren, true);
 #endif
 
 		int32 WidgetCount = Num();
@@ -103,6 +113,7 @@ public:
 	/** Option to give a name to Children to SlotAttribute purposes or for debugging. */
 	FName GetName() const
 	{
+		Debug_TestDestroyTag();
 		return Name;
 	}
 
@@ -181,18 +192,32 @@ protected:
 	/** @return ref to the Widget at the specified Index. */
 	virtual FConstWidgetRef GetChildRefAt(int32 Index) const = 0;
 
-
 protected:
-#if !WITH_SLATE_DEBUGGING 
-	virtual ~FChildren() = default;
-#else
+#if UE_WITH_SLATE_CHILDREN_DEBUGGING 
 	virtual ~FChildren()
 	{
-		UE_CLOG(bIsIteratingChildren, LogSlate, Error,
+		Debug_TestDestroyTag();
+		Debug_DestroyedTag = 0xA3;
+		UE_CLOG(Debug_bIsIteratingChildren, LogSlate, Fatal,
 			TEXT("Destroying widget while iterating children! Owner: %s [%s]"),
 			*FReflectionMetaData::GetWidgetDebugInfo(Owner),
 			*Name.ToString());
 	}
+
+	void Debug_TestDestroyTag() const
+	{
+		/** There is no guarantee that it will work. The memory could have been reused and the new data can luckily matches what we are looking for. */
+		UE_CLOG(Debug_DestroyedTag != 0xDC, LogSlate, Fatal, TEXT("The FChildren is destroyed. You probably have 1 widget owned by 2 different FChildren."));
+	}
+
+#else
+
+	virtual ~FChildren() = default;
+
+	void Debug_TestDestroyTag() const
+	{
+	}
+
 #endif // !WITH_SLATE_DEBUGGING
 
 protected:
@@ -202,8 +227,9 @@ protected:
 private:
 	FName Name;
 
-#if WITH_SLATE_DEBUGGING
-	mutable bool bIsIteratingChildren = false;
+#if UE_WITH_SLATE_CHILDREN_DEBUGGING
+	mutable bool Debug_bIsIteratingChildren = false;
+	mutable uint8 Debug_DestroyedTag = 0xDC;
 #endif // WITH_SLATE_DEBUGGING
 };
 PRAGMA_ENABLE_DEPRECATION_WARNINGS
diff --git a/Engine/Source/Runtime/SlateCore/Public/Widgets/SWidget.h b/Engine/Source/Runtime/SlateCore/Public/Widgets/SWidget.h
--- a/Engine/Source/Runtime/SlateCore/Public/Widgets/SWidget.h
+++ b/Engine/Source/Runtime/SlateCore/Public/Widgets/SWidget.h
@@ -1888,6 +1888,8 @@ private:
 #if WITH_SLATE_DEBUGGING
 	/** The last time this widget got painted. */
 	uint32 LastPaintFrame = 0;
+	/** Flag to help detect when we access an invalid Widget. */
+	uint8 Debug_DestroyedTag = 0xDC;
 #endif
 
 #if !(UE_BUILD_SHIPPING || UE_BUILD_TEST)
