CIS fix.

diff --git a/Engine/Source/Runtime/SlateCore/Public/Containers/ObservableArray.h b/Engine/Source/Runtime/SlateCore/Public/Containers/ObservableArray.h
--- a/Engine/Source/Runtime/SlateCore/Public/Containers/ObservableArray.h
+++ b/Engine/Source/Runtime/SlateCore/Public/Containers/ObservableArray.h
@@ -329,7 +329,7 @@ public:
 	template <typename CountType>
 	void RemoveAt(SizeType Index, CountType NumToRemove = 1, bool bAllowShrinking = true)
 	{
-		static_assert(!TAreTypesEqual<CountType, bool>::Value, "TObservableArray::RemoveAt: unexpected bool passed as the Count argument");
+		static_assert(!std::is_same_v<CountType, bool>, "TObservableArray::RemoveAt: unexpected bool passed as the Count argument");
 		check((NumToRemove > 0) & (Index >= 0) & (Index + NumToRemove <= Num()));
 		if (NumToRemove > 0)
 		{		
@@ -364,7 +364,7 @@ public:
 	template <typename CountType>
 	void RemoveAtSwap(SizeType Index, CountType NumToRemove = 1, bool bAllowShrinking = true)
 	{
-		static_assert(!TAreTypesEqual<CountType, bool>::Value, "TObservableArray::RemoveAtSwap: unexpected bool passed as the Count argument");
+		static_assert(!std::is_same_v<CountType, bool>, "TObservableArray::RemoveAtSwap: unexpected bool passed as the Count argument");
 		check((NumToRemove > 0) & (Index >= 0) & (Index + NumToRemove <= Num()));
 		if (NumToRemove > 0)
 		{
