Slate: Added native Make/Break aliases for FDeprecateSlateVector2D

diff --git a/Engine/Source/Runtime/SlateCore/Public/Types/SlateVector2.h b/Engine/Source/Runtime/SlateCore/Public/Types/SlateVector2.h
--- a/Engine/Source/Runtime/SlateCore/Public/Types/SlateVector2.h
+++ b/Engine/Source/Runtime/SlateCore/Public/Types/SlateVector2.h
@@ -113,7 +113,7 @@ namespace UE::Slate
  *
  * NOTE: DO NOT USE THIS TYPE DIRECTLY IN CLIENT CODE - ALL USAGES SHOULD USE FVector2f or float
  */
-USTRUCT(BlueprintInternalUseOnly, DisplayName="Vector2D", meta=(HiddenByDefault, ShortTooltip="Vector2D (single-precision)", ToolTip="Vector2D (single-precision)"))
+USTRUCT(BlueprintInternalUseOnly, DisplayName="Vector2D", meta=(HiddenByDefault, ShortTooltip="Vector2D (single-precision)", ToolTip="Vector2D (single-precision)", HasNativeMake="/Script/Engine.KismetMathLibrary.MakeVector2D", HasNativeBreak="/Script/Engine.KismetMathLibrary.BreakVector2D"))
 struct FDeprecateSlateVector2D : public FVector2f
 {
 	GENERATED_BODY()
