Check for any disabled parent before returning a destination widget in Slate navigation.

diff --git a/Engine/Source/Runtime/SlateCore/Private/Input/HittestGrid.cpp b/Engine/Source/Runtime/SlateCore/Private/Input/HittestGrid.cpp
--- a/Engine/Source/Runtime/SlateCore/Private/Input/HittestGrid.cpp
+++ b/Engine/Source/Runtime/SlateCore/Private/Input/HittestGrid.cpp
@@ -299,6 +299,22 @@ bool FHittestGrid::IsDescendantOf(const SWidget* ParentWidget, const FWidgetData
 	return false;
 }
 
+namespace UE::Slate::Private
+{
+	bool IsParentsEnabled(const SWidget* Widget)
+	{
+		while (Widget)
+		{
+			if (!Widget->IsEnabled())
+			{
+				return false;
+			}
+			Widget = Widget->Advanced_GetPaintParentWidget().Get();
+		}
+		return true;
+	}
+}
+
 #if WITH_SLATE_DEBUGGING
 namespace HittestGridDebuggingText
 {
@@ -308,6 +324,7 @@ namespace HittestGridDebuggingText
 	static FText PreviousWidgetIsBetter = LOCTEXT("StatePreviousWidgetIsBetter", "Previous Widget was better"); //~ The widget would be valid but the previous valid is closer
 	static FText NotADescendant = LOCTEXT("StateNotADescendant", "Not a descendant"); //~ We have a non escape boundary condition and the widget isn't a descendant of our boundary
 	static FText Disabled = LOCTEXT("StateNotEnabled", "Disabled"); //~ The widget is not enabled
+	static FText ParentDisabled = LOCTEXT("StateParentNotEnabled", "ParentDisabled"); //~ A parent of the widget is disabled
 	static FText DoesNotSuportKeyboardFocus = LOCTEXT("StateDoesNotSuportKeyboardFocus", "Keyboard focus unsupported"); //~ THe widget does not support keyboard focus
 }
 	#define AddToNextFocusableWidgetCondidateDebugResults(Candidate, Result) { if (IntermediateResultsPtr) { IntermediateResultsPtr->Emplace((Candidate), (Result)); } }
@@ -316,7 +333,7 @@ namespace HittestGridDebuggingText
 #endif
 
 template<typename TCompareFunc, typename TSourceSideFunc, typename TDestSideFunc>
-TSharedPtr<SWidget> FHittestGrid::FindFocusableWidget(FSlateRect WidgetRect, const FSlateRect SweptRect, int32 AxisIndex, int32 Increment, const EUINavigation Direction, const FNavigationReply& NavigationReply, TCompareFunc CompareFunc, TSourceSideFunc SourceSideFunc, TDestSideFunc DestSideFunc, int32 UserIndex, TArray<FDebuggingFindNextFocusableWidgetArgs::FWidgetResult>* IntermediateResultsPtr) const
+TSharedPtr<SWidget> FHittestGrid::FindFocusableWidget(FSlateRect WidgetRect, const FSlateRect SweptRect, int32 AxisIndex, int32 Increment, const EUINavigation Direction, const FNavigationReply& NavigationReply, TCompareFunc CompareFunc, TSourceSideFunc SourceSideFunc, TDestSideFunc DestSideFunc, int32 UserIndex, TArray<FDebuggingFindNextFocusableWidgetArgs::FWidgetResult>* IntermediateResultsPtr, TSet<TSharedPtr<SWidget>>* DisabledDestinations) const
 {
 	FIntPoint CurrentCellPoint = GetCellCoordinate(WidgetRect.GetCenter());
 
@@ -414,6 +431,12 @@ TSharedPtr<SWidget> FHittestGrid::FindFocusableWidget(FSlateRect WidgetRect, con
 					continue;
 				}
 
+				if (DisabledDestinations->Contains(TestWidget))
+				{
+					AddToNextFocusableWidgetCondidateDebugResults(TestWidget, HittestGridDebuggingText::ParentDisabled);
+					continue;
+				}
+
 				BestWidgetRect = TestCandidateRect;
 				BestWidget = TestWidget;
 				AddToNextFocusableWidgetCondidateDebugResults(TestWidget, HittestGridDebuggingText::Valid);
@@ -451,7 +474,18 @@ TSharedPtr<SWidget> FHittestGrid::FindFocusableWidget(FSlateRect WidgetRect, con
 				}
 			}
 
-			return BestWidget;
+			// Make sure all parents of the chosen widget are enabled before returning.
+			// Note that IsParentsEnabled is a costly function. We call it here as the last step to minimize the number of calls to it.
+			if (!UE::Slate::Private::IsParentsEnabled(BestWidget.Get()))
+			{	
+				// Find the next best widget because this one has disabled parents.
+				DisabledDestinations->Add(BestWidget);
+				return FindFocusableWidget(WidgetRect, SweptRect, AxisIndex, Increment, Direction, NavigationReply, CompareFunc, SourceSideFunc, DestSideFunc, UserIndex, IntermediateResultsPtr, DisabledDestinations);
+			}
+			else
+			{
+				return BestWidget;
+			}
 		}
 
 		// break if we have looped back to where we started.
@@ -510,6 +544,8 @@ TSharedPtr<SWidget> FHittestGrid::FindNextFocusableWidget(const FArrangedWidget&
 	TArray<FDebuggingFindNextFocusableWidgetArgs::FWidgetResult>* IntermediateResultsPtr = nullptr;
 #endif
 
+	TSet<TSharedPtr<SWidget>> DisabledDestinations = TSet<TSharedPtr<SWidget>>();
+
 	switch (Direction)
 	{
 	case EUINavigation::Left:
@@ -521,7 +557,7 @@ TSharedPtr<SWidget> FHittestGrid::FindNextFocusableWidget(const FArrangedWidget&
 			[](float A, float B) { return A - 0.1f < B; }, // Compare function
 			[](FSlateRect SourceRect) { return SourceRect.Left; }, // Source side function
 			[](FSlateRect DestRect) { return DestRect.Right; }, // Dest side function
-			UserIndex, IntermediateResultsPtr);
+			UserIndex, IntermediateResultsPtr, &DisabledDestinations);
 		break;
 	case EUINavigation::Right:
 		SweptWidgetRect.Left = BoundingRuleRect.Left;
@@ -532,7 +568,7 @@ TSharedPtr<SWidget> FHittestGrid::FindNextFocusableWidget(const FArrangedWidget&
 			[](float A, float B) { return A + 0.1f > B; }, // Compare function
 			[](FSlateRect SourceRect) { return SourceRect.Right; }, // Source side function
 			[](FSlateRect DestRect) { return DestRect.Left; }, // Dest side function
-			UserIndex, IntermediateResultsPtr);
+			UserIndex, IntermediateResultsPtr, &DisabledDestinations);
 		break;
 	case EUINavigation::Up:
 		SweptWidgetRect.Top = BoundingRuleRect.Top;
@@ -543,7 +579,7 @@ TSharedPtr<SWidget> FHittestGrid::FindNextFocusableWidget(const FArrangedWidget&
 			[](float A, float B) { return A - 0.1f < B; }, // Compare function
 			[](FSlateRect SourceRect) { return SourceRect.Top; }, // Source side function
 			[](FSlateRect DestRect) { return DestRect.Bottom; }, // Dest side function
-			UserIndex, IntermediateResultsPtr);
+			UserIndex, IntermediateResultsPtr, &DisabledDestinations);
 		break;
 	case EUINavigation::Down:
 		SweptWidgetRect.Top = BoundingRuleRect.Top;
@@ -554,7 +590,7 @@ TSharedPtr<SWidget> FHittestGrid::FindNextFocusableWidget(const FArrangedWidget&
 			[](float A, float B) { return A + 0.1f > B; }, // Compare function
 			[](FSlateRect SourceRect) { return SourceRect.Bottom; }, // Source side function
 			[](FSlateRect DestRect) { return DestRect.Top; }, // Dest side function
-			UserIndex, IntermediateResultsPtr);
+			UserIndex, IntermediateResultsPtr, &DisabledDestinations);
 		break;
 
 	default:
diff --git a/Engine/Source/Runtime/SlateCore/Public/Input/HittestGrid.h b/Engine/Source/Runtime/SlateCore/Public/Input/HittestGrid.h
--- a/Engine/Source/Runtime/SlateCore/Public/Input/HittestGrid.h
+++ b/Engine/Source/Runtime/SlateCore/Public/Input/HittestGrid.h
@@ -272,7 +272,7 @@ private:
 
 	/** Utility function for searching for the next focusable widget. */
 	template<typename TCompareFunc, typename TSourceSideFunc, typename TDestSideFunc>
-	TSharedPtr<SWidget> FindFocusableWidget(const FSlateRect WidgetRect, const FSlateRect SweptRect, int32 AxisIndex, int32 Increment, const EUINavigation Direction, const FNavigationReply& NavigationReply, TCompareFunc CompareFunc, TSourceSideFunc SourceSideFunc, TDestSideFunc DestSideFunc, int32 UserIndex, TArray<FDebuggingFindNextFocusableWidgetArgs::FWidgetResult>* IntermediatedResultPtr) const;
+	TSharedPtr<SWidget> FindFocusableWidget(const FSlateRect WidgetRect, const FSlateRect SweptRect, int32 AxisIndex, int32 Increment, const EUINavigation Direction, const FNavigationReply& NavigationReply, TCompareFunc CompareFunc, TSourceSideFunc SourceSideFunc, TDestSideFunc DestSideFunc, int32 UserIndex, TArray<FDebuggingFindNextFocusableWidgetArgs::FWidgetResult>* IntermediatedResultPtr, TSet<TSharedPtr<SWidget>>* DisabledDestinations) const;
 
 	/** Constrains a float position into the grid coordinate. */
 	FIntPoint GetCellCoordinate(UE::Slate::FDeprecateVector2DParameter Position) const;
