Slate: Additional SlateCore FVector2D deprecation

diff --git a/Engine/Source/Runtime/SlateCore/Private/Fonts/FontCacheUtils.h b/Engine/Source/Runtime/SlateCore/Private/Fonts/FontCacheUtils.h
--- a/Engine/Source/Runtime/SlateCore/Private/Fonts/FontCacheUtils.h
+++ b/Engine/Source/Runtime/SlateCore/Private/Fonts/FontCacheUtils.h
@@ -7,7 +7,7 @@
 #include "SlateGlobals.h"
 
 typedef FString KeyType;
-typedef FVector2D ValueType;
+typedef FVector2f ValueType;
 
 DECLARE_MEMORY_STAT_EXTERN(TEXT("Font Measure Memory"), STAT_SlateFontMeasureCacheMemory, STATGROUP_SlateMemory, SLATECORE_API);
 
diff --git a/Engine/Source/Runtime/SlateCore/Private/Fonts/FontMeasure.cpp b/Engine/Source/Runtime/SlateCore/Private/Fonts/FontMeasure.cpp
--- a/Engine/Source/Runtime/SlateCore/Private/Fonts/FontMeasure.cpp
+++ b/Engine/Source/Runtime/SlateCore/Private/Fonts/FontMeasure.cpp
@@ -56,20 +56,20 @@ FSlateFontMeasure::FSlateFontMeasure( const TSharedRef<class FSlateFontCache>& I
 
 }
 
-FVector2D FSlateFontMeasure::Measure( const FString& Text, const FSlateFontInfo& InFontInfo, float FontScale ) const
+UE::Slate::FDeprecateVector2DResult FSlateFontMeasure::Measure( const FString& Text, const FSlateFontInfo& InFontInfo, float FontScale ) const
 {
 	int32 DummyLastCharacterIndex;
 	return MeasureStringInternal( Text, 0, Text.Len(), InFontInfo, false, FontScale, INDEX_NONE, ELastCharacterIndexFormat::Unused, DummyLastCharacterIndex);
 }
 
-FVector2D FSlateFontMeasure::Measure( const FText& Text, const FSlateFontInfo& InFontInfo, float FontScale ) const
+UE::Slate::FDeprecateVector2DResult FSlateFontMeasure::Measure( const FText& Text, const FSlateFontInfo& InFontInfo, float FontScale ) const
 {
 	int32 DummyLastCharacterIndex;
 	const FString& String = Text.ToString();
 	return MeasureStringInternal( String, 0, String.Len(), InFontInfo, false, FontScale, INDEX_NONE, ELastCharacterIndexFormat::Unused, DummyLastCharacterIndex);
 }
 
-FVector2D FSlateFontMeasure::Measure( const FString& Text, int32 StartIndex, int32 EndIndex, const FSlateFontInfo& InFontInfo, bool IncludeKerningWithPrecedingChar, float FontScale ) const
+UE::Slate::FDeprecateVector2DResult FSlateFontMeasure::Measure( const FString& Text, int32 StartIndex, int32 EndIndex, const FSlateFontInfo& InFontInfo, bool IncludeKerningWithPrecedingChar, float FontScale ) const
 {
 	int32 DummyLastCharacterIndex;
 	return MeasureStringInternal( Text, StartIndex, EndIndex, InFontInfo, IncludeKerningWithPrecedingChar, FontScale, INDEX_NONE, ELastCharacterIndexFormat::Unused, DummyLastCharacterIndex);
@@ -137,7 +137,7 @@ int32 FSlateFontMeasure::FindCharacterIndexAtOffset( const FString& Text, int32
 	return FoundLastCharacterIndex;
 }
 
-FVector2D FSlateFontMeasure::MeasureStringInternal( const FString& Text, int32 StartIndex, int32 EndIndex, const FSlateFontInfo& InFontInfo, bool IncludeKerningWithPrecedingChar, float FontScale, int32 StopAfterHorizontalOffset, ELastCharacterIndexFormat CharIndexFormat, int32& OutLastCharacterIndex ) const
+FVector2f FSlateFontMeasure::MeasureStringInternal( const FString& Text, int32 StartIndex, int32 EndIndex, const FSlateFontInfo& InFontInfo, bool IncludeKerningWithPrecedingChar, float FontScale, int32 StopAfterHorizontalOffset, ELastCharacterIndexFormat CharIndexFormat, int32& OutLastCharacterIndex ) const
 {
 	QUICK_SCOPE_CYCLE_COUNTER(STAT_Slate_MeasureStringInternal)
 	FCharacterList& CharacterList = FontCache->GetCharacterList( InFontInfo, FontScale );
@@ -148,7 +148,7 @@ FVector2D FSlateFontMeasure::MeasureStringInternal( const FString& Text, int32 S
 	const int32 TextRangeLength = EndIndex - StartIndex;
 	if ( Text.IsEmpty() || EndIndex - StartIndex <= 0 || EndIndex <= 0 || StartIndex < 0 || EndIndex <= StartIndex )
 	{
-		return FVector2D( 0, MaxHeight );
+		return FVector2f( 0.f, MaxHeight );
 	}
 
 #if USE_MEASURE_CACHING
@@ -160,7 +160,7 @@ FVector2D FSlateFontMeasure::MeasureStringInternal( const FString& Text, int32 S
 
 		if ( CurrentMeasureCache )
 		{
-			const FVector2D* CachedMeasurement = CurrentMeasureCache->AccessItem( Text );
+			const FVector2f* CachedMeasurement = CurrentMeasureCache->AccessItem( Text );
 			if( CachedMeasurement )
 			{
 				return *CachedMeasurement;
@@ -170,7 +170,7 @@ FVector2D FSlateFontMeasure::MeasureStringInternal( const FString& Text, int32 S
 #endif
 
 	// The size of the string
-	FVector2D Size(0,0);
+	FVector2f Size(0.f,0.f);
 	// Widest line encountered while drawing this text.
 	int32 MaxLineWidth = 0;
 	// The width of the current line so far.
diff --git a/Engine/Source/Runtime/SlateCore/Public/Fonts/FontMeasure.h b/Engine/Source/Runtime/SlateCore/Public/Fonts/FontMeasure.h
--- a/Engine/Source/Runtime/SlateCore/Public/Fonts/FontMeasure.h
+++ b/Engine/Source/Runtime/SlateCore/Public/Fonts/FontMeasure.h
@@ -3,20 +3,19 @@
 
 #include "CoreMinimal.h"
 #include "Fonts/FontTypes.h"
+#include "Types/SlateVector2.h"
 
 struct FSlateFontMeasureCache;
 
 typedef class FLRUStringCache FMeasureCache;
 struct FSlateFontMeasureCache;
 
-class SLATECORE_API FSlateFontMeasure 
+class SLATECORE_API FSlateFontMeasure final
 {
 public:
 
 	static TSharedRef< FSlateFontMeasure > Create( const TSharedRef<class FSlateFontCache>& FontCache );
 
-	virtual ~FSlateFontMeasure() {}
-
 public:
 
 	/** 
@@ -27,7 +26,7 @@ public:
 	 * @param FontScale	The scale to apply to the font
 	 * @return The width and height of the string.
 	 */
-	virtual FVector2D Measure( const FString& Text, const FSlateFontInfo &InFontInfo, float FontScale = 1.0f ) const;
+	UE::Slate::FDeprecateVector2DResult Measure( const FString& Text, const FSlateFontInfo &InFontInfo, float FontScale = 1.0f ) const;
 
 	/** 
 	 * Measures the width and height of a passed in text.  The height is the maximum height of the largest character in the font/size pair
@@ -37,7 +36,7 @@ public:
 	 * @param FontScale	The scale to apply to the font
 	 * @return The width and height of the text.
 	 */
-	virtual FVector2D Measure( const FText& Text, const FSlateFontInfo &InFontInfo, float FontScale = 1.0f ) const;
+	UE::Slate::FDeprecateVector2DResult Measure( const FText& Text, const FSlateFontInfo &InFontInfo, float FontScale = 1.0f ) const;
 	
 	/** 
 	 * Measures the width and height of a passed in string.  The height is the maximum height of the largest character in the font/size pair
@@ -49,7 +48,7 @@ public:
 	 * @param FontScale	The scale to apply to the font
 	 * @return The width and height of the string.
 	 */
-	virtual FVector2D Measure( const FString& Text, int32 StartIndex, int32 EndIndex, const FSlateFontInfo &InFontInfo, bool IncludeKerningWithPrecedingChar = true, float FontScale = 1.0f ) const;
+	UE::Slate::FDeprecateVector2DResult Measure( const FString& Text, int32 StartIndex, int32 EndIndex, const FSlateFontInfo &InFontInfo, bool IncludeKerningWithPrecedingChar = true, float FontScale = 1.0f ) const;
 
 	/** 
 	 * Finds the last whole character index before the specified position in pixels along the string horizontally
@@ -132,7 +131,7 @@ public:
 	 *
 	 * @return The index of the character closest to the specified horizontal offset
 	 */
-	virtual int32 FindCharacterIndexAtOffset( const FString& Text, const FSlateFontInfo &InFontInfo, const int32 HorizontalOffset, float FontScale = 1.0f ) const;
+	int32 FindCharacterIndexAtOffset( const FString& Text, const FSlateFontInfo &InFontInfo, const int32 HorizontalOffset, float FontScale = 1.0f ) const;
 
 	/** 
 	 * Finds the character index at the specified position in pixels along the text horizontally
@@ -145,7 +144,7 @@ public:
 	 *
 	 * @return The index of the character closest to the specified horizontal offset
 	 */
-	virtual int32 FindCharacterIndexAtOffset( const FText& Text, const FSlateFontInfo &InFontInfo, const int32 HorizontalOffset, float FontScale = 1.0f ) const;
+	int32 FindCharacterIndexAtOffset( const FText& Text, const FSlateFontInfo &InFontInfo, const int32 HorizontalOffset, float FontScale = 1.0f ) const;
 
 	/** 
 	 * Finds the character index at the specified position in pixels along the string horizontally
@@ -160,7 +159,7 @@ public:
 	 *
 	 * @return The index of the character closest to the specified horizontal offset
 	 */
-	virtual int32 FindCharacterIndexAtOffset( const FString& Text, int32 StartIndex, int32 EndIndex, const FSlateFontInfo &InFontInfo, const int32 HorizontalOffset, bool IncludeKerningWithPrecedingChar = true, float FontScale = 1.0f ) const;
+	int32 FindCharacterIndexAtOffset( const FString& Text, int32 StartIndex, int32 EndIndex, const FSlateFontInfo &InFontInfo, const int32 HorizontalOffset, bool IncludeKerningWithPrecedingChar = true, float FontScale = 1.0f ) const;
 
 	/**
 	 * Returns the height of the largest character in the font. 
@@ -226,7 +225,7 @@ private:
 	 *
 	 * @return The width and height of the string.
 	 */
-	FVector2D MeasureStringInternal( const FString& Text, int32 StartIndex, int32 EndIndex, const FSlateFontInfo& InFontInfo, bool IncludeKerningWithPrecedingChar, float FontScale, int32 StopAfterHorizontalOffset, ELastCharacterIndexFormat CharIndexFormat, int32& OutLastCharacterIndex ) const;
+	FVector2f MeasureStringInternal( const FString& Text, int32 StartIndex, int32 EndIndex, const FSlateFontInfo& InFontInfo, bool IncludeKerningWithPrecedingChar, float FontScale, int32 StopAfterHorizontalOffset, ELastCharacterIndexFormat CharIndexFormat, int32& OutLastCharacterIndex ) const;
 
 	/**
 	 * Check to see if there's an existing cached measurement, or failing that, add a new entry so that we can cache a new measurement
diff --git a/Engine/Source/Runtime/SlateCore/Public/Rendering/SlateDrawBuffer.h b/Engine/Source/Runtime/SlateCore/Public/Rendering/SlateDrawBuffer.h
--- a/Engine/Source/Runtime/SlateCore/Public/Rendering/SlateDrawBuffer.h
+++ b/Engine/Source/Runtime/SlateCore/Public/Rendering/SlateDrawBuffer.h
@@ -5,6 +5,7 @@
 #include "CoreMinimal.h"
 #include "UObject/GCObject.h"
 #include "HAL/CriticalSection.h"
+#include "Types/SlateVector2.h"
 
 #include <atomic>
 
@@ -88,5 +89,5 @@ private:
 	uint32 ResourceVersion;
 
 public:
-	FVector2D ViewOffset;
+	FDeprecateSlateVector2D ViewOffset;
 };
