Check for disabled descendants when moving focus to a widget using tab.

diff --git a/Engine/Source/Runtime/SlateCore/Private/Layout/WidgetPath.cpp b/Engine/Source/Runtime/SlateCore/Private/Layout/WidgetPath.cpp
--- a/Engine/Source/Runtime/SlateCore/Private/Layout/WidgetPath.cpp
+++ b/Engine/Source/Runtime/SlateCore/Private/Layout/WidgetPath.cpp
@@ -224,15 +224,24 @@ bool FWidgetPath::MoveFocus(int32 PathLevel, EUINavigation NavigationType, bool
 					const bool bFoundNextFocusable = (PathToFocusableChild.Num() > 0) || ArrangedChildren[FocusedChildIndex].Widget->SupportsKeyboardFocus();
 					if (bFoundNextFocusable)
 					{
-						// We found the next focusable widget, so make this path point at the new widget by:
-						// First, truncating the FocusPath up to the current level (i.e. PathLevel).
-						Widgets.Remove(PathLevel + 1, Widgets.Num() - PathLevel - 1);
-						// Second, add the immediate child that is focus or whose descendant is focused.
-						Widgets.AddWidget(ArrangedChildren[FocusedChildIndex]);
-						// Add path to focused descendants if any.
-						Widgets.Append(PathToFocusableChild);
-						// We successfully moved focus!
-						return true;
+						// Check if any descendant is disabled.
+						int32 DisabledDescendantIndex = PathToFocusableChild.IndexOfByPredicate([](const FArrangedWidget& ArrangedWidget) 
+							{ 
+								return !ArrangedWidget.Widget->IsEnabled(); 
+							});
+
+						if (DisabledDescendantIndex == INDEX_NONE)
+						{
+							// We found the next focusable widget, so make this path point at the new widget by:
+							// First, truncating the FocusPath up to the current level (i.e. PathLevel).
+							Widgets.Remove(PathLevel + 1, Widgets.Num() - PathLevel - 1);
+							// Second, add the immediate child that is focus or whose descendant is focused.
+							Widgets.AddWidget(ArrangedChildren[FocusedChildIndex]);
+							// Add path to focused descendants if any.
+							Widgets.Append(PathToFocusableChild);
+							// We successfully moved focus!
+							return true;
+						}
 					}
 				}
 			}
