Fix SlateRect calculating area incorrectly.

diff --git a/Engine/Source/Runtime/SlateCore/Public/Layout/SlateRect.h b/Engine/Source/Runtime/SlateCore/Public/Layout/SlateRect.h
--- a/Engine/Source/Runtime/SlateCore/Public/Layout/SlateRect.h
+++ b/Engine/Source/Runtime/SlateCore/Public/Layout/SlateRect.h
@@ -65,11 +65,11 @@ public:
 	}
 
 	/**
-	 * @return true, if the rectangle has a size of 0.
+	 * @return true, if the rectangle has an effective size of 0.
 	 */
 	FORCEINLINE bool IsEmpty() const
 	{
-		return GetSize().SizeSquared() == 0.0f;
+		return GetArea() <= UE_SMALL_NUMBER;
 	}
 
 	/**
