Slate: Fix small double to float cast

diff --git a/Engine/Source/Runtime/SlateCore/Private/Debugging/ConsoleSlateDebuggerInvalidationRoot.cpp b/Engine/Source/Runtime/SlateCore/Private/Debugging/ConsoleSlateDebuggerInvalidationRoot.cpp
--- a/Engine/Source/Runtime/SlateCore/Private/Debugging/ConsoleSlateDebuggerInvalidationRoot.cpp
+++ b/Engine/Source/Runtime/SlateCore/Private/Debugging/ConsoleSlateDebuggerInvalidationRoot.cpp
@@ -189,10 +189,10 @@ void FConsoleSlateDebuggerInvalidationRoot::HandlePaintDebugInfo(const FPaintArg
 	if (bShowLegend)
 	{
 		//const TSharedRef<FSlateFontMeasure>& FontMeasureService = FSlateApplicationBase::Get().GetRenderer()->GetFontMeasureService();
-		//FVector2D FontSize = FontMeasureService->Measure(TEXT("No Paint occurred"), FontInfo);
+		//FVector2f FontSize = FontMeasureService->Measure(TEXT("No Paint occurred"), FontInfo);
 		//FontSize.Y *= 3;
 		//const FSlateBrush* BoxBrush = FCoreStyle::Get().GetBrush("WhiteBrush");
-		//const FGeometry Geometry = FGeometry::MakeRoot(FontSize, FSlateLayoutTransform(1.f, FVector2D(10.f, 10.f + 0.f)));
+		//const FGeometry Geometry = FGeometry::MakeRoot(FontSize, FSlateLayoutTransform(1.f, FVector2f(10.f, 10.f + 0.f)));
 
 		//FSlateDrawElement::MakeBox(
 		//	InOutDrawElements,
diff --git a/Engine/Source/Runtime/SlateCore/Private/Debugging/WidgetList.cpp b/Engine/Source/Runtime/SlateCore/Private/Debugging/WidgetList.cpp
--- a/Engine/Source/Runtime/SlateCore/Private/Debugging/WidgetList.cpp
+++ b/Engine/Source/Runtime/SlateCore/Private/Debugging/WidgetList.cpp
@@ -156,7 +156,7 @@ void LogAllWidgetsDebugInfoImpl(FOutputDevice& Ar, const FLogAllWidgetsDebugInfo
 
 		if (DebugInfoFlags.bPaint)
 		{
-			FVector2D AbsoluteSize = Widget->GetPersistentState().AllottedGeometry.GetAbsoluteSize();
+			FVector2f AbsoluteSize = Widget->GetPersistentState().AllottedGeometry.GetAbsoluteSize();
 			MessageBuilder << TEXT(";");
 			MessageBuilder << Widget->Debug_GetLastPaintFrame();
 			MessageBuilder << TEXT(";");
diff --git a/Engine/Source/Runtime/SlateCore/Public/Brushes/SlateRoundedBoxBrush.h b/Engine/Source/Runtime/SlateCore/Public/Brushes/SlateRoundedBoxBrush.h
--- a/Engine/Source/Runtime/SlateCore/Public/Brushes/SlateRoundedBoxBrush.h
+++ b/Engine/Source/Runtime/SlateCore/Public/Brushes/SlateRoundedBoxBrush.h
@@ -35,11 +35,14 @@ struct SLATECORE_API FSlateRoundedBoxBrush
 		: FSlateRoundedBoxBrush(InColor, UE::Slate::FDeprecateVector2DParameter(InImageSize))
 	{ 
 	}
+
+#if UE_ENABLE_SLATE_VECTOR_DEPRECATION_MECHANISMS
 	template<typename FillColorType>
 	FORCENOINLINE FSlateRoundedBoxBrush(const FillColorType& InColor, const FVector2D& InImageSize)
 		: FSlateRoundedBoxBrush(InColor, UE::Slate::FDeprecateVector2DParameter(UE::Slate::CastToVector2f(InImageSize)))
 	{ 
 	}
+#endif
 	template<typename FillColorType>
 	FORCENOINLINE FSlateRoundedBoxBrush(const FillColorType& InColor, const FVector2f& InImageSize)
 		: FSlateRoundedBoxBrush(InColor, UE::Slate::FDeprecateVector2DParameter(InImageSize))
diff --git a/Engine/Source/Runtime/SlateCore/Public/Types/Attributes/SlateAttributeBase.inl b/Engine/Source/Runtime/SlateCore/Public/Types/Attributes/SlateAttributeBase.inl
--- a/Engine/Source/Runtime/SlateCore/Public/Types/Attributes/SlateAttributeBase.inl
+++ b/Engine/Source/Runtime/SlateCore/Public/Types/Attributes/SlateAttributeBase.inl
@@ -70,19 +70,19 @@ namespace SlateAttributePrivate
 			}
 			else if constexpr (std::is_same<FVector2f, InObjectType>::value)
 			{
-				ensureMsgf(!Value.ContainsNaN(), TEXT("Value contains a NaN. Initialize your FVector2f properly (see FVector2D::EForceInit)"));
+				ensureMsgf(!Value.ContainsNaN(), TEXT("Value contains a NaN. Initialize your FVector2f properly (see FVector2f::EForceInit)"));
 			}
 			else if constexpr (std::is_same<FVector2d, InObjectType>::value)
 			{
-				ensureMsgf(!Value.ContainsNaN(), TEXT("Value contains a NaN. Initialize your FVector2d properly (see FVector2D::EForceInit)"));
+				ensureMsgf(!Value.ContainsNaN(), TEXT("Value contains a NaN. Initialize your FVector2d properly (see FVector2d::EForceInit)"));
 			}
 			else if constexpr (std::is_same<FVector3f, InObjectType>::value)
 			{
-				ensureMsgf(!Value.ContainsNaN(), TEXT("Value contains a NaN. Initialize your FVector3f properly (see FVector::EForceInit)"));
+				ensureMsgf(!Value.ContainsNaN(), TEXT("Value contains a NaN. Initialize your FVector3f properly (see FVector3f::EForceInit)"));
 			}
 			else if constexpr (std::is_same<FVector3d, InObjectType>::value)
 			{
-				ensureMsgf(!Value.ContainsNaN(), TEXT("Value contains a NaN. Initialize your FVector3d properly (see FVector::EForceInit)"));
+				ensureMsgf(!Value.ContainsNaN(), TEXT("Value contains a NaN. Initialize your FVector3d properly (see FVector3d::EForceInit)"));
 			}
 			else if constexpr (std::is_same<FVector4f, InObjectType>::value)
 			{
