StarshipCoreStyle: SecondaryToolbars now copy the style directly from primary toolbars

diff --git a/Engine/Source/Runtime/SlateCore/Private/Styling/StarshipCoreStyle.cpp b/Engine/Source/Runtime/SlateCore/Private/Styling/StarshipCoreStyle.cpp
--- a/Engine/Source/Runtime/SlateCore/Private/Styling/StarshipCoreStyle.cpp
+++ b/Engine/Source/Runtime/SlateCore/Private/Styling/StarshipCoreStyle.cpp
@@ -2229,6 +2229,12 @@ void FStarshipCoreStyle::SetupMultiboxStyles(TSharedRef<FStyle>& Style)
 
 		// Used by empty toolbars and misc widgets to have the same background color
 		Style->Set("AssetEditorToolbar.Background", SlimToolbarBackground);
+		
+		// Styling for secondary toolbars that appear below the main toolbar, currently the same as the main toolbar
+		{
+			FToolBarStyle SecondaryToolbarStyle(SlimToolbarStyle);
+			Style->Set("SecondaryToolbar", SecondaryToolbarStyle);
+		}
 
 		// Callout Toolbar - Used to "call out" the toolbar button with text
 		{
@@ -2237,75 +2243,6 @@ void FStarshipCoreStyle::SetupMultiboxStyles(TSharedRef<FStyle>& Style)
 			Style->Set("CalloutToolbar", SlimToolbarStyle);
 		}
 
-		// Styling for secondary toolbars that appear below the main toolbar
-		{
-			FButtonStyle SecondaryToolBarButton = FButtonStyle(SimpleButtonStyle)
-			.SetNormalForeground(FStyleColors::Foreground)
-			.SetPressedForeground(FStyleColors::ForegroundHover)
-			.SetHoveredForeground(FStyleColors::ForegroundHover)
-			.SetDisabledForeground(FStyleColors::Foreground)
-			.SetNormalPadding(FMargin(8.f, 4.f, 8.f, 4.f))
-			.SetPressedPadding(FMargin(8.f, 5.f, 8.f, 3.f));
-
-			FComboButtonStyle SecondaryToolBarComboButton = FComboButtonStyle(Style->GetWidgetStyle<FComboButtonStyle>("SimpleComboButton"))
-				.SetDownArrowPadding(FMargin(1.f, 0.f, 2.f, 0.f));
-
-			
-			FToolBarStyle SecondaryToolbarStyle =
-				FToolBarStyle()
-				.SetShowLabels(false)
-				.SetBackground(*SlimToolbarBackground)
-				.SetBackgroundPadding(FMargin(4.f, 6.f, 0.f, 6.f))
-				.SetExpandBrush(IMAGE_BRUSH("Icons/toolbar_expand_16x", Icon16x16))
-				.SetComboButtonPadding(FMargin(0.0f, 0.0f))
-				.SetButtonPadding(FMargin(4.0f, 0.0f))
-				.SetCheckBoxPadding(FMargin(8.0f, 0.0f))
-				.SetSeparatorBrush(FSlateColorBrush(FStyleColors::Recessed))
-				.SetSeparatorPadding(FMargin(4.f, -5.0f))
-				.SetLabelStyle(FTextBlockStyle(NormalText))
-				.SetComboButtonStyle(SlimToolBarComboButton)
-				.SetLabelPadding(FMargin(4.f, 0.f, 0.f, 0.f))
-				.SetEditableTextStyle(FEditableTextBoxStyle(NormalEditableTextBoxStyle));
-
-			FCheckBoxStyle SecondaryToolBarToggleButtonCheckBoxStyle = FCheckBoxStyle(Style->GetWidgetStyle<FCheckBoxStyle>("ToggleButtonCheckbox"))
-				.SetPadding(FMargin(8.f, 4.f, 8.f, 4.f));
-
-			SecondaryToolbarStyle.SetToggleButtonStyle(SecondaryToolBarToggleButtonCheckBoxStyle);
-			SecondaryToolbarStyle.SetButtonStyle(SecondaryToolBarButton);
-
-			FButtonStyle SecondaryToolBarButtonSettingsRegion = FButtonStyle(SimpleButtonStyle)
-				.SetPressed(FSlateNoResource())
-				.SetHovered(FSlateNoResource())
-				.SetNormal(FSlateNoResource())
-				.SetNormalForeground(FStyleColors::Foreground)
-				.SetPressedForeground(FStyleColors::ForegroundHover)
-				.SetHoveredForeground(FStyleColors::ForegroundHover)
-				.SetDisabledForeground(FStyleColors::Foreground)
-				.SetNormalPadding(FMargin(8.f, 4.f, 8.f, 4.f))
-				.SetPressedPadding(FMargin(8.f, 5.f, 8.f, 3.f));
-
-
-			FCheckBoxStyle SecondaryToolbarSettingsToggleButtonStyle = FCheckBoxStyle(Style->GetWidgetStyle<FCheckBoxStyle>("TransparentCheckBox"))
-				.SetPadding(FMargin(8.f, 4.f, 8.f, 4.f));
-
-			SecondaryToolbarStyle.SetSettingsButtonStyle(SecondaryToolBarButtonSettingsRegion);
-			SecondaryToolbarStyle.SetSettingsToggleButtonStyle(SecondaryToolbarSettingsToggleButtonStyle);
-			SecondaryToolbarStyle.SetIconSize(Icon20x20);
-
-			SecondaryToolBarButtonSettingsRegion
-				.SetNormalPadding(FMargin(4.f, 4.f, 4.f, 4.f))
-				.SetPressedPadding(FMargin(4.f, 5.f, 4.f, 3.f));
-
-			FComboButtonStyle SecondaryToolBarSettingsComboButton = FComboButtonStyle(SecondaryToolBarComboButton)
-				.SetContentPadding(0)
-				.SetButtonStyle(SecondaryToolBarButtonSettingsRegion)
-				.SetDownArrowImage(IMAGE_BRUSH_SVG("Starship/Common/ellipsis-vertical-narrow", FVector2D(6, 20)));
-
-			SecondaryToolbarStyle.SetSettingsComboButtonStyle(SecondaryToolBarSettingsComboButton);
-			
-			Style->Set("SecondaryToolbar", SecondaryToolbarStyle);
-		}
-
 	}
 
 	// MenuBar
