UMG/Slate - A bit of tweaking to the names and such, but otherwise looks good.  Integrating the snapping to pixel functionality to work throughout the hierarchy, works with SGI.  Is of the utmost importance to being able to do indicators over players heads and such that feel buttery smooth, and small icons on a minimap - of course at the cost of sharpness, but it's a treadoff people should be allowed to make.

PR #8570: Exposing Slate Snap To Pixel setting in UMG (Contributed by [REDACTED_PERSON])

diff --git a/Engine/Source/Runtime/SlateCore/Private/FastUpdate/WidgetProxy.cpp b/Engine/Source/Runtime/SlateCore/Private/FastUpdate/WidgetProxy.cpp
--- a/Engine/Source/Runtime/SlateCore/Private/FastUpdate/WidgetProxy.cpp
+++ b/Engine/Source/Runtime/SlateCore/Private/FastUpdate/WidgetProxy.cpp
@@ -247,6 +247,12 @@ FWidgetProxy::FUpdateResult FWidgetProxy::Repaint(const FPaintArgs& PaintArgs, F
 			}
 		}
 	}
+
+	if (MyState.InitialPixelSnappingMethod != EWidgetPixelSnapping::Inherit)
+	{
+		OutDrawElements.PushPixelSnappingMethod(MyState.InitialPixelSnappingMethod);
+	}
+	
 	const int32 NewLayerId = WidgetPtr->Paint(UpdatedArgs, MyState.AllottedGeometry, MyState.CullingBounds, OutDrawElements, MyState.LayerId, MyState.WidgetStyle, MyState.bParentEnabled);
 
 	PaintArgs.GetHittestGrid().SetUserIndex(PrevUserIndex);
@@ -258,6 +264,11 @@ FWidgetProxy::FUpdateResult FWidgetProxy::Repaint(const FPaintArgs& PaintArgs, F
 		check(StartingClipIndex == OutDrawElements.GetClippingIndex());
 	}
 	
+	if (MyState.InitialPixelSnappingMethod != EWidgetPixelSnapping::Inherit)
+	{
+		OutDrawElements.PopPixelSnappingMethod();
+	}
+
 	OutDrawElements.SetIsInGameLayer(MyState.bIsInGameLayer);
 	return FUpdateResult{ PrevLayerId, NewLayerId };
 }
diff --git a/Engine/Source/Runtime/SlateCore/Private/Rendering/DrawElements.cpp b/Engine/Source/Runtime/SlateCore/Private/Rendering/DrawElements.cpp
--- a/Engine/Source/Runtime/SlateCore/Private/Rendering/DrawElements.cpp
+++ b/Engine/Source/Runtime/SlateCore/Private/Rendering/DrawElements.cpp
@@ -177,6 +177,23 @@ void FSlateDrawElement::Init(FSlateWindowElementList& ElementList, EElementType
 		InDrawEffects &= ~ESlateDrawEffect::DisabledEffect;
 	}
 
+	// Apply the pixel snapping effects, if it's unset/inherited we don't do anything.
+	switch (ElementList.GetPixelSnappingMethod())
+	{
+		case EWidgetPixelSnapping::Disabled:
+			InDrawEffects |= ESlateDrawEffect::NoPixelSnapping;
+			break;
+		case EWidgetPixelSnapping::Inherit:
+			// NOTE: Do nothing on inherit.
+		case EWidgetPixelSnapping::SnapToPixel:
+			// NOTE: We don't remove the drawing effects pixel snapping here.  The default behavior of InDrawEffects is always to snap to pixels
+			// so all we need to do here is not to snap.  This will allow the continued behavior that if someone does set snap to pixel at a
+			// higher level, it won't break people that directly set InDrawEffects |= ESlateDrawEffect::NoPixelSnapping, before drawing their
+			// draw element.
+		default:
+			break;
+	}
+
 	RenderTransform = PaintGeometry.GetAccumulatedRenderTransform();
 	Position = PaintGeometry.DrawPosition;
 	Scale = PaintGeometry.DrawScale;
@@ -995,6 +1012,24 @@ void FSlateWindowElementList::PopClipToStackIndex(int32 Index)
 	ClippingManager.PopToStackIndex(Index);
 }
 
+int32 FSlateWindowElementList::PushPixelSnappingMethod(EWidgetPixelSnapping InPixelSnappingMethod)
+{
+	PixelSnappingMethodStack.Add(InPixelSnappingMethod);
+	return PixelSnappingMethodStack.Num();
+}
+
+void FSlateWindowElementList::PopPixelSnappingMethod()
+{
+	if (ensure(PixelSnappingMethodStack.Num() > 0))
+	{
+		PixelSnappingMethodStack.Pop();
+	}
+}
+
+EWidgetPixelSnapping FSlateWindowElementList::GetPixelSnappingMethod() const
+{
+	return (PixelSnappingMethodStack.Num() > 0) ? PixelSnappingMethodStack.Top() : EWidgetPixelSnapping::Inherit;
+}
 
 void FSlateWindowElementList::SetRenderTargetWindow(SWindow* InRenderTargetWindow)
 {
diff --git a/Engine/Source/Runtime/SlateCore/Private/Widgets/SWidget.cpp b/Engine/Source/Runtime/SlateCore/Private/Widgets/SWidget.cpp
--- a/Engine/Source/Runtime/SlateCore/Private/Widgets/SWidget.cpp
+++ b/Engine/Source/Runtime/SlateCore/Private/Widgets/SWidget.cpp
@@ -210,6 +210,7 @@ SWidget::SWidget()
 	, AccessibleSummaryBehavior(EAccessibleBehavior::Auto)
 #endif
 	, Clipping(EWidgetClipping::Inherit)
+	, PixelSnappingMethod(EWidgetPixelSnapping::Inherit)
 	, FlowDirectionPreference(EFlowDirectionPreference::Inherit)
 	// Note we are defaulting to tick for backwards compatibility
 	, UpdateFlags(EWidgetUpdateFlags::NeedsTick)
@@ -326,6 +327,7 @@ void SWidget::Construct(
 	Args._RenderOpacity = InRenderOpacity;
 	Args._ForceVolatile = InForceVolatile;
 	Args._Clipping = InClipping;
+	Args._PixelSnappingMethod = EWidgetPixelSnapping::Inherit;
 	Args._FlowDirectionPreference = InFlowPreference;
 	Args._RenderTransform = InTransform;
 	Args._RenderTransformPivot = InTransformPivot;
@@ -1215,6 +1217,26 @@ void SWidget::SetVisibility(TAttribute<EVisibility> InVisibility)
 	VisibilityAttribute.Assign(*this, MoveTemp(InVisibility));
 }
 
+void SWidget::SetClipping(EWidgetClipping InClipping)
+{
+	if (Clipping != InClipping)
+	{
+		Clipping = InClipping;
+		OnClippingChanged();
+		// @todo - Fast path should this be Paint?
+		Invalidate(EInvalidateWidgetReason::Layout);
+	}
+}
+
+void SWidget::SetPixelSnapping(EWidgetPixelSnapping InPixelSnappingMethod)
+{
+	if (PixelSnappingMethod != InPixelSnappingMethod)
+	{
+		PixelSnappingMethod = InPixelSnappingMethod;
+		Invalidate(EInvalidateWidget::Paint);
+	}
+}
+
 bool SWidget::IsFastPathVisible() const
 {
 	return FastPathProxyHandle.GetWidgetVisibility(this).IsVisible();
@@ -1472,6 +1494,7 @@ int32 SWidget::Paint(const FPaintArgs& Args, const FGeometry& AllottedGeometry,
 	PersistentState.DesktopGeometry = DesktopSpaceGeometry;
 	PersistentState.WidgetStyle = InWidgetStyle;
 	PersistentState.CullingBounds = MyCullingRect;
+	PersistentState.InitialPixelSnappingMethod = OutDrawElements.GetPixelSnappingMethod();
 
 	const int32 IncomingUserIndex = Args.GetHittestGrid().GetUserIndex();
 	ensure(IncomingUserIndex <= std::numeric_limits<int8>::max()); // shorten to save memory
@@ -1507,6 +1530,12 @@ int32 SWidget::Paint(const FPaintArgs& Args, const FGeometry& AllottedGeometry,
 		OutDrawElements.PushClip(ClippingZone);
 	}
 
+	const bool bNewPixelSnappingMethod = PixelSnappingMethod != EWidgetPixelSnapping::Inherit;
+	
+	if (bNewPixelSnappingMethod)
+	{
+		OutDrawElements.PushPixelSnappingMethod(PixelSnappingMethod);
+	}
 
 #if WITH_SLATE_DEBUGGING
 	FSlateDebugging::BeginWidgetPaint.Broadcast(this, UpdatedArgs, AllottedGeometry, CullingBounds, OutDrawElements, LayerId);
@@ -1624,6 +1653,10 @@ int32 SWidget::Paint(const FPaintArgs& Args, const FGeometry& AllottedGeometry,
 		OutDrawElements.PopClip();
 	}
 
+	if (bNewPixelSnappingMethod)
+	{
+		OutDrawElements.PopPixelSnappingMethod();
+	}
 
 #if PLATFORM_UI_NEEDS_FOCUS_OUTLINES
 	// Check if we need to show the keyboard focus ring, this is only necessary if the widget could be focused.
diff --git a/Engine/Source/Runtime/SlateCore/Public/FastUpdate/WidgetProxy.h b/Engine/Source/Runtime/SlateCore/Public/FastUpdate/WidgetProxy.h
--- a/Engine/Source/Runtime/SlateCore/Public/FastUpdate/WidgetProxy.h
+++ b/Engine/Source/Runtime/SlateCore/Public/FastUpdate/WidgetProxy.h
@@ -243,6 +243,7 @@ struct FSlateWidgetPersistentState
 		, OutgoingLayerId(0)
 		, IncomingUserIndex(INDEX_NONE)
 		, IncomingFlowDirection(EFlowDirection::LeftToRight)
+		, InitialPixelSnappingMethod(EWidgetPixelSnapping::Inherit)
 		, bParentEnabled(true)
 		, bInheritedHittestability(false)
 		, bDeferredPainting(false)
@@ -261,6 +262,7 @@ struct FSlateWidgetPersistentState
 	int32 OutgoingLayerId;
 	int8 IncomingUserIndex;
 	EFlowDirection IncomingFlowDirection;
+	EWidgetPixelSnapping InitialPixelSnappingMethod;
 	uint8 bParentEnabled : 1;
 	uint8 bInheritedHittestability : 1;
 	uint8 bDeferredPainting : 1;
diff --git a/Engine/Source/Runtime/SlateCore/Public/Rendering/DrawElements.h b/Engine/Source/Runtime/SlateCore/Public/Rendering/DrawElements.h
--- a/Engine/Source/Runtime/SlateCore/Public/Rendering/DrawElements.h
+++ b/Engine/Source/Runtime/SlateCore/Public/Rendering/DrawElements.h
@@ -19,6 +19,7 @@
 #include "Rendering/SlateRenderBatch.h"
 #include "DrawElementTextOverflowArgs.h"
 #include "ElementBatcher.h"
+#include "Widgets/WidgetPixelSnapping.h"
 #include "Types/SlateVector2.h"
 
 #if UE_ENABLE_INCLUDE_ORDER_DEPRECATED_IN_5_2
@@ -579,6 +580,14 @@ public:
 	FSlateClippingManager& GetClippingManager() { return ClippingManager; }
 	const FSlateClippingManager& GetClippingManager() const { return ClippingManager; }
 
+	//--------------------------------------------------------------------------
+	// PIXEL SNAPPING
+	//--------------------------------------------------------------------------
+
+	SLATECORE_API int32 PushPixelSnappingMethod(EWidgetPixelSnapping InPixelSnappingMethod);
+	SLATECORE_API void PopPixelSnappingMethod();
+	SLATECORE_API EWidgetPixelSnapping GetPixelSnappingMethod() const;
+	
 	//--------------------------------------------------------------------------
 	// DEFERRED PAINTING
 	//--------------------------------------------------------------------------
@@ -692,6 +701,9 @@ private:
 	/**  */
 	FSlateClippingManager ClippingManager;
 
+	/** Manages what pixel snapping method should be applied */
+	TArray<EWidgetPixelSnapping, TInlineAllocator<4>> PixelSnappingMethodStack;
+	
 	/**
 	 * Some widgets want their logical children to appear at a different "layer" in the physical hierarchy.
 	 * We accomplish this by deferring their painting.
diff --git a/Engine/Source/Runtime/SlateCore/Public/Widgets/DeclarativeSyntaxSupport.h b/Engine/Source/Runtime/SlateCore/Public/Widgets/DeclarativeSyntaxSupport.h
--- a/Engine/Source/Runtime/SlateCore/Public/Widgets/DeclarativeSyntaxSupport.h
+++ b/Engine/Source/Runtime/SlateCore/Public/Widgets/DeclarativeSyntaxSupport.h
@@ -6,6 +6,7 @@
 #include "Misc/Attribute.h"
 #include "Layout/Visibility.h"
 #include "Layout/Clipping.h"
+#include "Widgets/WidgetPixelSnapping.h"
 #include "Layout/FlowDirection.h"
 #include "Rendering/SlateRenderTransform.h"
 #include "GenericPlatform/ICursor.h"
@@ -655,6 +656,7 @@ struct FSlateBaseNamedArgs
 	/** If true, bound Slate Attributes will be updated once per frame. */
 	SLATE_PRIVATE_ARGUMENT_VARIABLE(bool, EnabledAttributesUpdate) = true;
 	SLATE_PRIVATE_ARGUMENT_VARIABLE(EWidgetClipping, Clipping) = EWidgetClipping::Inherit;
+	SLATE_PRIVATE_ARGUMENT_VARIABLE(EWidgetPixelSnapping, PixelSnappingMethod) = EWidgetPixelSnapping::Inherit;
 	SLATE_PRIVATE_ARGUMENT_VARIABLE(EFlowDirectionPreference, FlowDirectionPreference) = EFlowDirectionPreference::Inherit;
 	SLATE_PRIVATE_ARGUMENT_VARIABLE(float, RenderOpacity) = 1.f;
 	SLATE_PRIVATE_ATTRIBUTE_VARIABLE(TOptional<FSlateRenderTransform>, RenderTransform);
@@ -679,6 +681,7 @@ struct TSlateBaseNamedArgs : public FSlateBaseNamedArgs
 	SLATE_PRIVATE_ATTRIBUTE_FUNCTION(EVisibility, Visibility)
 	SLATE_PRIVATE_ARGUMENT_FUNCTION(bool, ForceVolatile)
 	SLATE_PRIVATE_ARGUMENT_FUNCTION(EWidgetClipping, Clipping)
+	SLATE_PRIVATE_ARGUMENT_FUNCTION(EWidgetPixelSnapping, PixelSnappingMethod)
 	SLATE_PRIVATE_ARGUMENT_FUNCTION(EFlowDirectionPreference, FlowDirectionPreference)
 	SLATE_PRIVATE_ARGUMENT_FUNCTION(float, RenderOpacity)
 	SLATE_PRIVATE_ATTRIBUTE_FUNCTION(TOptional<FSlateRenderTransform>, RenderTransform)
diff --git a/Engine/Source/Runtime/SlateCore/Public/Widgets/SWidget.h b/Engine/Source/Runtime/SlateCore/Public/Widgets/SWidget.h
--- a/Engine/Source/Runtime/SlateCore/Public/Widgets/SWidget.h
+++ b/Engine/Source/Runtime/SlateCore/Public/Widgets/SWidget.h
@@ -32,6 +32,7 @@
 #include "InvalidateWidgetReason.h"
 #include "Widgets/SlateControlledConstruction.h"
 #include "Widgets/Accessibility/SlateWidgetAccessibleTypes.h"
+#include "WidgetPixelSnapping.h"
 
 class FActiveTimerHandle;
 class FArrangedChildren;
@@ -1252,22 +1253,24 @@ public:
 	/**
 	 * Sets the clipping to bounds rules for this widget.
 	 */
-	FORCEINLINE void SetClipping(EWidgetClipping InClipping)
-	{
-		if (Clipping != InClipping)
-		{
-			Clipping = InClipping;
-			OnClippingChanged();
-			// @todo - Fast path should this be Paint?
-			Invalidate(EInvalidateWidgetReason::Layout);
-		}
-	}
+	void SetClipping(EWidgetClipping InClipping);
 
 	/** @return The current clipping rules for this widget. */
 	FORCEINLINE EWidgetClipping GetClipping() const
 	{
 		return Clipping;
 	}
+	
+	/**
+	* Sets the pixel snapping method for this widget.
+	*/
+	void SetPixelSnapping(EWidgetPixelSnapping InPixelSnappingMethod);
+
+	/** @return The current pixel snapping rules for this widget. */
+	FORCEINLINE EWidgetPixelSnapping GetPixelSnapping() const
+	{
+		return PixelSnappingMethod;
+	}
 
 	/**
 	 * Sets an additional culling padding that is added to a widget to give more leeway when culling widgets.  Useful if 
@@ -1787,6 +1790,13 @@ protected:
 	 * Set to true if all content of the widget should clip to the bounds of this widget.
 	 */
 	EWidgetClipping Clipping;
+	
+	/**
+	 * When set to EPixelSnappingMethod::SnapToPixel, the widget is drawn at the nearest pixel. Will improve sharpness 
+	 * but could show a stepping effect when moved in an animation.  By default everything in slate is Inherit, and the default
+	 * state all things inherit is SnapToPixel.
+	 */
+	EWidgetPixelSnapping PixelSnappingMethod;
 
 protected:
 	/** Establishes a new flow direction potentially, if this widget has a particular preference for it and all its children. */
diff --git a/Engine/Source/Runtime/SlateCore/Public/Widgets/WidgetPixelSnapping.h b/Engine/Source/Runtime/SlateCore/Public/Widgets/WidgetPixelSnapping.h
new file mode 100644
--- /dev/null
+++ b/Engine/Source/Runtime/SlateCore/Public/Widgets/WidgetPixelSnapping.h
@@ -0,0 +1,24 @@
+// Copyright Epic Games, Inc. All Rights Reserved.
+
+#pragma once
+
+#include "CoreMinimal.h"
+#include "UObject/ObjectMacros.h"
+
+#include "WidgetPixelSnapping.generated.h"
+
+/**
+ * The different states of pixel snapping a widget can be in.
+ */
+UENUM()
+enum class EWidgetPixelSnapping : uint8
+{
+	/** Inherits the snapping method set by the parent widget. */
+	Inherit,
+
+	/** Draws the widget without snapping. Useful during animations or moving indicators. */
+	Disabled,
+
+	/** Draws the widget at the nearest pixel. Improves sharpness of widgets. */
+	SnapToPixel
+};
diff --git a/Engine/Source/Runtime/UMG/Private/Components/Widget.cpp b/Engine/Source/Runtime/UMG/Private/Components/Widget.cpp
--- a/Engine/Source/Runtime/UMG/Private/Components/Widget.cpp
+++ b/Engine/Source/Runtime/UMG/Private/Components/Widget.cpp
@@ -191,6 +191,7 @@ PRAGMA_DISABLE_DEPRECATION_WARNINGS
 	RenderOpacity = 1.0f;
 	RenderTransformPivot = FVector2D(0.5f, 0.5f);
 	Cursor = EMouseCursor::Default;
+	PixelSnapping = EWidgetPixelSnapping::Inherit;
 
 #if WITH_EDITORONLY_DATA
 	bOverrideAccessibleDefaults = false;
@@ -457,6 +458,28 @@ void UWidget::SetClipping(EWidgetClipping InClipping)
 }
 PRAGMA_ENABLE_DEPRECATION_WARNINGS
 
+EWidgetPixelSnapping UWidget::GetPixelSnapping() const
+{
+	TSharedPtr<SWidget> SafeWidget = GetCachedWidget();
+	if (SafeWidget.IsValid())
+	{
+		return SafeWidget->GetPixelSnapping();
+	}
+
+	return PixelSnapping;
+}
+
+void UWidget::SetPixelSnapping(EWidgetPixelSnapping InPixelSnappingMethod)
+{
+	PixelSnapping = InPixelSnappingMethod;
+
+	TSharedPtr<SWidget> SafeWidget = GetCachedWidget();
+	if (SafeWidget.IsValid())
+	{
+		SafeWidget->SetPixelSnapping(InPixelSnappingMethod);
+	}
+}
+
 void UWidget::ForceVolatile(bool bForce)
 {
 	bIsVolatile = bForce;
@@ -1382,6 +1405,7 @@ PRAGMA_DISABLE_DEPRECATION_WARNINGS
 	SafeWidget->SetClipping(Clipping);
 #endif
 
+	SafeWidget->SetPixelSnapping(PixelSnapping);
 	SafeWidget->SetFlowDirectionPreference(FlowDirectionPreference);
 
 	SafeWidget->ForceVolatile(bIsVolatile);
diff --git a/Engine/Source/Runtime/UMG/Public/Components/Widget.h b/Engine/Source/Runtime/UMG/Public/Components/Widget.h
--- a/Engine/Source/Runtime/UMG/Public/Components/Widget.h
+++ b/Engine/Source/Runtime/UMG/Public/Components/Widget.h
@@ -22,6 +22,7 @@
 #include "UObject/UObjectThreadContext.h"
 #include "GameFramework/PlayerController.h"
 #include "Blueprint/WidgetNavigation.h"
+#include "Widgets/WidgetPixelSnapping.h"
 
 #include "Widget.generated.h"
 
@@ -415,7 +416,7 @@ public:
 	 * performance cost to clipping.  Do not enable clipping unless a panel actually needs to prevent
 	 * content from showing up outside its bounds.
 	 */
-	UPROPERTY(EditAnywhere, BlueprintReadWrite, Getter, Setter, Category="Clipping")
+	UPROPERTY(EditAnywhere, BlueprintReadWrite, Getter, Setter, Category="Rendering")
 	EWidgetClipping Clipping;
 
 	UE_DEPRECATED(5.1, "Direct access to Visibility is deprecated. Please use the getter or setter.")
@@ -425,9 +426,14 @@ public:
 
 	UE_DEPRECATED(5.1, "Direct access to RenderOpacity is deprecated. Please use the getter or setter.")
 	/** The opacity of the widget */
-	UPROPERTY(EditAnywhere, BlueprintReadWrite, Getter, Setter, BlueprintGetter="GetRenderOpacity", BlueprintSetter="SetRenderOpacity", Category="Behavior")
+	UPROPERTY(EditAnywhere, BlueprintReadWrite, Getter, Setter, BlueprintGetter="GetRenderOpacity", BlueprintSetter="SetRenderOpacity", Category="Rendering")
 	float RenderOpacity;
 
+private:
+	/** If the widget will draw snapped to the nearest pixel.  Improves clarity but might cause visibile stepping in animation */
+	UPROPERTY(EditAnywhere, BlueprintReadWrite, Getter, Setter, Category = "Rendering", meta=(AllowPrivateAccess = true))
+	EWidgetPixelSnapping PixelSnapping;
+
 private:
 	/** A custom set of accessibility rules for this widget. If null, default rules for the widget are used. */
 	UPROPERTY(Instanced)
@@ -585,6 +591,14 @@ public:
 	/** Sets the clipping state of this widget. */
 	UFUNCTION(BlueprintCallable, Category = "Widget")
 	void SetClipping(EWidgetClipping InClipping);
+	
+	/** Gets the pixel snapping method of this widget. */
+	UFUNCTION(BlueprintCallable, Category = "Widget")
+	EWidgetPixelSnapping GetPixelSnapping() const;
+
+	/** Sets the pixel snapping method of this widget. */
+	UFUNCTION(BlueprintCallable, Category = "Widget")
+	void SetPixelSnapping(EWidgetPixelSnapping InPixelSnapping);
 
 	/** Sets the forced volatility of the widget. */
 	UFUNCTION(BlueprintCallable, Category="Widget")
