Fix clipping against incorrect slate rect

diff --git a/Engine/Source/Runtime/SlateCore/Private/Rendering/DrawElements.cpp b/Engine/Source/Runtime/SlateCore/Private/Rendering/DrawElements.cpp
--- a/Engine/Source/Runtime/SlateCore/Private/Rendering/DrawElements.cpp
+++ b/Engine/Source/Runtime/SlateCore/Private/Rendering/DrawElements.cpp
@@ -45,7 +45,7 @@ static bool ShouldCull(const FSlateWindowElementList& ElementList)
 		}
 		else if (ClippingState.GetClippingMethod() == EClippingMethod::Stencil)
 		{
-			FSlateRect WindowRect = FSlateRect(FVector2f(0, 0), ElementList.GetWindowSize());
+			FSlateRect WindowRect = ElementList.GetPaintWindow()->GetClippingRectangleInWindow();
 			if (WindowRect.GetArea() > 0)
 			{
 				for (const FSlateClippingZone& Stencil : ClippingState.StencilQuads)
