Added Variation Selector characters to list of characters require full text shaping

diff --git a/Engine/Source/Runtime/SlateCore/Private/Fonts/SlateTextShaper.cpp b/Engine/Source/Runtime/SlateCore/Private/Fonts/SlateTextShaper.cpp
--- a/Engine/Source/Runtime/SlateCore/Private/Fonts/SlateTextShaper.cpp
+++ b/Engine/Source/Runtime/SlateCore/Private/Fonts/SlateTextShaper.cpp
@@ -166,6 +166,7 @@ void FSlateTextShaper::PerformTextShaping(const TCHAR* InText, const int32 InTex
 					RETURN_TRUE_IF_CHAR_WITHIN_RANGE(TEXT('\u20D0'), TEXT('\u20FF'));
 					RETURN_TRUE_IF_CHAR_WITHIN_RANGE(TEXT('\u3099'), TEXT('\u309A'));
 					RETURN_TRUE_IF_CHAR_WITHIN_RANGE(TEXT('\u31C0'), TEXT('\u31EF'));
+					RETURN_TRUE_IF_CHAR_WITHIN_RANGE(TEXT('\uFE00'), TEXT('\uFE0F'));
 					RETURN_TRUE_IF_CHAR_WITHIN_RANGE(TEXT('\uFE20'), TEXT('\uFE2F'));
 
 					// Devanagari
