Fixes log errors for uninitialized members in FToolBarStyle

diff --git a/Engine/Source/Runtime/SlateCore/Private/Styling/ToolBarStyle.cpp b/Engine/Source/Runtime/SlateCore/Private/Styling/ToolBarStyle.cpp
--- a/Engine/Source/Runtime/SlateCore/Private/Styling/ToolBarStyle.cpp
+++ b/Engine/Source/Runtime/SlateCore/Private/Styling/ToolBarStyle.cpp
@@ -17,6 +17,9 @@ FToolBarStyle::FToolBarStyle()
 	, SettingsComboButton()
 	, ButtonStyle()
 	, LabelPadding(0)
+	, UniformBlockWidth(0.f)
+	, UniformBlockHeight(0.f)
+	, NumColumns(0)
 	, IconPadding(0)
 	, SeparatorPadding(0)
 	, ComboButtonPadding(0)
