Fix slate crash with fastgenprofile

#rnx

diff --git a/Engine/Source/Runtime/SlateCore/Private/FastUpdate/WidgetProxy.cpp b/Engine/Source/Runtime/SlateCore/Private/FastUpdate/WidgetProxy.cpp
--- a/Engine/Source/Runtime/SlateCore/Private/FastUpdate/WidgetProxy.cpp
+++ b/Engine/Source/Runtime/SlateCore/Private/FastUpdate/WidgetProxy.cpp
@@ -51,6 +51,10 @@ FWidgetProxy::FUpdateResult FWidgetProxy::Update(const FPaintArgs& PaintArgs, FS
 
 	// If Outgoing layer id remains index none, there was no change
 	FUpdateResult Result;
+
+	if (!CurrentWidget.IsValid())
+		return Result;
+		
 	if (CurrentWidget->HasAnyUpdateFlags(EWidgetUpdateFlags::NeedsRepaint|EWidgetUpdateFlags::NeedsVolatilePaint))
 	{
 		Result = Repaint(PaintArgs, OutDrawElements);
