Removed unnecessary string allocations. The changed code showed up when profling the editor.

diff --git a/Engine/Source/Runtime/SlateCore/Private/Widgets/SWidget.cpp b/Engine/Source/Runtime/SlateCore/Private/Widgets/SWidget.cpp
--- a/Engine/Source/Runtime/SlateCore/Private/Widgets/SWidget.cpp
+++ b/Engine/Source/Runtime/SlateCore/Private/Widgets/SWidget.cpp
@@ -1076,7 +1076,9 @@ int32 SWidget::FindChildUnderPosition( const FArrangedChildren& Children, const
 
 FString SWidget::ToString() const
 {
-	return FString::Printf(TEXT("%s [%s]"), *this->TypeOfWidget.ToString(), *this->GetReadableLocation() );
+	TStringBuilder<256> StringBuilder;
+	StringBuilder << this->TypeOfWidget << " [" << *this->GetReadableLocation() << "]";
+	return FString(StringBuilder);
 }
 
 FString SWidget::GetTypeAsString() const
@@ -1092,7 +1094,9 @@ FName SWidget::GetType() const
 FString SWidget::GetReadableLocation() const
 {
 #if !UE_BUILD_SHIPPING
-	return FString::Printf(TEXT("%s(%d)"), *FPaths::GetCleanFilename(this->CreatedInLocation.GetPlainNameString()), this->CreatedInLocation.GetNumber());
+	TStringBuilder<256> StringBuilder;
+	StringBuilder << *FPaths::GetCleanFilename(this->CreatedInLocation.GetPlainNameString()) << "(" << this->CreatedInLocation.GetNumber() << ")";
+	return FString(StringBuilder);
 #else
 	return FString();
 #endif
