Updating AppIcon to be an SVG and updating style files accessing the original image version

diff --git a/Engine/Source/Runtime/SlateCore/Private/Styling/StarshipCoreStyle.cpp b/Engine/Source/Runtime/SlateCore/Private/Styling/StarshipCoreStyle.cpp
--- a/Engine/Source/Runtime/SlateCore/Private/Styling/StarshipCoreStyle.cpp
+++ b/Engine/Source/Runtime/SlateCore/Private/Styling/StarshipCoreStyle.cpp
@@ -283,7 +283,7 @@ TSharedRef<ISlateStyle> FStarshipCoreStyle::Create()
 		Style->Set("AppIcon", new IMAGE_BRUSH_SVG("Starship/Common/unreal", FVector2f(36.f, 36.f), FStyleColors::Foreground));
 		Style->Set("AppIconPadding", FMargin(11.f, 11.f, 3.f, 5.f));
 #else
-		Style->Set("AppIcon", new IMAGE_BRUSH("Starship/Common/UELogo", FVector2f(45.f, 45.f), FStyleColors::White));
+		Style->Set("AppIcon", new IMAGE_BRUSH_SVG("Starship/Common/UELogo", FVector2f(45.f, 45.f), FStyleColors::White));
 		Style->Set("AppIconPadding", FMargin(5.f, 5.f, 5.f, 5.f));
 #endif
 
diff --git a/Engine/Source/Runtime/SlateCore/Private/Styling/UMGCoreStyle.cpp b/Engine/Source/Runtime/SlateCore/Private/Styling/UMGCoreStyle.cpp
--- a/Engine/Source/Runtime/SlateCore/Private/Styling/UMGCoreStyle.cpp
+++ b/Engine/Source/Runtime/SlateCore/Private/Styling/UMGCoreStyle.cpp
@@ -207,7 +207,7 @@ TSharedRef<ISlateStyle> FUMGCoreStyle::Create()
 	static const FSlateColor SelectionColor_Pressed(Style->SelectionColor_Pressed_LinearRef);
 	static const FSlateColor HighlightColor(FUMGColor::Highlight);
 
-	Style->Set("AppIcon", new IMAGE_BRUSH("Starship/Common/UELogo", Icon24x24, FStyleColors::White));
+	Style->Set("AppIcon", new IMAGE_BRUSH_SVG("Starship/Common/UELogo", Icon24x24, FStyleColors::White));
 	Style->Set("AppIcon.Small", new IMAGE_BRUSH_SVG("Starship/Common/unreal-small", Icon24x24, FStyleColors::Foreground));
 
 	Style->Set("AppIconPadding", FMargin(4, 4, 0, 0));
