SDragAndDropVerticalBox::OnDragOver : Added index validation when checking the index of the element being dragged.
- This is required to avoid "out of range" exception when trying to access the geometry of a non existent arranged children index.

diff --git a/Engine/Source/Runtime/SlateCore/Private/Widgets/SBoxPanel.cpp b/Engine/Source/Runtime/SlateCore/Private/Widgets/SBoxPanel.cpp
--- a/Engine/Source/Runtime/SlateCore/Private/Widgets/SBoxPanel.cpp
+++ b/Engine/Source/Runtime/SlateCore/Private/Widgets/SBoxPanel.cpp
@@ -148,7 +148,6 @@ static FVector2D ComputeDesiredSizeForBox( const TPanelChildren<SlotType>& Child
 		
 		if ( CurChild.GetWidget()->GetVisibility() != EVisibility::Collapsed )
 		{
-
 			FVector2f CurChildDesiredSize = CurChild.GetWidget()->GetDesiredSize();
 
 			if (Orientation == Orient_Vertical)
@@ -334,7 +333,7 @@ FReply SDragAndDropVerticalBox::OnDragOver(const FGeometry& MyGeometry, const FD
 
 		TSharedPtr<FDragAndDropVerticalBoxOp> DragOp = DragDropEvent.GetOperationAs<FDragAndDropVerticalBoxOp>();
 
-		if (DragOp.IsValid())
+		if (DragOp.IsValid() && ArrangedChildren.IsValidIndex(DragOp->SlotIndexBeingDragged))
 		{
 			int32 DragOverSlotIndex = SWidget::FindChildUnderPosition(ArrangedChildren, DragDropEvent.GetScreenSpacePosition());
 
