Fix to SlateFontInfo to limit Outline to 1024 and prevent a crash

diff --git a/Engine/Source/Runtime/SlateCore/Public/Fonts/SlateFontInfo.h b/Engine/Source/Runtime/SlateCore/Public/Fonts/SlateFontInfo.h
--- a/Engine/Source/Runtime/SlateCore/Public/Fonts/SlateFontInfo.h
+++ b/Engine/Source/Runtime/SlateCore/Public/Fonts/SlateFontInfo.h
@@ -33,7 +33,7 @@ struct SLATECORE_API FFontOutlineSettings
 	GENERATED_USTRUCT_BODY()
 
 	/** Size of the outline in slate units (at 1.0 font scale this unit is a pixel)*/
-	UPROPERTY(EditAnywhere, BlueprintReadWrite, Category=OutlineSettings, meta=(ClampMin="0"))
+	UPROPERTY(EditAnywhere, BlueprintReadWrite, Category=OutlineSettings, meta=(ClampMin="0", ClampMax="1024"))
 	int32 OutlineSize;
 
 	/**
