Latest Clang toolchain fixes for FBindingContext, FGeometry, FInputChord and FPointerEvent. deprecated-copy-with-user-provided-copy.

diff --git a/Engine/Source/Runtime/Slate/Public/Framework/Commands/InputChord.h b/Engine/Source/Runtime/Slate/Public/Framework/Commands/InputChord.h
--- a/Engine/Source/Runtime/Slate/Public/Framework/Commands/InputChord.h
+++ b/Engine/Source/Runtime/Slate/Public/Framework/Commands/InputChord.h
@@ -121,14 +121,28 @@ struct SLATE_API FInputChord
 	 *
 	 * @param Other
 	 */
-	FInputChord( const FInputChord& Other )
-		: Key(Other.Key)
-		, bShift(Other.bShift)
-		, bCtrl(Other.bCtrl)
-		, bAlt(Other.bAlt)
-		, bCmd(Other.bCmd)
-	{
-	}
+	FInputChord(const FInputChord& Other) = default;
+
+	/**
+	 * Move constructor.
+	 *
+	 * @param Other
+	 */
+	FInputChord(FInputChord&& Other) = default;
+
+	/**
+	 * Copy assignment operator.
+	 *
+	 * @param Other
+	 */
+	FInputChord& operator=(const FInputChord& Other) = default;
+
+	/**
+	 * Move assignment operator.
+	 *
+	 * @param Other
+	 */
+	FInputChord& operator=(FInputChord&& Other) = default;
 
 	/**
 	 * Compares this input chord with another for equality.
diff --git a/Engine/Source/Runtime/Slate/Public/Framework/Commands/UICommandInfo.h b/Engine/Source/Runtime/Slate/Public/Framework/Commands/UICommandInfo.h
--- a/Engine/Source/Runtime/Slate/Public/Framework/Commands/UICommandInfo.h
+++ b/Engine/Source/Runtime/Slate/Public/Framework/Commands/UICommandInfo.h
@@ -101,12 +101,10 @@ public:
 		check(!InStyleSetName.IsNone());
 	}
 
-	FBindingContext( const FBindingContext &Other )
-		: ContextName( Other.ContextName )
-		, ContextParent( Other.ContextParent )
-		, ContextDesc( Other.ContextDesc )
-		, StyleSetName( Other.StyleSetName )
-	{}
+	FBindingContext(const FBindingContext&) = default;
+	FBindingContext(FBindingContext&&) = default;
+	FBindingContext& operator=(const FBindingContext&) = default;
+	FBindingContext& operator=(FBindingContext&&) = default;
 
 	/**
 	 * Creates a new command declaration used to populate commands with data
diff --git a/Engine/Source/Runtime/SlateCore/Public/Input/Events.h b/Engine/Source/Runtime/SlateCore/Public/Input/Events.h
--- a/Engine/Source/Runtime/SlateCore/Public/Input/Events.h
+++ b/Engine/Source/Runtime/SlateCore/Public/Input/Events.h
@@ -924,6 +924,7 @@ public:
 		, CursorDelta(FVector2f(LastScreenSpacePosition) - FVector2f(ScreenSpacePosition))
 		, PressedButtons(&InPressedButtons)
 		, PointerIndex(0)
+		, TouchpadIndex(0)
 		, Force(1.0f)
 		, bIsTouchEvent(false)
 		, GestureType(InGestureType)
@@ -997,29 +998,6 @@ public:
 	/** Returns the full set of pressed buttons */
 	const TSet<FKey>& GetPressedButtons() const { return *PressedButtons; }
 
-	/** We override the assignment operator to allow generated code to compile with the const ref member. */
-	void operator=( const FPointerEvent& Other )
-	{
-		FInputEvent::operator=( Other );
-
-		// Pointer
-		ScreenSpacePosition = Other.ScreenSpacePosition;
-		LastScreenSpacePosition = Other.LastScreenSpacePosition;
-		CursorDelta = Other.CursorDelta;
-		PressedButtons = Other.PressedButtons;
-		EffectingButton = Other.EffectingButton;
-		UserIndex = Other.UserIndex;
-		PointerIndex = Other.PointerIndex;
-		TouchpadIndex = Other.TouchpadIndex;
-		Force = Other.Force;
-		bIsTouchEvent = Other.bIsTouchEvent;
-		GestureType = Other.GestureType;
-		WheelOrGestureDelta = Other.WheelOrGestureDelta;
-		bIsDirectionInvertedFromDevice = Other.bIsDirectionInvertedFromDevice;
-		bIsTouchForceChanged = Other.bIsTouchForceChanged;
-		bIsTouchFirstMove = Other.bIsTouchFirstMove;
-	}
-
 	virtual FText ToText() const override;
 
 	virtual bool IsPointerEvent() const override;
diff --git a/Engine/Source/Runtime/SlateCore/Public/Layout/Geometry.h b/Engine/Source/Runtime/SlateCore/Public/Layout/Geometry.h
--- a/Engine/Source/Runtime/SlateCore/Public/Layout/Geometry.h
+++ b/Engine/Source/Runtime/SlateCore/Public/Layout/Geometry.h
@@ -47,6 +47,11 @@ public:
 	 */
 	FGeometry();
 
+	/**
+	 * Copy constructor.
+	 */
+	FGeometry(const FGeometry& RHS) = default;
+
 	/**
 	 * !!! HACK!!! We're keeping members of FGeometry const to prevent mutability without making them private, for backward compatibility.
 	 * But this means the assignment operator no longer works. We implement one ourselves now and force a memcpy.
