Workaround for PGO FastGen Bug

diff --git a/Engine/Source/Runtime/SlateCore/Public/FastUpdate/WidgetProxy.h b/Engine/Source/Runtime/SlateCore/Public/FastUpdate/WidgetProxy.h
--- a/Engine/Source/Runtime/SlateCore/Public/FastUpdate/WidgetProxy.h
+++ b/Engine/Source/Runtime/SlateCore/Public/FastUpdate/WidgetProxy.h
@@ -158,7 +158,8 @@ public:
 	}
 #else
 	SWidget* GetWidget() const { return Widget; }
-	TSharedPtr<SWidget> GetWidgetAsShared() const;
+	// ForceNoInline to workaround PGO FastGen issue
+	FORCENOINLINE TSharedPtr<SWidget> GetWidgetAsShared() const;
 	void ResetWidget() { Widget = nullptr; }
 	bool IsSameWidget(const SWidget* InWidget) const { return InWidget == Widget; }
 #endif
