Skip manually adding border and titlebar size to Autosized windows because it's already included in the desired size.

diff --git a/Engine/Source/Runtime/SlateCore/Private/Widgets/SWindow.cpp b/Engine/Source/Runtime/SlateCore/Private/Widgets/SWindow.cpp
--- a/Engine/Source/Runtime/SlateCore/Private/Widgets/SWindow.cpp
+++ b/Engine/Source/Runtime/SlateCore/Private/Widgets/SWindow.cpp
@@ -194,7 +194,7 @@ UE::Slate::FDeprecateVector2DResult SWindow::GetWindowSizeFromClientSize(UE::Sla
 {
 	// If this is a regular non-OS window, we need to compensate for the border and title bar area that we will add
 	// Note: Windows with an OS border do this in ReshapeWindow
-	if (IsRegularWindow() && !HasOSWindowBorder())
+	if (IsRegularWindow() && !HasOSWindowBorder() && !IsAutosized())
 	{
 		const FMargin BorderSize = GetWindowBorderSize();
 
