Update defaults on UMG widgets to use default style cache.

diff --git a/Engine/Source/Runtime/SlateCore/Private/Styling/DefaultStyleCache.cpp b/Engine/Source/Runtime/SlateCore/Private/Styling/DefaultStyleCache.cpp
new file mode 100644
--- /dev/null
+++ b/Engine/Source/Runtime/SlateCore/Private/Styling/DefaultStyleCache.cpp
@@ -0,0 +1,135 @@
+// Copyright Epic Games, Inc. All Rights Reserved.
+
+#include "Styling/DefaultStyleCache.h"
+
+#include "Misc/LazySingleton.h"
+#include "Styling/SlateTypes.h"
+#include "Styling/UMGCoreStyle.h"
+
+#if WITH_EDITOR
+#include "Styling/CoreStyle.h"
+#endif
+
+namespace UE::Slate::Private
+{
+
+FDefaultStyleCache& FDefaultStyleCache::Get()
+{
+	return TLazySingleton<FDefaultStyleCache>::Get();
+}
+
+FDefaultStyleCache::FDefaultStyleCache()
+{
+	if (!IsRunningDedicatedServer())
+	{
+		Runtime.ButtonStyle = FUMGCoreStyle::Get().GetWidgetStyle<FButtonStyle>("Button");
+		Runtime.ButtonStyle.UnlinkColors();
+
+		Runtime.CheckboxStyle = FUMGCoreStyle::Get().GetWidgetStyle<FCheckBoxStyle>("Checkbox");
+		Runtime.CheckboxStyle.UnlinkColors();
+
+		Runtime.CircularThrobberBrushStyle = *FUMGCoreStyle::Get().GetBrush("Throbber.CircleChunk");
+		Runtime.CircularThrobberBrushStyle.UnlinkColors();
+
+		Runtime.ComboBoxStyle = FUMGCoreStyle::Get().GetWidgetStyle<FComboBoxStyle>("ComboBox");
+		Runtime.ComboBoxStyle.UnlinkColors();
+
+		Runtime.EditableTextStyle = FUMGCoreStyle::Get().GetWidgetStyle<FEditableTextStyle>("NormalEditableText");
+		Runtime.EditableTextStyle.UnlinkColors();
+
+		Runtime.EditableTextBoxStyle = FUMGCoreStyle::Get().GetWidgetStyle<FEditableTextBoxStyle>("NormalEditableTextBox");
+		Runtime.EditableTextBoxStyle.UnlinkColors();
+
+		Runtime.ExpandableAreaStyle = FUMGCoreStyle::Get().GetWidgetStyle<FExpandableAreaStyle>("ExpandableArea");
+		Runtime.ExpandableAreaStyle.UnlinkColors();
+		Runtime.ExpandableAreaBorderBrush = *FUMGCoreStyle::Get().GetBrush("ExpandableArea.Border");
+		Runtime.ExpandableAreaBorderBrush.UnlinkColors();
+
+		Runtime.ListViewStyle = FUMGCoreStyle::Get().GetWidgetStyle<FTableViewStyle>("ListView");
+		Runtime.ListViewStyle.UnlinkColors();
+
+		Runtime.ProgressBarStyle = FUMGCoreStyle::Get().GetWidgetStyle<FProgressBarStyle>("ProgressBar");
+		Runtime.ProgressBarStyle.UnlinkColors();
+
+		Runtime.ScrollBarStyle = FUMGCoreStyle::Get().GetWidgetStyle<FScrollBarStyle>("ScrollBar");
+		Runtime.ScrollBarStyle.UnlinkColors();
+
+		Runtime.ScrollBoxStyle = FUMGCoreStyle::Get().GetWidgetStyle<FScrollBoxStyle>("ScrollBox");
+		Runtime.ScrollBoxStyle.UnlinkColors();
+
+		Runtime.SliderStyle = FUMGCoreStyle::Get().GetWidgetStyle<FSliderStyle>("Slider");
+		Runtime.SliderStyle.UnlinkColors();
+
+		Runtime.SpinBoxStyle = FUMGCoreStyle::Get().GetWidgetStyle<FSpinBoxStyle>("SpinBox");
+		Runtime.SpinBoxStyle.UnlinkColors();
+
+		Runtime.TableRowStyle = FUMGCoreStyle::Get().GetWidgetStyle<FTableRowStyle>("TableView.Row");
+		Runtime.TableRowStyle.UnlinkColors();
+
+		Runtime.ThrobberBrush = *FUMGCoreStyle::Get().GetBrush("Throbber.Chunk");
+		Runtime.ThrobberBrush.UnlinkColors();
+
+		Runtime.TextBlockStyle = FUMGCoreStyle::Get().GetWidgetStyle<FTextBlockStyle>("NormalText");
+		Runtime.TextBlockStyle.UnlinkColors();
+
+		Runtime.TreeViewStyle = FUMGCoreStyle::Get().GetWidgetStyle<FTableViewStyle>("TreeView");
+		Runtime.TreeViewStyle.UnlinkColors();
+
+#if WITH_EDITOR
+		Editor.ButtonStyle = FCoreStyle::Get().GetWidgetStyle<FButtonStyle>("EditorUtilityButton");
+		Editor.ButtonStyle.UnlinkColors();
+
+		Editor.CheckboxStyle = FCoreStyle::Get().GetWidgetStyle<FCheckBoxStyle>("Checkbox");
+		Editor.CheckboxStyle.UnlinkColors();
+
+		Editor.CircularThrobberBrushStyle = *FCoreStyle::Get().GetBrush("Throbber.CircleChunk");
+		Editor.CircularThrobberBrushStyle.UnlinkColors();
+
+		Editor.ComboBoxStyle = FCoreStyle::Get().GetWidgetStyle<FComboBoxStyle>("EditorUtilityComboBox");
+		Editor.ComboBoxStyle.UnlinkColors();
+
+		Editor.EditableTextStyle = FCoreStyle::Get().GetWidgetStyle<FEditableTextStyle>("NormalEditableText");
+		Editor.EditableTextStyle.UnlinkColors();
+
+		Editor.EditableTextBoxStyle = FCoreStyle::Get().GetWidgetStyle<FEditableTextBoxStyle>("NormalEditableTextBox");
+		Editor.EditableTextBoxStyle.UnlinkColors();
+
+		Editor.ExpandableAreaStyle = FCoreStyle::Get().GetWidgetStyle<FExpandableAreaStyle>("ExpandableArea");
+		Editor.ExpandableAreaStyle.UnlinkColors();
+		Editor.ExpandableAreaBorderBrush = *FCoreStyle::Get().GetBrush("ExpandableArea.Border");
+		Editor.ExpandableAreaBorderBrush.UnlinkColors();
+
+		Editor.ListViewStyle = FCoreStyle::Get().GetWidgetStyle<FTableViewStyle>("ListView");
+		Editor.ListViewStyle.UnlinkColors();
+
+		Editor.ProgressBarStyle = FCoreStyle::Get().GetWidgetStyle<FProgressBarStyle>("ProgressBar");
+		Editor.ProgressBarStyle.UnlinkColors();
+
+		Editor.ScrollBarStyle = FCoreStyle::Get().GetWidgetStyle<FScrollBarStyle>("ScrollBar");
+		Editor.ScrollBarStyle.UnlinkColors();
+
+		Editor.ScrollBoxStyle = FCoreStyle::Get().GetWidgetStyle<FScrollBoxStyle>("ScrollBox");
+		Editor.ScrollBoxStyle.UnlinkColors();
+
+		Editor.SliderStyle = FCoreStyle::Get().GetWidgetStyle<FSliderStyle>("Slider");
+		Editor.SliderStyle.UnlinkColors();
+
+		Editor.SpinBoxStyle = FCoreStyle::Get().GetWidgetStyle<FSpinBoxStyle>("SpinBox");
+		Editor.SpinBoxStyle.UnlinkColors();
+
+		Editor.TableRowStyle = FCoreStyle::Get().GetWidgetStyle<FTableRowStyle>("TableView.Row");
+		Editor.TableRowStyle.UnlinkColors();
+
+		Editor.ThrobberBrush = *FCoreStyle::Get().GetBrush("Throbber.Chunk");
+		Editor.ThrobberBrush.UnlinkColors();
+
+		Editor.TextBlockStyle = FCoreStyle::Get().GetWidgetStyle<FTextBlockStyle>("NormalText");
+		Editor.TextBlockStyle.UnlinkColors();
+
+		Editor.TreeViewStyle = FCoreStyle::Get().GetWidgetStyle<FTableViewStyle>("TreeView");
+		Editor.TreeViewStyle.UnlinkColors();
+#endif
+	}
+}
+
+} //namespace UE::Slate::Private
diff --git a/Engine/Source/Runtime/SlateCore/Public/Styling/DefaultStyleCache.h b/Engine/Source/Runtime/SlateCore/Public/Styling/DefaultStyleCache.h
new file mode 100644
--- /dev/null
+++ b/Engine/Source/Runtime/SlateCore/Public/Styling/DefaultStyleCache.h
@@ -0,0 +1,121 @@
+// Copyright Epic Games, Inc. All Rights Reserved.
+
+#pragma once
+
+#include "CoreMinimal.h"
+#include "Styling/SlateTypes.h"
+
+namespace UE::Slate::Private
+{
+
+/**
+ * Single point of access for various default styles used in UMG for runtime and editor
+ * with each style already having it's colors unlinked for reuse
+ */
+struct SLATECORE_API FDefaultStyleCache
+{
+	/**
+	 * Internal style default holder struct, used so that one can get a bundle of defaults relevant for runtime or editor
+	 */
+	struct FStyles
+	{
+		const FButtonStyle& GetButtonStyle() const { return ButtonStyle; };
+
+		const FCheckBoxStyle& GetCheckboxStyle() const { return CheckboxStyle; };
+
+		const FSlateBrush& GetCircularThrobberBrushStyle() const { return CircularThrobberBrushStyle; };
+
+		const FComboBoxStyle& GetComboBoxStyle() const { return ComboBoxStyle; };
+
+		const FEditableTextStyle& GetEditableTextStyle() const { return EditableTextStyle; };
+
+		const FEditableTextBoxStyle& GetEditableTextBoxStyle() const { return EditableTextBoxStyle; };
+
+		const FExpandableAreaStyle& GetExpandableAreaStyle() const { return ExpandableAreaStyle; };
+		const FSlateBrush& GetExpandableAreaBorderBrush() const { return ExpandableAreaBorderBrush; };
+
+		const FTableViewStyle& GetListViewStyle() const { return ListViewStyle; };
+
+		const FProgressBarStyle& GetProgressBarStyle() const { return ProgressBarStyle; };
+
+		const FScrollBarStyle& GetScrollBarStyle() const { return ScrollBarStyle; };
+
+		const FScrollBoxStyle& GetScrollBoxStyle() const { return ScrollBoxStyle; };
+
+		const FSliderStyle& GetSliderStyle() const { return SliderStyle; };
+
+		const FSpinBoxStyle& GetSpinBoxStyle() const { return SpinBoxStyle; };
+
+		const FTableRowStyle& GetTableRowStyle() const { return TableRowStyle; };
+
+		const FSlateBrush& GetThrobberBrush() const { return ThrobberBrush; };
+
+		const FTextBlockStyle& GetTextBlockStyle() const { return TextBlockStyle; };
+
+		const FTableViewStyle& GetTreeViewStyle() const { return TreeViewStyle; };
+
+	private:
+		friend FDefaultStyleCache;
+
+		FButtonStyle ButtonStyle;
+
+		FCheckBoxStyle CheckboxStyle;
+
+		FSlateBrush CircularThrobberBrushStyle;
+
+		FComboBoxStyle ComboBoxStyle;
+
+		FEditableTextStyle EditableTextStyle;
+
+		FEditableTextBoxStyle EditableTextBoxStyle;
+
+		FExpandableAreaStyle ExpandableAreaStyle;
+		FSlateBrush ExpandableAreaBorderBrush;
+
+		FTableViewStyle ListViewStyle;
+
+		FProgressBarStyle ProgressBarStyle;
+
+		FScrollBarStyle ScrollBarStyle;
+
+		FScrollBoxStyle ScrollBoxStyle;
+
+		FSliderStyle SliderStyle;
+
+		FSpinBoxStyle SpinBoxStyle;
+
+		FTableRowStyle TableRowStyle;
+
+		FSlateBrush ThrobberBrush;
+
+		FTextBlockStyle TextBlockStyle;
+
+		FTableViewStyle TreeViewStyle;
+
+		TMap<UScriptStruct*, FSlateWidgetStyle*> TypeInstanceMap;
+	};
+
+	/** Gets singleton and returns runtime styles from singleton */
+	static const FStyles& GetRuntime() { return Get().Runtime; }
+#if WITH_EDITOR
+	/** Gets singleton and returns editor styles from singleton */
+	static const FStyles& GetEditor() { return Get().Editor; }
+#endif
+
+private:
+
+	/** Singleton getter, however private used since GetRuntime & GetEditor is preferred for styles */
+	static FDefaultStyleCache& Get();
+
+	friend class FLazySingleton;
+
+	FDefaultStyleCache();
+
+	FStyles Runtime;
+
+#if WITH_EDITOR
+	FStyles Editor;
+#endif
+};
+
+} //namespace UE::Slate::Private
\ No newline at end of file
diff --git a/Engine/Source/Runtime/UMG/Private/Components/Button.cpp b/Engine/Source/Runtime/UMG/Private/Components/Button.cpp
--- a/Engine/Source/Runtime/UMG/Private/Components/Button.cpp
+++ b/Engine/Source/Runtime/UMG/Private/Components/Button.cpp
@@ -8,6 +8,7 @@
 #include "Widgets/DeclarativeSyntaxSupport.h"
 #include "Widgets/Input/SButton.h"
 #include "Components/ButtonSlot.h"
+#include "Styling/DefaultStyleCache.h"
 #include "Styling/UMGCoreStyle.h"
 #include "Blueprint/WidgetTree.h"
 
@@ -18,38 +19,18 @@
 /////////////////////////////////////////////////////
 // UButton
 
-static FButtonStyle* DefaultButtonStyle = nullptr;
-
-#if WITH_EDITOR
-static FButtonStyle* EditorButtonStyle = nullptr;
-#endif 
-
 UButton::UButton(const FObjectInitializer& ObjectInitializer)
 	: Super(ObjectInitializer)
 {
-	if (DefaultButtonStyle == nullptr)
-	{
-		DefaultButtonStyle = new FButtonStyle(FUMGCoreStyle::Get().GetWidgetStyle<FButtonStyle>("Button"));
-
-		// Unlink UMG default colors.
-		DefaultButtonStyle->UnlinkColors();
-	}
 PRAGMA_DISABLE_DEPRECATION_WARNINGS
-	WidgetStyle = *DefaultButtonStyle;
+	WidgetStyle = UE::Slate::Private::FDefaultStyleCache::GetRuntime().GetButtonStyle();
 PRAGMA_ENABLE_DEPRECATION_WARNINGS
-#if WITH_EDITOR 
-	if (EditorButtonStyle == nullptr)
-	{
-		EditorButtonStyle = new FButtonStyle(FCoreStyle::Get().GetWidgetStyle<FButtonStyle>("EditorUtilityButton"));
-
-		// Unlink UMG Editor colors from the editor settings colors.
-		EditorButtonStyle->UnlinkColors();
-	}
 
+#if WITH_EDITOR 
 	if (IsEditorWidget())
 	{
 PRAGMA_DISABLE_DEPRECATION_WARNINGS
-		WidgetStyle = *EditorButtonStyle;
+		WidgetStyle = UE::Slate::Private::FDefaultStyleCache::GetEditor().GetButtonStyle();
 PRAGMA_ENABLE_DEPRECATION_WARNINGS
 
 		// The CDO isn't an editor widget and thus won't use the editor style, call post edit change to mark difference from CDO
diff --git a/Engine/Source/Runtime/UMG/Private/Components/CheckBox.cpp b/Engine/Source/Runtime/UMG/Private/Components/CheckBox.cpp
--- a/Engine/Source/Runtime/UMG/Private/Components/CheckBox.cpp
+++ b/Engine/Source/Runtime/UMG/Private/Components/CheckBox.cpp
@@ -7,6 +7,7 @@
 #include "Widgets/DeclarativeSyntaxSupport.h"
 #include "Widgets/Input/SCheckBox.h"
 #include "Slate/SlateBrushAsset.h"
+#include "Styling/DefaultStyleCache.h"
 #include "Styling/UMGCoreStyle.h"
 
 #include UE_INLINE_GENERATED_CPP_BY_NAME(CheckBox)
@@ -16,38 +17,16 @@
 /////////////////////////////////////////////////////
 // UCheckBox
 
-static FCheckBoxStyle* DefaultCheckboxStyle = nullptr;
-
-#if WITH_EDITOR
-static FCheckBoxStyle* EditorCheckboxStyle = nullptr;
-#endif 
-
 UCheckBox::UCheckBox(const FObjectInitializer& ObjectInitializer)
 	: Super(ObjectInitializer)
 {
-	if (DefaultCheckboxStyle == nullptr)
-	{
-		DefaultCheckboxStyle = new FCheckBoxStyle(FUMGCoreStyle::Get().GetWidgetStyle<FCheckBoxStyle>("Checkbox"));
-
-		// Unlink UMG default colors.
-		DefaultCheckboxStyle->UnlinkColors();
-	}
-
 PRAGMA_DISABLE_DEPRECATION_WARNINGS
-	WidgetStyle = *DefaultCheckboxStyle;
-
+	WidgetStyle = UE::Slate::Private::FDefaultStyleCache::GetRuntime().GetCheckboxStyle();
+	
 #if WITH_EDITOR 
-	if (EditorCheckboxStyle == nullptr)
-	{
-		EditorCheckboxStyle = new FCheckBoxStyle(FCoreStyle::Get().GetWidgetStyle<FCheckBoxStyle>("Checkbox"));
-
-		// Unlink UMG Editor colors from the editor settings colors.
-		EditorCheckboxStyle->UnlinkColors();
-	}
-
 	if (IsEditorWidget())
 	{
-		WidgetStyle = *EditorCheckboxStyle;
+		WidgetStyle = UE::Slate::Private::FDefaultStyleCache::GetEditor().GetCheckboxStyle();
 
 		// The CDO isn't an editor widget and thus won't use the editor style, call post edit change to mark difference from CDO
 		PostEditChange();
diff --git a/Engine/Source/Runtime/UMG/Private/Components/CircularThrobber.cpp b/Engine/Source/Runtime/UMG/Private/Components/CircularThrobber.cpp
--- a/Engine/Source/Runtime/UMG/Private/Components/CircularThrobber.cpp
+++ b/Engine/Source/Runtime/UMG/Private/Components/CircularThrobber.cpp
@@ -5,6 +5,7 @@
 #include "Widgets/DeclarativeSyntaxSupport.h"
 #include "Components/CanvasPanelSlot.h"
 #include "Widgets/Images/SThrobber.h"
+#include "Styling/DefaultStyleCache.h"
 #include "Styling/UMGCoreStyle.h"
 
 #include UE_INLINE_GENERATED_CPP_BY_NAME(CircularThrobber)
@@ -14,40 +15,18 @@
 /////////////////////////////////////////////////////
 // UCircularThrobber
 
-static FSlateBrush* DefaultCircularThrobberBrushStyle = nullptr;
-
-#if WITH_EDITOR
-static FSlateBrush* EditorCircularThrobberBrushStyle = nullptr;
-#endif 
-
 UCircularThrobber::UCircularThrobber(const FObjectInitializer& ObjectInitializer)
 	: Super(ObjectInitializer)
 	, bEnableRadius(true)
 {
-	if (DefaultCircularThrobberBrushStyle == nullptr)
-	{
-		DefaultCircularThrobberBrushStyle = new FSlateBrush(*FUMGCoreStyle::Get().GetBrush("Throbber.CircleChunk"));
-
-		// Unlink UMG default colors.
-		DefaultCircularThrobberBrushStyle->UnlinkColors();
-	}
-
 PRAGMA_DISABLE_DEPRECATION_WARNINGS
-	Image = *DefaultCircularThrobberBrushStyle;
-
+	Image = UE::Slate::Private::FDefaultStyleCache::GetRuntime().GetCircularThrobberBrushStyle();
+	
 #if WITH_EDITOR 
-	if (EditorCircularThrobberBrushStyle == nullptr)
-	{
-		EditorCircularThrobberBrushStyle = new FSlateBrush(*FCoreStyle::Get().GetBrush("Throbber.CircleChunk"));
-
-		// Unlink UMG Editor colors from the editor settings colors.
-		EditorCircularThrobberBrushStyle->UnlinkColors();
-	}
-
 	if (IsEditorWidget())
 	{
-		Image = *EditorCircularThrobberBrushStyle;
-		
+		Image = UE::Slate::Private::FDefaultStyleCache::GetEditor().GetCircularThrobberBrushStyle();
+
 		// The CDO isn't an editor widget and thus won't use the editor style, call post edit change to mark difference from CDO
 		PostEditChange();
 	}
diff --git a/Engine/Source/Runtime/UMG/Private/Components/ComboBox.cpp b/Engine/Source/Runtime/UMG/Private/Components/ComboBox.cpp
--- a/Engine/Source/Runtime/UMG/Private/Components/ComboBox.cpp
+++ b/Engine/Source/Runtime/UMG/Private/Components/ComboBox.cpp
@@ -1,7 +1,7 @@
 // Copyright Epic Games, Inc. All Rights Reserved.
 
 #include "Components/ComboBox.h"
-#include "DefaultStyleCache.h"
+#include "Styling/DefaultStyleCache.h"
 
 #include UE_INLINE_GENERATED_CPP_BY_NAME(ComboBox)
 
@@ -14,7 +14,7 @@ UComboBox::UComboBox(const FObjectInitializer& ObjectInitializer)
 	: Super(ObjectInitializer)
 {
 	bIsFocusable = true;
-	ScrollBarStyle = FDefaultStyleCache::Get().GetScrollBarStyle();
+	ScrollBarStyle = UE::Slate::Private::FDefaultStyleCache::GetRuntime().GetScrollBarStyle();
 }
 
 void UComboBox::ReleaseSlateResources(bool bReleaseChildren)
diff --git a/Engine/Source/Runtime/UMG/Private/Components/ComboBoxKey.cpp b/Engine/Source/Runtime/UMG/Private/Components/ComboBoxKey.cpp
--- a/Engine/Source/Runtime/UMG/Private/Components/ComboBoxKey.cpp
+++ b/Engine/Source/Runtime/UMG/Private/Components/ComboBoxKey.cpp
@@ -2,7 +2,7 @@
 
 #include "Components/ComboBoxKey.h"
 
-#include "DefaultStyleCache.h"
+#include "Styling/DefaultStyleCache.h"
 #include "Styling/UMGCoreStyle.h"
 #include "UObject/ConstructorHelpers.h"
 #include "Widgets/Layout/SBox.h"
@@ -23,9 +23,9 @@ UComboBoxKey::UComboBoxKey()
 	{
 #if WITH_EDITOR 
 PRAGMA_DISABLE_DEPRECATION_WARNINGS
-		WidgetStyle = IsEditorWidget() ? FDefaultStyleCache::Get().GetEditorComboBoxStyle() : FDefaultStyleCache::Get().GetComboBoxStyle();
-		ItemStyle = IsEditorWidget() ? FDefaultStyleCache::Get().GetEditorComboBoxRowStyle() : FDefaultStyleCache::Get().GetComboBoxRowStyle();
-		ScrollBarStyle = IsEditorWidget() ? FDefaultStyleCache::Get().GetEditorScrollBarStyle() : FDefaultStyleCache::Get().GetScrollBarStyle();
+		WidgetStyle = IsEditorWidget() ? UE::Slate::Private::FDefaultStyleCache::GetEditor().GetComboBoxStyle() : UE::Slate::Private::FDefaultStyleCache::GetRuntime().GetComboBoxStyle();
+		ItemStyle = IsEditorWidget() ? UE::Slate::Private::FDefaultStyleCache::GetEditor().GetTableRowStyle() : UE::Slate::Private::FDefaultStyleCache::GetRuntime().GetTableRowStyle();
+		ScrollBarStyle = IsEditorWidget() ? UE::Slate::Private::FDefaultStyleCache::GetEditor().GetScrollBarStyle() : UE::Slate::Private::FDefaultStyleCache::GetRuntime().GetScrollBarStyle();
 PRAGMA_ENABLE_DEPRECATION_WARNINGS
 
 		if (IsEditorWidget())
@@ -35,9 +35,9 @@ PRAGMA_ENABLE_DEPRECATION_WARNINGS
 		}
 #else
 PRAGMA_DISABLE_DEPRECATION_WARNINGS
-		WidgetStyle = FDefaultStyleCache::Get().GetComboBoxStyle();
-		ItemStyle = FDefaultStyleCache::Get().GetComboBoxRowStyle();
-		ScrollBarStyle = FDefaultStyleCache::Get().GetScrollBarStyle();
+		WidgetStyle = UE::Slate::Private::FDefaultStyleCache::GetRuntime().GetComboBoxStyle();
+		ItemStyle = UE::Slate::Private::FDefaultStyleCache::GetRuntime().GetTableRowStyle();
+		ScrollBarStyle = UE::Slate::Private::FDefaultStyleCache::GetRuntime().GetScrollBarStyle();
 PRAGMA_ENABLE_DEPRECATION_WARNINGS
 #endif // WITH_EDITOR
 	}
diff --git a/Engine/Source/Runtime/UMG/Private/Components/ComboBoxString.cpp b/Engine/Source/Runtime/UMG/Private/Components/ComboBoxString.cpp
--- a/Engine/Source/Runtime/UMG/Private/Components/ComboBoxString.cpp
+++ b/Engine/Source/Runtime/UMG/Private/Components/ComboBoxString.cpp
@@ -2,11 +2,11 @@
 
 #include "Components/ComboBoxString.h"
 
-#include "DefaultStyleCache.h"
 #include "Widgets/SNullWidget.h"
 #include "UObject/EditorObjectVersion.h"
 #include "UObject/ConstructorHelpers.h"
 #include "Engine/Font.h"
+#include "Styling/DefaultStyleCache.h"
 #include "Styling/UMGCoreStyle.h"
 
 #include UE_INLINE_GENERATED_CPP_BY_NAME(ComboBoxString)
@@ -20,16 +20,16 @@ UComboBoxString::UComboBoxString(const FObjectInitializer& ObjectInitializer)
 	: Super(ObjectInitializer)
 {
 PRAGMA_DISABLE_DEPRECATION_WARNINGS
-	WidgetStyle = FDefaultStyleCache::Get().GetComboBoxStyle();
-	ItemStyle = FDefaultStyleCache::Get().GetComboBoxRowStyle();
-	ScrollBarStyle = FDefaultStyleCache::Get().GetScrollBarStyle();
+	WidgetStyle = UE::Slate::Private::FDefaultStyleCache::GetRuntime().GetComboBoxStyle();
+	ItemStyle = UE::Slate::Private::FDefaultStyleCache::GetRuntime().GetTableRowStyle();
+	ScrollBarStyle = UE::Slate::Private::FDefaultStyleCache::GetRuntime().GetScrollBarStyle();
 
 #if WITH_EDITOR 
 	if (IsEditorWidget())
 	{
-		WidgetStyle = FDefaultStyleCache::Get().GetEditorComboBoxStyle();
-		ItemStyle = FDefaultStyleCache::Get().GetEditorComboBoxRowStyle();
-		ScrollBarStyle = FDefaultStyleCache::Get().GetEditorScrollBarStyle();
+		WidgetStyle = UE::Slate::Private::FDefaultStyleCache::GetEditor().GetComboBoxStyle();
+		ItemStyle = UE::Slate::Private::FDefaultStyleCache::GetEditor().GetTableRowStyle();
+		ScrollBarStyle = UE::Slate::Private::FDefaultStyleCache::GetEditor().GetScrollBarStyle();
 
 		// The CDO isn't an editor widget and thus won't use the editor style, call post edit change to mark difference from CDO
 		PostEditChange();
diff --git a/Engine/Source/Runtime/UMG/Private/Components/EditableText.cpp b/Engine/Source/Runtime/UMG/Private/Components/EditableText.cpp
--- a/Engine/Source/Runtime/UMG/Private/Components/EditableText.cpp
+++ b/Engine/Source/Runtime/UMG/Private/Components/EditableText.cpp
@@ -6,6 +6,7 @@
 #include "Widgets/DeclarativeSyntaxSupport.h"
 #include "Widgets/Input/SEditableText.h"
 #include "Slate/SlateBrushAsset.h"
+#include "Styling/DefaultStyleCache.h"
 #include "Styling/UMGCoreStyle.h"
 #include "Materials/MaterialInterface.h"
 
@@ -16,24 +17,13 @@
 /////////////////////////////////////////////////////
 // UEditableText
 
-static FEditableTextStyle* DefaultEditableTextStyle = nullptr;
-
-#if WITH_EDITOR
-static FEditableTextStyle* EditorEditableTextStyle = nullptr;
-#endif 
-
 UEditableText::UEditableText(const FObjectInitializer& ObjectInitializer)
 	: Super(ObjectInitializer)
 {
-	if (DefaultEditableTextStyle == nullptr)
-	{
-		DefaultEditableTextStyle = new FEditableTextStyle(FUMGCoreStyle::Get().GetWidgetStyle<FEditableTextStyle>("NormalEditableText"));
 
-		// Unlink UMG default colors.
-		DefaultEditableTextStyle->UnlinkColors();
-	}
+PRAGMA_DISABLE_DEPRECATION_WARNINGS
+	WidgetStyle = UE::Slate::Private::FDefaultStyleCache::GetRuntime().GetEditableTextStyle();
 
-	WidgetStyle = *DefaultEditableTextStyle;
 	if (!IsRunningDedicatedServer())
 	{
 		static ConstructorHelpers::FObjectFinder<UFont> DefaultFontObj(*UWidget::GetDefaultFontName());
@@ -43,27 +33,17 @@ UEditableText::UEditableText(const FObjectInitializer& ObjectInitializer)
 		Font.CompositeFont = WidgetStyle.Font.CompositeFont;
 		WidgetStyle.SetFont(Font);
 	}
-
-
+	
 #if WITH_EDITOR 
-	if (EditorEditableTextStyle == nullptr)
-	{
-		EditorEditableTextStyle = new FEditableTextStyle(FCoreStyle::Get().GetWidgetStyle<FEditableTextStyle>("NormalEditableText"));
-
-		// Unlink UMG Editor colors from the editor settings colors.
-		EditorEditableTextStyle->UnlinkColors();
-	}
-
 	if (IsEditorWidget())
 	{
-		WidgetStyle = *EditorEditableTextStyle;
+		WidgetStyle = UE::Slate::Private::FDefaultStyleCache::GetEditor().GetEditableTextStyle();
 
 		// The CDO isn't an editor widget and thus won't use the editor style, call post edit change to mark difference from CDO
 		PostEditChange();
 	}
 #endif // WITH_EDITOR
 
-	PRAGMA_DISABLE_DEPRECATION_WARNINGS
 	IsReadOnly = false;
 	IsPassword = false;
 	MinimumDesiredWidth = 0.0f;
@@ -77,7 +57,7 @@ UEditableText::UEditableText(const FObjectInitializer& ObjectInitializer)
 	VirtualKeyboardDismissAction = EVirtualKeyboardDismissAction::TextChangeOnDismiss;
 	SetClipping(EWidgetClipping::ClipToBounds);
 	OverflowPolicy = ETextOverflowPolicy::Clip;
-	PRAGMA_ENABLE_DEPRECATION_WARNINGS
+PRAGMA_ENABLE_DEPRECATION_WARNINGS
 
 #if WITH_EDITORONLY_DATA
 	AccessibleBehavior = ESlateAccessibleBehavior::Auto;
diff --git a/Engine/Source/Runtime/UMG/Private/Components/EditableTextBox.cpp b/Engine/Source/Runtime/UMG/Private/Components/EditableTextBox.cpp
--- a/Engine/Source/Runtime/UMG/Private/Components/EditableTextBox.cpp
+++ b/Engine/Source/Runtime/UMG/Private/Components/EditableTextBox.cpp
@@ -7,6 +7,7 @@
 #include "Engine/Font.h"
 #include "Widgets/DeclarativeSyntaxSupport.h"
 #include "Widgets/Input/SEditableTextBox.h"
+#include "Styling/DefaultStyleCache.h"
 #include "Styling/UMGCoreStyle.h"
 
 #include UE_INLINE_GENERATED_CPP_BY_NAME(EditableTextBox)
@@ -16,19 +17,13 @@
 /////////////////////////////////////////////////////
 // UEditableTextBox
 
-static FEditableTextBoxStyle* DefaultEditableTextBoxStyle = nullptr;
-
-#if WITH_EDITOR
-static FEditableTextBoxStyle* EditorEditableTextBoxStyle = nullptr;
-#endif 
-
 UEditableTextBox::UEditableTextBox(const FObjectInitializer& ObjectInitializer)
 	: Super(ObjectInitializer)
 {
 #if WITH_EDITOR
 	bIsFontDeprecationDone = false;
 #endif
-	PRAGMA_DISABLE_DEPRECATION_WARNINGS
+PRAGMA_DISABLE_DEPRECATION_WARNINGS
 	IsReadOnly = false;
 	IsPassword = false;
 	MinimumDesiredWidth = 0.0f;
@@ -40,17 +35,9 @@ UEditableTextBox::UEditableTextBox(const FObjectInitializer& ObjectInitializer)
 	AllowContextMenu = true;
 	VirtualKeyboardDismissAction = EVirtualKeyboardDismissAction::TextChangeOnDismiss;
 	OverflowPolicy = ETextOverflowPolicy::Clip;
-	PRAGMA_ENABLE_DEPRECATION_WARNINGS
 
-	if (DefaultEditableTextBoxStyle == nullptr)
-	{
-		DefaultEditableTextBoxStyle = new FEditableTextBoxStyle(FUMGCoreStyle::Get().GetWidgetStyle<FEditableTextBoxStyle>("NormalEditableTextBox"));
-
-		// Unlink UMG default colors.
-		DefaultEditableTextBoxStyle->UnlinkColors();
-	}
+	WidgetStyle = UE::Slate::Private::FDefaultStyleCache::GetRuntime().GetEditableTextBoxStyle();
 
-	WidgetStyle = *DefaultEditableTextBoxStyle;
 	if (!IsRunningDedicatedServer())
 	{
 		static ConstructorHelpers::FObjectFinder<UFont> DefaultFontObj(*UWidget::GetDefaultFontName());
@@ -61,24 +48,17 @@ UEditableTextBox::UEditableTextBox(const FObjectInitializer& ObjectInitializer)
 
 		WidgetStyle.SetFont(Font);
 	}
-
+	
 #if WITH_EDITOR 
-	if (EditorEditableTextBoxStyle == nullptr)
-	{
-		EditorEditableTextBoxStyle = new FEditableTextBoxStyle(FCoreStyle::Get().GetWidgetStyle<FEditableTextBoxStyle>("NormalEditableTextBox"));
-
-		// Unlink UMG Editor colors from the editor settings colors.
-		EditorEditableTextBoxStyle->UnlinkColors();
-	}
-
 	if (IsEditorWidget())
 	{
-		WidgetStyle = *EditorEditableTextBoxStyle;
+		WidgetStyle = UE::Slate::Private::FDefaultStyleCache::GetEditor().GetEditableTextBoxStyle();
 
 		// The CDO isn't an editor widget and thus won't use the editor style, call post edit change to mark difference from CDO
 		PostEditChange();
 	}
 #endif // WITH_EDITOR
+PRAGMA_ENABLE_DEPRECATION_WARNINGS
 
 #if WITH_EDITORONLY_DATA
 	AccessibleBehavior = ESlateAccessibleBehavior::Auto;
diff --git a/Engine/Source/Runtime/UMG/Private/Components/ExpandableArea.cpp b/Engine/Source/Runtime/UMG/Private/Components/ExpandableArea.cpp
--- a/Engine/Source/Runtime/UMG/Private/Components/ExpandableArea.cpp
+++ b/Engine/Source/Runtime/UMG/Private/Components/ExpandableArea.cpp
@@ -3,6 +3,7 @@
 #include "Components/ExpandableArea.h"
 #include "Widgets/SNullWidget.h"
 #include "Widgets/DeclarativeSyntaxSupport.h"
+#include "Styling/DefaultStyleCache.h"
 #include "Styling/UMGCoreStyle.h"
 
 #include "Widgets/Layout/SExpandableArea.h"
@@ -17,69 +18,27 @@
 static const FName HeaderName(TEXT("Header"));
 static const FName BodyName(TEXT("Body"));
 
-static FExpandableAreaStyle* DefaultExpandableAreaStyle = nullptr;
-static FSlateBrush* DefaultExpandableAreaBorderBrush = nullptr;
-
-#if WITH_EDITOR
-static FExpandableAreaStyle* EditorExpandableAreaStyle = nullptr;
-static FSlateBrush* EditorExpandableAreaBorderBrush = nullptr;
-#endif 
-
 UExpandableArea::UExpandableArea(const FObjectInitializer& ObjectInitializer)
 	: Super(ObjectInitializer)
 	, bIsExpanded(false)
 {
 	bIsVariable = true;
-
-	if (DefaultExpandableAreaStyle == nullptr)
-	{
-		DefaultExpandableAreaStyle = new FExpandableAreaStyle(FUMGCoreStyle::Get().GetWidgetStyle<FExpandableAreaStyle>("ExpandableArea"));
-
-		// Unlink UMG default colors.
-		DefaultExpandableAreaStyle->UnlinkColors();
-	}
-
-	if (DefaultExpandableAreaBorderBrush == nullptr)
-	{
-		DefaultExpandableAreaBorderBrush = new FSlateBrush(*FUMGCoreStyle::Get().GetBrush("ExpandableArea.Border"));
-
-		// Unlink UMG default colors.
-		DefaultExpandableAreaBorderBrush->UnlinkColors();
-	}
-
+	
 PRAGMA_DISABLE_DEPRECATION_WARNINGS
-	Style = *DefaultExpandableAreaStyle;
-	BorderBrush = *DefaultExpandableAreaBorderBrush;
-PRAGMA_ENABLE_DEPRECATION_WARNINGS
-
+	Style = UE::Slate::Private::FDefaultStyleCache::GetRuntime().GetExpandableAreaStyle();
+	BorderBrush = UE::Slate::Private::FDefaultStyleCache::GetRuntime().GetExpandableAreaBorderBrush();
+	
 #if WITH_EDITOR 
-	if (EditorExpandableAreaStyle == nullptr)
-	{
-		EditorExpandableAreaStyle = new FExpandableAreaStyle(FCoreStyle::Get().GetWidgetStyle<FExpandableAreaStyle>("ExpandableArea"));
-
-		// Unlink UMG Editor colors from the editor settings colors.
-		EditorExpandableAreaStyle->UnlinkColors();
-	}
-
-	if (EditorExpandableAreaBorderBrush == nullptr)
-	{
-		EditorExpandableAreaBorderBrush = new FSlateBrush(*FCoreStyle::Get().GetBrush("ExpandableArea.Border"));
-
-		// Unlink UMG Editor colors from the editor settings colors.
-		EditorExpandableAreaBorderBrush->UnlinkColors();
-	}
 	if (IsEditorWidget())
 	{
-PRAGMA_DISABLE_DEPRECATION_WARNINGS
-		Style = *EditorExpandableAreaStyle;
-		BorderBrush = *EditorExpandableAreaBorderBrush;
-PRAGMA_ENABLE_DEPRECATION_WARNINGS
+		Style = UE::Slate::Private::FDefaultStyleCache::GetEditor().GetExpandableAreaStyle();
+		BorderBrush = UE::Slate::Private::FDefaultStyleCache::GetEditor().GetExpandableAreaBorderBrush();
+
 		// The CDO isn't an editor widget and thus won't use the editor style, call post edit change to mark difference from CDO
 		PostEditChange();
 	}
 #endif // WITH_EDITOR
 
-PRAGMA_DISABLE_DEPRECATION_WARNINGS
 	BorderColor = FLinearColor::White;
 	AreaPadding = FMargin(1);
 	HeaderPadding = FMargin(4.0f, 2.0f);
diff --git a/Engine/Source/Runtime/UMG/Private/Components/InputKeySelector.cpp b/Engine/Source/Runtime/UMG/Private/Components/InputKeySelector.cpp
--- a/Engine/Source/Runtime/UMG/Private/Components/InputKeySelector.cpp
+++ b/Engine/Source/Runtime/UMG/Private/Components/InputKeySelector.cpp
@@ -7,70 +7,31 @@
 #include "Widgets/DeclarativeSyntaxSupport.h"
 #include "Widgets/Input/SInputKeySelector.h"
 #include "Internationalization/Internationalization.h"
+#include "Styling/DefaultStyleCache.h"
 #include "Styling/UMGCoreStyle.h"
 
 #include UE_INLINE_GENERATED_CPP_BY_NAME(InputKeySelector)
 
 #define LOCTEXT_NAMESPACE "UMG"
 
-static FButtonStyle* DefaultInputKeySelectorButtonStyle = nullptr;
-static FTextBlockStyle* DefaultInputKeySelectorTextStyle = nullptr;
-
-#if WITH_EDITOR 
-static FButtonStyle* EditorInputKeySelectorButtonStyle = nullptr;
-static FTextBlockStyle* EditorInputKeySelectorTextStyle = nullptr;
-#endif
-
 UInputKeySelector::UInputKeySelector( const FObjectInitializer& ObjectInitializer )
 	: Super(ObjectInitializer)
 {
-	if (DefaultInputKeySelectorButtonStyle == nullptr)
-	{
-		DefaultInputKeySelectorButtonStyle = new FButtonStyle(FUMGCoreStyle::Get().GetWidgetStyle<FButtonStyle>("Button"));
-
-		// Unlink UMG default colors.
-		DefaultInputKeySelectorButtonStyle->UnlinkColors();
-	}
-
-	if (DefaultInputKeySelectorTextStyle == nullptr)
-	{
-		DefaultInputKeySelectorTextStyle = new FTextBlockStyle(FUMGCoreStyle::Get().GetWidgetStyle<FTextBlockStyle>("NormalText"));
-
-		// Unlink UMG default colors.
-		DefaultInputKeySelectorTextStyle->UnlinkColors();
-	}
 PRAGMA_DISABLE_DEPRECATION_WARNINGS
-	WidgetStyle = *DefaultInputKeySelectorButtonStyle;
-	TextStyle = *DefaultInputKeySelectorTextStyle;
-PRAGMA_ENABLE_DEPRECATION_WARNINGS
+	WidgetStyle = UE::Slate::Private::FDefaultStyleCache::GetRuntime().GetButtonStyle();
+	TextStyle = UE::Slate::Private::FDefaultStyleCache::GetRuntime().GetTextBlockStyle();
+	
 #if WITH_EDITOR 
-	if (EditorInputKeySelectorButtonStyle == nullptr)
-	{
-		EditorInputKeySelectorButtonStyle = new FButtonStyle(FCoreStyle::Get().GetWidgetStyle<FButtonStyle>("Button"));
-
-		// Unlink UMG Editor colors from the editor settings colors.
-		EditorInputKeySelectorButtonStyle->UnlinkColors();
-	}
-
-	if (EditorInputKeySelectorTextStyle == nullptr)
-	{
-		EditorInputKeySelectorTextStyle = new FTextBlockStyle(FCoreStyle::Get().GetWidgetStyle<FTextBlockStyle>("NormalText"));
-
-		// Unlink UMG Editor colors from the editor settings colors.
-		EditorInputKeySelectorTextStyle->UnlinkColors();
-	}
-
 	if (IsEditorWidget())
 	{
-PRAGMA_DISABLE_DEPRECATION_WARNINGS
-		WidgetStyle = *EditorInputKeySelectorButtonStyle;
-		TextStyle = *EditorInputKeySelectorTextStyle;
-PRAGMA_ENABLE_DEPRECATION_WARNINGS
+		WidgetStyle = UE::Slate::Private::FDefaultStyleCache::GetEditor().GetButtonStyle();
+		TextStyle = UE::Slate::Private::FDefaultStyleCache::GetEditor().GetTextBlockStyle();
+
 		// The CDO isn't an editor widget and thus won't use the editor style, call post edit change to mark difference from CDO
 		PostEditChange();
 	}
 #endif // WITH_EDITOR
-PRAGMA_DISABLE_DEPRECATION_WARNINGS
+
 	KeySelectionText = NSLOCTEXT("InputKeySelector", "DefaultKeySelectionText", "...");
 	NoKeySpecifiedText = NSLOCTEXT("InputKeySelector", "DefaultEmptyText", "Empty");
 	SelectedKey = FInputChord(EKeys::Invalid);
diff --git a/Engine/Source/Runtime/UMG/Private/Components/ListView.cpp b/Engine/Source/Runtime/UMG/Private/Components/ListView.cpp
--- a/Engine/Source/Runtime/UMG/Private/Components/ListView.cpp
+++ b/Engine/Source/Runtime/UMG/Private/Components/ListView.cpp
@@ -5,6 +5,7 @@
 #include "Engine/World.h"
 #include "TimerManager.h"
 #include "Blueprint/ListViewDesignerPreviewItem.h"
+#include "Styling/DefaultStyleCache.h"
 #include "Styling/UMGCoreStyle.h"
 #include "UMGPrivate.h"
 
@@ -15,58 +16,21 @@
 /////////////////////////////////////////////////////
 // UListView
 
-static FTableViewStyle* DefaultListViewStyle = nullptr;
-static FScrollBarStyle* DefaultListViewScrollBarStyle = nullptr;
-
-#if WITH_EDITOR
-static FTableViewStyle* EditorListViewStyle = nullptr;
-static FScrollBarStyle* EditorListViewScrollBarStyle = nullptr;
-#endif 
-
 UListView::UListView(const FObjectInitializer& ObjectInitializer)
 	: Super(ObjectInitializer)
 	, Orientation(EOrientation::Orient_Vertical)
 {
-	if (DefaultListViewStyle == nullptr)
-	{
-		DefaultListViewStyle = new FTableViewStyle(FUMGCoreStyle::Get().GetWidgetStyle<FTableViewStyle>("ListView"));
-
-		// Unlink UMG default colors.
-		DefaultListViewStyle->UnlinkColors();
-	}
-
-	if (DefaultListViewScrollBarStyle == nullptr)
-	{
-		DefaultListViewScrollBarStyle = new FScrollBarStyle(FUMGCoreStyle::Get().GetWidgetStyle<FScrollBarStyle>("Scrollbar"));
-
-		// Unlink UMG default colors.
-		DefaultListViewScrollBarStyle->UnlinkColors();
-	}
-
-	WidgetStyle = *DefaultListViewStyle;
-	ScrollBarStyle = *DefaultListViewScrollBarStyle;
-
+	WidgetStyle = UE::Slate::Private::FDefaultStyleCache::GetRuntime().GetListViewStyle();
+	ScrollBarStyle = UE::Slate::Private::FDefaultStyleCache::GetRuntime().GetScrollBarStyle();
+	
 #if WITH_EDITOR 
-	if (EditorListViewStyle == nullptr)
-	{
-		EditorListViewStyle = new FTableViewStyle(FAppStyle::Get().GetWidgetStyle<FTableViewStyle>("ListView"));
-
-		// Unlink UMG default colors.
-		EditorListViewStyle->UnlinkColors();
-	}
-
-	if (EditorListViewScrollBarStyle == nullptr)
-	{
-		EditorListViewScrollBarStyle = new FScrollBarStyle(FCoreStyle::Get().GetWidgetStyle<FScrollBarStyle>("Scrollbar"));
-
-		// Unlink UMG default colors.
-		EditorListViewScrollBarStyle->UnlinkColors();
-	}
-
 	if (IsEditorWidget())
 	{
-		WidgetStyle = *EditorListViewStyle;
-		ScrollBarStyle = *EditorListViewScrollBarStyle;
+		WidgetStyle = UE::Slate::Private::FDefaultStyleCache::GetEditor().GetListViewStyle();
+		ScrollBarStyle = UE::Slate::Private::FDefaultStyleCache::GetEditor().GetScrollBarStyle();
+
+		// The CDO isn't an editor widget and thus won't use the editor style, call post edit change to mark difference from CDO
+		PostEditChange();
 	}
 #endif // WITH_EDITOR
 }
diff --git a/Engine/Source/Runtime/UMG/Private/Components/MultiLineEditableText.cpp b/Engine/Source/Runtime/UMG/Private/Components/MultiLineEditableText.cpp
--- a/Engine/Source/Runtime/UMG/Private/Components/MultiLineEditableText.cpp
+++ b/Engine/Source/Runtime/UMG/Private/Components/MultiLineEditableText.cpp
@@ -5,6 +5,7 @@
 #include "Engine/Font.h"
 #include "Widgets/DeclarativeSyntaxSupport.h"
 #include "Widgets/Text/SMultiLineEditableText.h"
+#include "Styling/DefaultStyleCache.h"
 #include "Styling/UMGCoreStyle.h"
 #include "Materials/MaterialInterface.h"
 
@@ -15,38 +16,16 @@
 /////////////////////////////////////////////////////
 // UMultiLineEditableText
 
-static FTextBlockStyle* DefaultMultiLineEditableTextStyle = nullptr;
-
-#if WITH_EDITOR
-static FTextBlockStyle* EditorMultiLineEditableTextStyle = nullptr;
-#endif 
-
 PRAGMA_DISABLE_DEPRECATION_WARNINGS
 UMultiLineEditableText::UMultiLineEditableText(const FObjectInitializer& ObjectInitializer)
 	: Super(ObjectInitializer)
 {
-	if (DefaultMultiLineEditableTextStyle == nullptr)
-	{
-		DefaultMultiLineEditableTextStyle = new FTextBlockStyle(FUMGCoreStyle::Get().GetWidgetStyle<FTextBlockStyle>("NormalText"));
-
-		// Unlink UMG default colors.
-		DefaultMultiLineEditableTextStyle->UnlinkColors();
-	}
-
-	WidgetStyle = *DefaultMultiLineEditableTextStyle;
+	WidgetStyle = UE::Slate::Private::FDefaultStyleCache::GetRuntime().GetTextBlockStyle();
 	
 #if WITH_EDITOR 
-	if (EditorMultiLineEditableTextStyle == nullptr)
-	{
-		EditorMultiLineEditableTextStyle = new FTextBlockStyle(FCoreStyle::Get().GetWidgetStyle<FTextBlockStyle>("NormalText"));
-
-		// Unlink UMG Editor colors from the editor settings colors.
-		EditorMultiLineEditableTextStyle->UnlinkColors();
-	}
-
 	if (IsEditorWidget())
 	{
-		WidgetStyle = *EditorMultiLineEditableTextStyle;
+		WidgetStyle = UE::Slate::Private::FDefaultStyleCache::GetEditor().GetTextBlockStyle();
 
 		// The CDO isn't an editor widget and thus won't use the editor style, call post edit change to mark difference from CDO
 		PostEditChange();
diff --git a/Engine/Source/Runtime/UMG/Private/Components/MultiLineEditableTextBox.cpp b/Engine/Source/Runtime/UMG/Private/Components/MultiLineEditableTextBox.cpp
--- a/Engine/Source/Runtime/UMG/Private/Components/MultiLineEditableTextBox.cpp
+++ b/Engine/Source/Runtime/UMG/Private/Components/MultiLineEditableTextBox.cpp
@@ -7,6 +7,7 @@
 #include "Engine/Font.h"
 #include "Widgets/DeclarativeSyntaxSupport.h"
 #include "Widgets/Input/SMultiLineEditableTextBox.h"
+#include "Styling/DefaultStyleCache.h"
 #include "Styling/UMGCoreStyle.h"
 
 #include UE_INLINE_GENERATED_CPP_BY_NAME(MultiLineEditableTextBox)
@@ -16,27 +17,13 @@
 /////////////////////////////////////////////////////
 // UMultiLineEditableTextBox
 
-static FEditableTextBoxStyle* DefaultMultiLineEditableTextBoxStyle = nullptr;
-
-#if WITH_EDITOR
-static FEditableTextBoxStyle* EditorMultiLineEditableTextBoxStyle = nullptr;
-#endif 
-
 PRAGMA_DISABLE_DEPRECATION_WARNINGS
 UMultiLineEditableTextBox::UMultiLineEditableTextBox(const FObjectInitializer& ObjectInitializer)
 	: Super(ObjectInitializer)
 {
-	if (DefaultMultiLineEditableTextBoxStyle == nullptr)
-	{
-		DefaultMultiLineEditableTextBoxStyle = new FEditableTextBoxStyle(FUMGCoreStyle::Get().GetWidgetStyle<FEditableTextBoxStyle>("NormalEditableTextBox"));
-
-		// Unlink UMG default colors.
-		DefaultMultiLineEditableTextBoxStyle->UnlinkColors();
-	}
-
-	WidgetStyle = *DefaultMultiLineEditableTextBoxStyle;
+	WidgetStyle = UE::Slate::Private::FDefaultStyleCache::GetRuntime().GetEditableTextBoxStyle();
 #if WITH_EDITOR
-	TextStyle_DEPRECATED = DefaultMultiLineEditableTextBoxStyle->TextStyle;
+	TextStyle_DEPRECATED = WidgetStyle.TextStyle;
 #endif
 
 	if (!IsRunningDedicatedServer())
@@ -47,20 +34,12 @@ UMultiLineEditableTextBox::UMultiLineEditableTextBox(const FObjectInitializer& O
 		//so in the case the Font object is replaced by a null one, we have to keep the composite one as a fallback.
 		Font.CompositeFont = WidgetStyle.TextStyle.Font.CompositeFont;
 	}
-
-#if WITH_EDITOR 
-	if (EditorMultiLineEditableTextBoxStyle == nullptr)
-	{
-		EditorMultiLineEditableTextBoxStyle = new FEditableTextBoxStyle(FCoreStyle::Get().GetWidgetStyle<FEditableTextBoxStyle>("NormalEditableTextBox"));
-
-		// Unlink UMG Editor colors from the editor settings colors.
-		EditorMultiLineEditableTextBoxStyle->UnlinkColors();
-	}
 	
+#if WITH_EDITOR 
 	if (IsEditorWidget())
 	{
-		WidgetStyle = *EditorMultiLineEditableTextBoxStyle;
-		TextStyle_DEPRECATED = DefaultMultiLineEditableTextBoxStyle->TextStyle;
+		WidgetStyle = UE::Slate::Private::FDefaultStyleCache::GetEditor().GetEditableTextBoxStyle();
+		TextStyle_DEPRECATED = WidgetStyle.TextStyle;
 
 		// The CDO isn't an editor widget and thus won't use the editor style, call post edit change to mark difference from CDO
 		PostEditChange();
diff --git a/Engine/Source/Runtime/UMG/Private/Components/ProgressBar.cpp b/Engine/Source/Runtime/UMG/Private/Components/ProgressBar.cpp
--- a/Engine/Source/Runtime/UMG/Private/Components/ProgressBar.cpp
+++ b/Engine/Source/Runtime/UMG/Private/Components/ProgressBar.cpp
@@ -2,6 +2,7 @@
 
 #include "Components/ProgressBar.h"
 #include "Slate/SlateBrushAsset.h"
+#include "Styling/DefaultStyleCache.h"
 #include "Styling/UMGCoreStyle.h"
 
 #include UE_INLINE_GENERATED_CPP_BY_NAME(ProgressBar)
@@ -11,38 +12,16 @@
 /////////////////////////////////////////////////////
 // UProgressBar
 
-static FProgressBarStyle* DefaultProgressBarStyle = nullptr;
-
-#if WITH_EDITOR
-static FProgressBarStyle* EditorProgressBarStyle = nullptr;
-#endif 
-
 UProgressBar::UProgressBar(const FObjectInitializer& ObjectInitializer)
 	: Super(ObjectInitializer)
 {
-	if (DefaultProgressBarStyle == nullptr)
-	{
-		DefaultProgressBarStyle = new FProgressBarStyle(FUMGCoreStyle::Get().GetWidgetStyle<FProgressBarStyle>("ProgressBar"));
-
-		// Unlink UMG default colors.
-		DefaultProgressBarStyle->UnlinkColors();
-	}
-
 PRAGMA_DISABLE_DEPRECATION_WARNINGS
-	WidgetStyle = *DefaultProgressBarStyle;
-
+	WidgetStyle = UE::Slate::Private::FDefaultStyleCache::GetRuntime().GetProgressBarStyle();
+	
 #if WITH_EDITOR 
-	if (EditorProgressBarStyle == nullptr)
-	{
-		EditorProgressBarStyle = new FProgressBarStyle(FCoreStyle::Get().GetWidgetStyle<FProgressBarStyle>("ProgressBar"));
-
-		// Unlink UMG Editor colors from the editor settings colors.
-		EditorProgressBarStyle->UnlinkColors();
-	}
-
 	if (IsEditorWidget())
 	{
-		WidgetStyle = *EditorProgressBarStyle;
+		WidgetStyle = UE::Slate::Private::FDefaultStyleCache::GetEditor().GetProgressBarStyle();
 
 		// The CDO isn't an editor widget and thus won't use the editor style, call post edit change to mark difference from CDO
 		PostEditChange();
diff --git a/Engine/Source/Runtime/UMG/Private/Components/ScrollBar.cpp b/Engine/Source/Runtime/UMG/Private/Components/ScrollBar.cpp
--- a/Engine/Source/Runtime/UMG/Private/Components/ScrollBar.cpp
+++ b/Engine/Source/Runtime/UMG/Private/Components/ScrollBar.cpp
@@ -2,6 +2,7 @@
 
 #include "Components/ScrollBar.h"
 #include "UObject/EditorObjectVersion.h"
+#include "Styling/DefaultStyleCache.h"
 #include "Styling/UMGCoreStyle.h"
 
 #include UE_INLINE_GENERATED_CPP_BY_NAME(ScrollBar)
@@ -11,15 +12,6 @@
 /////////////////////////////////////////////////////
 // UScrollBar
 
-namespace
-{
-	static FScrollBarStyle* DefaultScrollBarStyle = nullptr;
-}
-
-#if WITH_EDITOR
-static FScrollBarStyle* EditorScrollBarStyle = nullptr;
-#endif 
-
 UScrollBar::UScrollBar(const FObjectInitializer& ObjectInitializer)
 	: Super(ObjectInitializer)
 {
@@ -31,35 +23,19 @@ PRAGMA_DISABLE_DEPRECATION_WARNINGS
 	Orientation = Orient_Vertical;
 	Thickness = FVector2D(16.0f, 16.0f);
 	Padding = FMargin(2.0f);
-PRAGMA_ENABLE_DEPRECATION_WARNINGS
-	if (DefaultScrollBarStyle == nullptr)
-	{
-		DefaultScrollBarStyle = new FScrollBarStyle(FUMGCoreStyle::Get().GetWidgetStyle<FScrollBarStyle>("Scrollbar"));
 
-		// Unlink UMG default colors.
-		DefaultScrollBarStyle->UnlinkColors();
-	}
-PRAGMA_DISABLE_DEPRECATION_WARNINGS
-	WidgetStyle = *DefaultScrollBarStyle;
-PRAGMA_ENABLE_DEPRECATION_WARNINGS
+	WidgetStyle = UE::Slate::Private::FDefaultStyleCache::GetRuntime().GetScrollBarStyle();
+	
 #if WITH_EDITOR 
-	if (EditorScrollBarStyle == nullptr)
-	{
-		EditorScrollBarStyle = new FScrollBarStyle(FCoreStyle::Get().GetWidgetStyle<FScrollBarStyle>("Scrollbar"));
-
-		// Unlink UMG Editor colors from the editor settings colors.
-		EditorScrollBarStyle->UnlinkColors();
-	}
-
 	if (IsEditorWidget())
 	{
-PRAGMA_DISABLE_DEPRECATION_WARNINGS
-		WidgetStyle = *EditorScrollBarStyle;
-PRAGMA_ENABLE_DEPRECATION_WARNINGS
+		WidgetStyle = UE::Slate::Private::FDefaultStyleCache::GetEditor().GetScrollBarStyle();
+
 		// The CDO isn't an editor widget and thus won't use the editor style, call post edit change to mark difference from CDO
 		PostEditChange();
 	}
 #endif // WITH_EDITOR
+PRAGMA_ENABLE_DEPRECATION_WARNINGS
 }
 
 void UScrollBar::ReleaseSlateResources(bool bReleaseChildren)
diff --git a/Engine/Source/Runtime/UMG/Private/Components/ScrollBox.cpp b/Engine/Source/Runtime/UMG/Private/Components/ScrollBox.cpp
--- a/Engine/Source/Runtime/UMG/Private/Components/ScrollBox.cpp
+++ b/Engine/Source/Runtime/UMG/Private/Components/ScrollBox.cpp
@@ -4,6 +4,7 @@
 #include "Containers/Ticker.h"
 #include "Components/ScrollBoxSlot.h"
 #include "UObject/EditorObjectVersion.h"
+#include "Styling/DefaultStyleCache.h"
 #include "Styling/UMGCoreStyle.h"
 
 #include UE_INLINE_GENERATED_CPP_BY_NAME(ScrollBox)
@@ -13,14 +14,6 @@
 /////////////////////////////////////////////////////
 // UScrollBox
 
-static FScrollBoxStyle* DefaultScrollBoxStyle = nullptr;
-static FScrollBarStyle* DefaultScrollBoxBarStyle = nullptr;
-
-#if WITH_EDITOR
-static FScrollBoxStyle* EditorScrollBoxStyle = nullptr;
-static FScrollBarStyle* EditorScrollBoxBarStyle = nullptr;
-#endif 
-
 UScrollBox::UScrollBox(const FObjectInitializer& ObjectInitializer)
 	: Super(ObjectInitializer)
 	, Orientation(Orient_Vertical)
@@ -41,49 +34,18 @@ UScrollBox::UScrollBox(const FObjectInitializer& ObjectInitializer)
 
 	SetVisibilityInternal(ESlateVisibility::Visible);
 	SetClipping(EWidgetClipping::ClipToBounds);
-
-	if (DefaultScrollBoxStyle == nullptr)
-	{
-		DefaultScrollBoxStyle = new FScrollBoxStyle(FUMGCoreStyle::Get().GetWidgetStyle<FScrollBoxStyle>("ScrollBox"));
-
-		// Unlink UMG default colors.
-		DefaultScrollBoxStyle->UnlinkColors();
-	}
-
-	if (DefaultScrollBoxBarStyle == nullptr)
-	{
-		DefaultScrollBoxBarStyle = new FScrollBarStyle(FUMGCoreStyle::Get().GetWidgetStyle<FScrollBarStyle>("ScrollBar"));
-
-		// Unlink UMG default colors.
-		DefaultScrollBoxBarStyle->UnlinkColors();
-	}
+	
 PRAGMA_DISABLE_DEPRECATION_WARNINGS
-	WidgetStyle = *DefaultScrollBoxStyle;
-	WidgetBarStyle = *DefaultScrollBoxBarStyle;
+	WidgetStyle = UE::Slate::Private::FDefaultStyleCache::GetRuntime().GetScrollBoxStyle();
+	WidgetBarStyle = UE::Slate::Private::FDefaultStyleCache::GetRuntime().GetScrollBarStyle();
 PRAGMA_ENABLE_DEPRECATION_WARNINGS
 
 #if WITH_EDITOR 
-	if (EditorScrollBoxStyle == nullptr)
-	{
-		EditorScrollBoxStyle = new FScrollBoxStyle(FCoreStyle::Get().GetWidgetStyle<FScrollBoxStyle>("ScrollBox"));
-
-		// Unlink UMG Editor colors from the editor settings colors.
-		EditorScrollBoxStyle->UnlinkColors();
-	}
-
-	if (EditorScrollBoxBarStyle == nullptr)
-	{
-		EditorScrollBoxBarStyle = new FScrollBarStyle(FCoreStyle::Get().GetWidgetStyle<FScrollBarStyle>("ScrollBar"));
-
-		// Unlink UMG Editor colors from the editor settings colors.
-		EditorScrollBoxBarStyle->UnlinkColors();
-	}
-	
 	if (IsEditorWidget())
 	{
 PRAGMA_DISABLE_DEPRECATION_WARNINGS
-		WidgetStyle = *EditorScrollBoxStyle;
-		WidgetBarStyle = *EditorScrollBoxBarStyle;
+		WidgetStyle = UE::Slate::Private::FDefaultStyleCache::GetEditor().GetScrollBoxStyle();
+		WidgetBarStyle = UE::Slate::Private::FDefaultStyleCache::GetEditor().GetScrollBarStyle();
 PRAGMA_ENABLE_DEPRECATION_WARNINGS
 
 		// The CDO isn't an editor widget and thus won't use the editor style, call post edit change to mark difference from CDO
diff --git a/Engine/Source/Runtime/UMG/Private/Components/Slider.cpp b/Engine/Source/Runtime/UMG/Private/Components/Slider.cpp
--- a/Engine/Source/Runtime/UMG/Private/Components/Slider.cpp
+++ b/Engine/Source/Runtime/UMG/Private/Components/Slider.cpp
@@ -3,6 +3,7 @@
 #include "Components/Slider.h"
 #include "Widgets/DeclarativeSyntaxSupport.h"
 #include "Widgets/Input/SSlider.h"
+#include "Styling/DefaultStyleCache.h"
 #include "Styling/UMGCoreStyle.h"
 
 #include UE_INLINE_GENERATED_CPP_BY_NAME(Slider)
@@ -12,12 +13,6 @@
 /////////////////////////////////////////////////////
 // USlider
 
-static FSliderStyle* DefaultSliderStyle = nullptr;
-
-#if WITH_EDITOR
-static FSliderStyle* EditorSliderStyle = nullptr;
-#endif 
-
 USlider::USlider(const FObjectInitializer& ObjectInitializer)
 	: Super(ObjectInitializer)
 {
@@ -32,29 +27,12 @@ PRAGMA_DISABLE_DEPRECATION_WARNINGS
 	MouseUsesStep = false;
 	RequiresControllerLock = true;
 
-	if (DefaultSliderStyle == nullptr)
-	{
-		DefaultSliderStyle = new FSliderStyle(FUMGCoreStyle::Get().GetWidgetStyle<FSliderStyle>("Slider"));
-
-		// Unlink UMG default colors.
-		DefaultSliderStyle->UnlinkColors();
-	}
-
-	WidgetStyle = *DefaultSliderStyle;
-
-
-#if WITH_EDITOR 
-	if (EditorSliderStyle == nullptr)
-	{
-		EditorSliderStyle = new FSliderStyle(FCoreStyle::Get().GetWidgetStyle<FSliderStyle>("Slider"));
-
-		// Unlink UMG Editor colors from the editor settings colors.
-		EditorSliderStyle->UnlinkColors();
-	}
+	WidgetStyle = UE::Slate::Private::FDefaultStyleCache::GetRuntime().GetSliderStyle();
 	
+#if WITH_EDITOR 
 	if (IsEditorWidget())
 	{
-		WidgetStyle = *EditorSliderStyle;
+		WidgetStyle = UE::Slate::Private::FDefaultStyleCache::GetEditor().GetSliderStyle();
 
 		// The CDO isn't an editor widget and thus won't use the editor style, call post edit change to mark difference from CDO
 		PostEditChange();
diff --git a/Engine/Source/Runtime/UMG/Private/Components/SpinBox.cpp b/Engine/Source/Runtime/UMG/Private/Components/SpinBox.cpp
--- a/Engine/Source/Runtime/UMG/Private/Components/SpinBox.cpp
+++ b/Engine/Source/Runtime/UMG/Private/Components/SpinBox.cpp
@@ -3,6 +3,7 @@
 #include "Components/SpinBox.h"
 #include "UObject/ConstructorHelpers.h"
 #include "Engine/Font.h"
+#include "Styling/DefaultStyleCache.h"
 #include "Styling/UMGCoreStyle.h"
 
 #include UE_INLINE_GENERATED_CPP_BY_NAME(SpinBox)
@@ -12,12 +13,6 @@
 /////////////////////////////////////////////////////
 // USpinBox
 
-static FSpinBoxStyle* DefaultSpinBoxStyle = nullptr;
-
-#if WITH_EDITOR
-static FSpinBoxStyle* EditorSpinBoxStyle = nullptr;
-#endif 
-
 USpinBox::USpinBox(const FObjectInitializer& ObjectInitializer)
 	: Super(ObjectInitializer)
 {
@@ -43,38 +38,19 @@ PRAGMA_DISABLE_DEPRECATION_WARNINGS
 	ClearKeyboardFocusOnCommit = false;
 	SelectAllTextOnCommit = true;
 	KeyboardType = EVirtualKeyboardType::Number;
-PRAGMA_ENABLE_DEPRECATION_WARNINGS
-
-	if (DefaultSpinBoxStyle == nullptr)
-	{
-		DefaultSpinBoxStyle = new FSpinBoxStyle(FUMGCoreStyle::Get().GetWidgetStyle<FSpinBoxStyle>("SpinBox"));
-
-		// Unlink UMG default colors.
-		DefaultSpinBoxStyle->UnlinkColors();
-	}
-PRAGMA_DISABLE_DEPRECATION_WARNINGS
-	WidgetStyle = *DefaultSpinBoxStyle;
-PRAGMA_ENABLE_DEPRECATION_WARNINGS
-#if WITH_EDITOR 
-	if (EditorSpinBoxStyle == nullptr)
-	{
-		EditorSpinBoxStyle = new FSpinBoxStyle(FCoreStyle::Get().GetWidgetStyle<FSpinBoxStyle>("SpinBox"));
 
-		// Unlink UMG Editor colors from the editor settings colors.
-		EditorSpinBoxStyle->UnlinkColors();
-	}
+	WidgetStyle = UE::Slate::Private::FDefaultStyleCache::GetRuntime().GetSpinBoxStyle();
 	
+#if WITH_EDITOR 
 	if (IsEditorWidget())
 	{
-PRAGMA_DISABLE_DEPRECATION_WARNINGS
-		WidgetStyle = *EditorSpinBoxStyle;
-PRAGMA_ENABLE_DEPRECATION_WARNINGS
+		WidgetStyle = UE::Slate::Private::FDefaultStyleCache::GetEditor().GetSpinBoxStyle();
+
 		// The CDO isn't an editor widget and thus won't use the editor style, call post edit change to mark difference from CDO
 		PostEditChange();
 	}
 #endif // WITH_EDITOR
 
-PRAGMA_DISABLE_DEPRECATION_WARNINGS
 	ForegroundColor = WidgetStyle.ForegroundColor;
 PRAGMA_ENABLE_DEPRECATION_WARNINGS
 }
diff --git a/Engine/Source/Runtime/UMG/Private/Components/Throbber.cpp b/Engine/Source/Runtime/UMG/Private/Components/Throbber.cpp
--- a/Engine/Source/Runtime/UMG/Private/Components/Throbber.cpp
+++ b/Engine/Source/Runtime/UMG/Private/Components/Throbber.cpp
@@ -4,6 +4,7 @@
 #include "SlateFwd.h"
 #include "SlateGlobals.h"
 #include "Slate/SlateBrushAsset.h"
+#include "Styling/DefaultStyleCache.h"
 #include "Styling/UMGCoreStyle.h"
 
 #include UE_INLINE_GENERATED_CPP_BY_NAME(Throbber)
@@ -13,12 +14,6 @@
 /////////////////////////////////////////////////////
 // UThrobber
 
-static FSlateBrush* DefaultThrobberBrush = nullptr;
-
-#if WITH_EDITOR
-static FSlateBrush* EditorThrobberBrush = nullptr;
-#endif 
-
 UThrobber::UThrobber(const FObjectInitializer& ObjectInitializer)
 	: Super(ObjectInitializer)
 {
@@ -28,29 +23,13 @@ PRAGMA_DISABLE_DEPRECATION_WARNINGS
 	bAnimateVertically = true;
 	bAnimateHorizontally = true;
 	bAnimateOpacity = true;
-
-	if (DefaultThrobberBrush == nullptr)
-	{
-		DefaultThrobberBrush = new FSlateBrush(*FUMGCoreStyle::Get().GetBrush("Throbber.Chunk"));
-
-		// Unlink UMG default colors.
-		DefaultThrobberBrush->UnlinkColors();
-	}
-
-	Image = *DefaultThrobberBrush;
-
-#if WITH_EDITOR 
-	if (EditorThrobberBrush == nullptr)
-	{
-		EditorThrobberBrush = new FSlateBrush(*FCoreStyle::Get().GetBrush("Throbber.Chunk"));
-
-		// Unlink UMG Editor colors from the editor settings colors.
-		EditorThrobberBrush->UnlinkColors();
-	}
 	
+	Image = UE::Slate::Private::FDefaultStyleCache::GetRuntime().GetThrobberBrush();
+	
+#if WITH_EDITOR 
 	if (IsEditorWidget())
 	{
-		Image = *EditorThrobberBrush;
+		Image = UE::Slate::Private::FDefaultStyleCache::GetEditor().GetThrobberBrush();
 
 		// The CDO isn't an editor widget and thus won't use the editor style, call post edit change to mark difference from CDO
 		PostEditChange();
diff --git a/Engine/Source/Runtime/UMG/Private/Components/TreeView.cpp b/Engine/Source/Runtime/UMG/Private/Components/TreeView.cpp
--- a/Engine/Source/Runtime/UMG/Private/Components/TreeView.cpp
+++ b/Engine/Source/Runtime/UMG/Private/Components/TreeView.cpp
@@ -1,6 +1,7 @@
 // Copyright Epic Games, Inc. All Rights Reserved.
 
 #include "Components/TreeView.h"
+#include "Styling/DefaultStyleCache.h"
 #include "Styling/UMGCoreStyle.h"
 
 #include UE_INLINE_GENERATED_CPP_BY_NAME(TreeView)
@@ -8,37 +9,18 @@
 /////////////////////////////////////////////////////
 // UTreeView
 
-static FTableViewStyle* DefaultTreeViewStyle = nullptr;
-
-#if WITH_EDITOR
-static FTableViewStyle* EditorTreeViewStyle = nullptr;
-#endif 
-
 UTreeView::UTreeView(const FObjectInitializer& ObjectInitializer)
 	: Super(ObjectInitializer)
 {
-	if (DefaultTreeViewStyle == nullptr)
-	{
-		DefaultTreeViewStyle = new FTableViewStyle(FUMGCoreStyle::Get().GetWidgetStyle<FTableViewStyle>("TreeView"));
-
-		// Unlink UMG default colors.
-		DefaultTreeViewStyle->UnlinkColors();
-	}
-
-	WidgetStyle = *DefaultTreeViewStyle;
-
+	WidgetStyle = UE::Slate::Private::FDefaultStyleCache::GetRuntime().GetTreeViewStyle();
+	
 #if WITH_EDITOR 
-	if (EditorTreeViewStyle == nullptr)
-	{
-		EditorTreeViewStyle = new FTableViewStyle(FAppStyle::Get().GetWidgetStyle<FTableViewStyle>("TreeView"));
-
-		// Unlink UMG default colors.
-		EditorTreeViewStyle->UnlinkColors();
-	}
-
 	if (IsEditorWidget())
 	{
-		WidgetStyle = *EditorTreeViewStyle;
+		WidgetStyle = UE::Slate::Private::FDefaultStyleCache::GetEditor().GetTreeViewStyle();
+
+		// The CDO isn't an editor widget and thus won't use the editor style, call post edit change to mark difference from CDO
+		PostEditChange();
 	}
 #endif // WITH_EDITOR
 }
diff --git a/Engine/Source/Runtime/UMG/Private/DefaultStyleCache.cpp b/Engine/Source/Runtime/UMG/Private/DefaultStyleCache.cpp
deleted file mode 100644
diff --git a/Engine/Source/Runtime/UMG/Private/DefaultStyleCache.h b/Engine/Source/Runtime/UMG/Private/DefaultStyleCache.h
deleted file mode 100644
