Proper fix private linux namespace

diff --git a/Engine/Source/Runtime/SlateCore/Public/Styling/DefaultStyleCache.h b/Engine/Source/Runtime/SlateCore/Public/Styling/DefaultStyleCache.h
--- a/Engine/Source/Runtime/SlateCore/Public/Styling/DefaultStyleCache.h
+++ b/Engine/Source/Runtime/SlateCore/Public/Styling/DefaultStyleCache.h
@@ -5,6 +5,8 @@
 #include "CoreMinimal.h"
 #include "Styling/SlateTypes.h"
 
+class FLazySingleton;
+
 namespace UE::Slate::Private
 {
 
@@ -14,8 +16,6 @@ namespace UE::Slate::Private
  */
 struct SLATECORE_API FDefaultStyleCache
 {
-	FDefaultStyleCache();
-
 	/**
 	 * Internal style default holder struct, used so that one can get a bundle of defaults relevant for runtime or editor
 	 */
@@ -109,7 +109,9 @@ private:
 	/** Singleton getter, however private used since GetRuntime & GetEditor is preferred for styles */
 	static FDefaultStyleCache& Get();
 
-	friend class FLazySingleton;
+	friend ::FLazySingleton;
+
+	FDefaultStyleCache();
 
 	FStyles Runtime;
 
