Fixed mismatched copy constructor/assignment operator in FSlateClippingState.

diff --git a/Engine/Source/Runtime/SlateCore/Private/Layout/Clipping.cpp b/Engine/Source/Runtime/SlateCore/Private/Layout/Clipping.cpp
--- a/Engine/Source/Runtime/SlateCore/Private/Layout/Clipping.cpp
+++ b/Engine/Source/Runtime/SlateCore/Private/Layout/Clipping.cpp
@@ -208,17 +208,6 @@ FSlateClippingState::FSlateClippingState(EClippingFlags InFlags /*= EClippingFla
 {
 }
 
-FSlateClippingState::FSlateClippingState(const FSlateClippingState& Other)
-	: StencilQuads(Other.StencilQuads)
-	, ScissorRect(Other.ScissorRect)
-	, Flags(Other.Flags)
-#if WITH_SLATE_DEBUGGING
-	, Debugging_StateIndex(Other.Debugging_StateIndex)
-	, Debugging_StateIndexFromFrame(Other.Debugging_StateIndexFromFrame)
-#endif
-{
-}
-
 bool FSlateClippingState::IsPointInside(const UE::Slate::FDeprecateVector2DParameter& Point) const
 {
 	if (ScissorRect.IsSet())
diff --git a/Engine/Source/Runtime/SlateCore/Public/Layout/Clipping.h b/Engine/Source/Runtime/SlateCore/Public/Layout/Clipping.h
--- a/Engine/Source/Runtime/SlateCore/Public/Layout/Clipping.h
+++ b/Engine/Source/Runtime/SlateCore/Public/Layout/Clipping.h
@@ -208,8 +208,6 @@ class SLATECORE_API FSlateClippingState
 public:
 	FSlateClippingState(EClippingFlags InFlags = EClippingFlags::None);
 	
-	FSlateClippingState(const FSlateClippingState& Other);
-
 	/** Is a point inside the clipping state? */
 	bool IsPointInside(const UE::Slate::FDeprecateVector2DParameter& Point) const;
 
