Fix the SLATE_EVENT macro's return expression in the _UObject case.
#rnx

diff --git a/Engine/Source/Runtime/SlateCore/Public/Widgets/DeclarativeSyntaxSupport.h b/Engine/Source/Runtime/SlateCore/Public/Widgets/DeclarativeSyntaxSupport.h
--- a/Engine/Source/Runtime/SlateCore/Public/Widgets/DeclarativeSyntaxSupport.h
+++ b/Engine/Source/Runtime/SlateCore/Public/Widgets/DeclarativeSyntaxSupport.h
@@ -534,7 +534,7 @@ struct NamedSlotProperty
 		WidgetArgsType& EventName##_UObject( UserClass* InUserObject, typename DelegateName::template TMethodPtr< UserClass, VarTypes... > InFunc, VarTypes... Vars )	\
 		{ \
 			_##EventName = DelegateName::CreateUObject( InUserObject, InFunc, Vars... ); \
-			return *this; \
+			return static_cast<WidgetArgsType*>(this)->Me(); \
 		} \
 		template< class UserClass, typename... VarTypes >	\
 		WidgetArgsType& EventName##_UObject( UserClass* InUserObject, typename DelegateName::template TConstMethodPtr< UserClass, VarTypes... > InFunc, VarTypes... Vars )	\
