Change the combo buttons on the status bar to use the regular "SimpleComboButton" style with chevrons

diff --git a/Engine/Plugins/TraceUtilities/Source/EditorTraceUtilities/Private/InsightsStatusBar/SInsightsStatusBar.cpp b/Engine/Plugins/TraceUtilities/Source/EditorTraceUtilities/Private/InsightsStatusBar/SInsightsStatusBar.cpp
--- a/Engine/Plugins/TraceUtilities/Source/EditorTraceUtilities/Private/InsightsStatusBar/SInsightsStatusBar.cpp
+++ b/Engine/Plugins/TraceUtilities/Source/EditorTraceUtilities/Private/InsightsStatusBar/SInsightsStatusBar.cpp
@@ -131,7 +131,7 @@ void SInsightsStatusBarWidget::Construct(const FArguments& InArgs)
 			SNew(SComboButton)
 			.ContentPadding(FMargin(6.0f, 0.0f))
 			.MenuPlacement(MenuPlacement_AboveAnchor)
-			.ComboButtonStyle(&FAppStyle::Get().GetWidgetStyle<FComboButtonStyle>("StatusBar.StatusBarComboButton"))
+			.ComboButtonStyle(&FAppStyle::Get().GetWidgetStyle<FComboButtonStyle>("SimpleComboButton"))
 			.OnGetMenuContent(this, &SInsightsStatusBarWidget::MakeTraceMenu)
 			.HasDownArrow(true)
 			.ToolTipText(this, &SInsightsStatusBarWidget::GetTitleToolTipText)
diff --git a/Engine/Source/Editor/DerivedDataEditor/Private/SDerivedDataStatusBar.cpp b/Engine/Source/Editor/DerivedDataEditor/Private/SDerivedDataStatusBar.cpp
--- a/Engine/Source/Editor/DerivedDataEditor/Private/SDerivedDataStatusBar.cpp
+++ b/Engine/Source/Editor/DerivedDataEditor/Private/SDerivedDataStatusBar.cpp
@@ -146,7 +146,7 @@ void SDerivedDataStatusBarWidget::Construct(const FArguments& InArgs)
 		SNew(SComboButton)
 		.ContentPadding(FMargin(6.0f, 0.0f))
 		.MenuPlacement(MenuPlacement_AboveAnchor)
-		.ComboButtonStyle(&FAppStyle::Get().GetWidgetStyle<FComboButtonStyle>("StatusBar.StatusBarComboButton"))
+		.ComboButtonStyle(&FAppStyle::Get().GetWidgetStyle<FComboButtonStyle>("SimpleComboButton"))
 		.ButtonContent()
 		[
 
diff --git a/Engine/Source/Editor/StatusBar/Private/SourceControlMenuHelpers.cpp b/Engine/Source/Editor/StatusBar/Private/SourceControlMenuHelpers.cpp
--- a/Engine/Source/Editor/StatusBar/Private/SourceControlMenuHelpers.cpp
+++ b/Engine/Source/Editor/StatusBar/Private/SourceControlMenuHelpers.cpp
@@ -664,7 +664,7 @@ TSharedRef<SWidget> FSourceControlMenuHelpers::MakeSourceControlStatusWidget()
 			.ContentPadding(FMargin(6.0f, 0.0f))
 			.ToolTipText_Static(&FSourceControlMenuHelpers::GetSourceControlTooltip)
 			.MenuPlacement(MenuPlacement_AboveAnchor)
-			.ComboButtonStyle(&FAppStyle::Get().GetWidgetStyle<FComboButtonStyle>("StatusBar.StatusBarComboButton"))
+			.ComboButtonStyle(&FAppStyle::Get().GetWidgetStyle<FComboButtonStyle>("SimpleComboButton"))
 			.ButtonContent()
 			[
 				SNew(SHorizontalBox)
diff --git a/Engine/Source/Runtime/Slate/Private/Framework/MultiBox/SToolBarStackButtonBlock.cpp b/Engine/Source/Runtime/Slate/Private/Framework/MultiBox/SToolBarStackButtonBlock.cpp
--- a/Engine/Source/Runtime/Slate/Private/Framework/MultiBox/SToolBarStackButtonBlock.cpp
+++ b/Engine/Source/Runtime/Slate/Private/Framework/MultiBox/SToolBarStackButtonBlock.cpp
@@ -264,7 +264,7 @@ void SToolBarStackButtonBlock::BuildMultiBlockWidget(const ISlateStyle* StyleSet
 	// Create the content for our button
 	EMultiBlockLocation::Type BlockLocation = GetMultiBlockLocation();
 	FName BlockStyle = EMultiBlockLocation::ToName(ISlateStyle::Join(StyleName, ".Button"), BlockLocation);
-	const FComboButtonStyle* ComboButtonStyle = &FAppStyle::Get().GetWidgetStyle<FComboButtonStyle>("StatusBar.StatusBarComboButton");
+	const FComboButtonStyle* ComboButtonStyle = &FAppStyle::Get().GetWidgetStyle<FComboButtonStyle>("SimpleComboButton");
 	const FButtonStyle* ButtonStyle = BlockLocation == EMultiBlockLocation::None ? &ToolBarStyle.ButtonStyle : &StyleSet->GetWidgetStyle<FButtonStyle>(BlockStyle);
 	const FCheckBoxStyle* CheckStyle = BlockLocation == EMultiBlockLocation::None ? &ToolBarStyle.ToggleButton : &StyleSet->GetWidgetStyle<FCheckBoxStyle>(BlockStyle);
 
diff --git a/Engine/Source/Runtime/SlateCore/Private/Styling/StarshipCoreStyle.cpp b/Engine/Source/Runtime/SlateCore/Private/Styling/StarshipCoreStyle.cpp
--- a/Engine/Source/Runtime/SlateCore/Private/Styling/StarshipCoreStyle.cpp
+++ b/Engine/Source/Runtime/SlateCore/Private/Styling/StarshipCoreStyle.cpp
@@ -1851,12 +1851,8 @@ void FStarshipCoreStyle::SetupDockingStyles(TSharedRef<FStyle>& Style)
 		.SetNormalPadding(FMargin(2, 2, 2, 2))
 		.SetPressedPadding(FMargin(2, 3, 2, 1));
 
-	const FComboButtonStyle StatusBarComboButton = FComboButtonStyle(Style->GetWidgetStyle<FComboButtonStyle>("ComboButton"))
-		.SetDownArrowImage(IMAGE_BRUSH_SVG("Starship/CoreWidgets/ComboBox/corner-dropdown", FVector2f(7.0f, 7.0f)))
-		.SetButtonStyle(StatusBarButton)
-		.SetDownArrowPadding(FMargin(0.0f))
-		.SetDownArrowAlignment(EVerticalAlignment::VAlign_Bottom);
-
+	const FComboButtonStyle StatusBarComboButton = FComboButtonStyle(Style->GetWidgetStyle<FComboButtonStyle>("SimpleComboButton"));
+	
 	const FComboButtonStyle StatusBarEllipsisComboButton = FComboButtonStyle(Style->GetWidgetStyle<FComboButtonStyle>("ComboButton"))
 		.SetDownArrowImage(IMAGE_BRUSH_SVG("Starship/Common/ellipsis-vertical-narrow", FVector2f(6.f, 24.f)))
 		.SetButtonStyle(StatusBarButton)
