Use static indexes for slate shaped text and text

diff --git a/Engine/Source/Runtime/SlateCore/Private/Rendering/ElementBatcher.cpp b/Engine/Source/Runtime/SlateCore/Private/Rendering/ElementBatcher.cpp
--- a/Engine/Source/Runtime/SlateCore/Private/Rendering/ElementBatcher.cpp
+++ b/Engine/Source/Runtime/SlateCore/Private/Rendering/ElementBatcher.cpp
@@ -1289,6 +1289,10 @@ void FSlateElementBatcher::AddTextElement(const FSlateDrawElement& DrawElement)
 
 		uint32 NumChars = Len;
 
+		// The number of glyphs can technically be unbound, so pick some number that if smaller we use the preallocated index array
+		const int32 MaxPreAllocatedGlyphIndicies = 4096 + 1024;
+		const bool bUseStaticIndicies = NumChars < MaxPreAllocatedGlyphIndicies;
+
 		uint32 NumLines = 1;
 		for( uint32 CharIndex = 0; CharIndex < NumChars; ++CharIndex )
 		{
@@ -1334,7 +1338,38 @@ void FSlateElementBatcher::AddTextElement(const FSlateDrawElement& DrawElement)
 					// Reserve memory for the glyphs.  This isn't perfect as the text could contain spaces and we might not render the rest of the text in this batch but its better than resizing constantly
 					const int32 GlyphsLeft = NumChars - CharIndex;
 					RenderBatch->ReserveVertices(GlyphsLeft * 4);
-					RenderBatch->ReserveIndices(GlyphsLeft * 6);
+
+					if (bUseStaticIndicies)
+					{
+						static FSlateIndexArray DefaultTextIndiciesArray = [MaxPreAllocatedGlyphIndicies = MaxPreAllocatedGlyphIndicies]()
+						{
+							FSlateIndexArray DefaultTextIndicies = {};
+							DefaultTextIndicies.Reserve(MaxPreAllocatedGlyphIndicies * 6);
+
+							for (int32 Glyph = 0; Glyph < MaxPreAllocatedGlyphIndicies; ++Glyph)
+							{
+								const uint32 IndexStart = Glyph * 4;
+
+								DefaultTextIndicies.Add(IndexStart + 0);
+								DefaultTextIndicies.Add(IndexStart + 1);
+								DefaultTextIndicies.Add(IndexStart + 2);
+								DefaultTextIndicies.Add(IndexStart + 1);
+								DefaultTextIndicies.Add(IndexStart + 3);
+								DefaultTextIndicies.Add(IndexStart + 2);
+							}
+
+							return DefaultTextIndicies;
+						}();
+
+						// NumIndicies can be dynamic based on newline chars, so increment it later as we iterate glyphs
+						RenderBatch->SourceIndices = &DefaultTextIndiciesArray;
+						RenderBatch->NumIndices = 0;
+						RenderBatch->IndexOffset = 0;
+					}
+					else
+					{
+						RenderBatch->ReserveIndices(GlyphsLeft * 6);
+					}
 
 					InvTextureSizeX = 1.0f / FontAtlasTexture->GetWidth();
 					InvTextureSizeY = 1.0f / FontAtlasTexture->GetHeight();
@@ -1398,12 +1433,19 @@ void FSlateElementBatcher::AddTextElement(const FSlateDrawElement& DrawElement)
 						RenderBatch->AddVertex(FSlateVertex::Make<Rounding>( RenderTransform, FVector2f(UpperLeft.X,LowerRight.Y),	FVector4f(U, V+SizeV, Ut,VtMax),			FVector2f(0.0f,1.0f), FontTint ));
 						RenderBatch->AddVertex(FSlateVertex::Make<Rounding>( RenderTransform, FVector2f(LowerRight),				FVector4f(U+SizeU, V+SizeV, UtMax,VtMax),	FVector2f(1.0f,1.0f), FontTint ));
 
-						RenderBatch->AddIndex(IndexStart + 0);
-						RenderBatch->AddIndex(IndexStart + 1);
-						RenderBatch->AddIndex(IndexStart + 2);
-						RenderBatch->AddIndex(IndexStart + 1);
-						RenderBatch->AddIndex(IndexStart + 3);
-						RenderBatch->AddIndex(IndexStart + 2);
+						if (bUseStaticIndicies)
+						{
+							RenderBatch->NumIndices += 6;
+						}
+						else
+						{
+							RenderBatch->AddIndex(IndexStart + 0);
+							RenderBatch->AddIndex(IndexStart + 1);
+							RenderBatch->AddIndex(IndexStart + 2);
+							RenderBatch->AddIndex(IndexStart + 1);
+							RenderBatch->AddIndex(IndexStart + 3);
+							RenderBatch->AddIndex(IndexStart + 2);
+						}
 					}
 				}
 
@@ -2924,6 +2966,10 @@ void FSlateElementBatcher::BuildShapedTextSequence(const FShapedTextBuildContext
 	// For right to left this value is unused. We just skip all leading whitespace
 	float PreviousWhitespaceAdvance = 0;
 
+	// The number of glyphs can technically be unbound, so pick some number that if smaller we use the preallocated index array
+	const int32 MaxPreAllocatedGlyphIndicies = 4096 + 1024;
+	const bool bUseStaticIndicies = GlyphSequenceToRender->GetGlyphsToRender().Num() < MaxPreAllocatedGlyphIndicies;
+
 	int32 NumGlyphs = GlyphSequenceToRender->GetGlyphsToRender().Num();
 	const TArray<FShapedGlyphEntry>& GlyphsToRender = GlyphSequenceToRender->GetGlyphsToRender();
 	for (int32 GlyphIndex = 0; GlyphIndex < NumGlyphs; ++GlyphIndex)
@@ -2988,7 +3034,38 @@ void FSlateElementBatcher::BuildShapedTextSequence(const FShapedTextBuildContext
 					// Reserve memory for the glyphs.  This isn't perfect as the text could contain spaces and we might not render the rest of the text in this batch but its better than resizing constantly
 					const int32 GlyphsLeft = NumGlyphs - GlyphIndex;
 					RenderBatch->ReserveVertices(GlyphsLeft * 4);
-					RenderBatch->ReserveIndices(GlyphsLeft * 6);
+
+					if (bUseStaticIndicies)
+					{
+						static FSlateIndexArray DefaultShapedTextIndiciesArray = [MaxPreAllocatedGlyphIndicies = MaxPreAllocatedGlyphIndicies]()
+						{
+							FSlateIndexArray DefaultShapedTextIndicies = {};
+							DefaultShapedTextIndicies.Reserve(MaxPreAllocatedGlyphIndicies * 6);
+
+							for (int32 Glyph = 0; Glyph < MaxPreAllocatedGlyphIndicies; ++Glyph)
+							{
+								const uint32 IndexStart = Glyph * 4;
+
+								DefaultShapedTextIndicies.Add(IndexStart + 0);
+								DefaultShapedTextIndicies.Add(IndexStart + 1);
+								DefaultShapedTextIndicies.Add(IndexStart + 2);
+								DefaultShapedTextIndicies.Add(IndexStart + 1);
+								DefaultShapedTextIndicies.Add(IndexStart + 3);
+								DefaultShapedTextIndicies.Add(IndexStart + 2);
+							}
+
+							return DefaultShapedTextIndicies;
+						}();
+
+						// NumIndicies can be dynamic based on newline & overflow / ellipsis, so increment it later as we iterate glyphs
+						RenderBatch->SourceIndices = &DefaultShapedTextIndiciesArray;
+						RenderBatch->NumIndices = 0; 
+						RenderBatch->IndexOffset = 0;
+					}
+					else
+					{
+						RenderBatch->ReserveIndices(GlyphsLeft * 6);
+					}
 
 					InvTextureSizeX = 1.0f / FontAtlasTexture->GetWidth();
 					InvTextureSizeY = 1.0f / FontAtlasTexture->GetHeight();
@@ -3122,12 +3199,19 @@ void FSlateElementBatcher::BuildShapedTextSequence(const FShapedTextBuildContext
 			RenderBatch->AddVertex(FSlateVertex::Make<Rounding>(RenderTransform, FVector2f(UpperLeft.X, LowerRight.Y), FVector4f(U, V + SizeV, Ut, VtMax), FVector2f(0.0f, 1.0f), Tint));
 			RenderBatch->AddVertex(FSlateVertex::Make<Rounding>(RenderTransform, FVector2f(LowerRight), FVector4f(U + SizeU, V + SizeV, UtMax, VtMax), FVector2f(1.0f, 1.0f), Tint));
 
-			RenderBatch->AddIndex(IndexStart + 0);
-			RenderBatch->AddIndex(IndexStart + 1);
-			RenderBatch->AddIndex(IndexStart + 2);
-			RenderBatch->AddIndex(IndexStart + 1);
-			RenderBatch->AddIndex(IndexStart + 3);
-			RenderBatch->AddIndex(IndexStart + 2);
+			if (bUseStaticIndicies)
+			{
+				RenderBatch->NumIndices += 6;
+			}
+			else
+			{
+				RenderBatch->AddIndex(IndexStart + 0);
+				RenderBatch->AddIndex(IndexStart + 1);
+				RenderBatch->AddIndex(IndexStart + 2);
+				RenderBatch->AddIndex(IndexStart + 1);
+				RenderBatch->AddIndex(IndexStart + 3);
+				RenderBatch->AddIndex(IndexStart + 2);
+			}
 
 			// Reset whitespace advance to 0, this is a visible character
 			PreviousWhitespaceAdvance = 0;
