Make retained widgets group by element type (Otherwise future element type optimizations won't affect retained rendering).

diff --git a/Engine/Source/Runtime/SlateCore/Private/Rendering/DrawElements.cpp b/Engine/Source/Runtime/SlateCore/Private/Rendering/DrawElements.cpp
--- a/Engine/Source/Runtime/SlateCore/Private/Rendering/DrawElements.cpp
+++ b/Engine/Source/Runtime/SlateCore/Private/Rendering/DrawElements.cpp
@@ -814,7 +814,7 @@ FSlateDrawElement& FSlateWindowElementList::AddUninitialized(EElementType InElem
 	if (bAllowCache)
 	{
 		// @todo get working with slate debugging
-		return AddCachedElement();
+		return AddCachedElement(InElementType);
 	}
 	else
 	{
@@ -977,7 +977,7 @@ void FSlateWindowElementList::PopBatchPriortyGroup()
 	//BatchDepthPriorityStack.Pop();
 }*/
 
-FSlateDrawElement& FSlateWindowElementList::AddCachedElement()
+FSlateDrawElement& FSlateWindowElementList::AddCachedElement(EElementType InElementType)
 {
 	FSlateCachedElementData* CurrentCachedElementData = GetCurrentCachedElementData();
 	check(CurrentCachedElementData);
@@ -990,7 +990,7 @@ FSlateDrawElement& FSlateWindowElementList::AddCachedElement()
 		CurrentWidgetState.CacheHandle = CurrentCachedElementData->AddCache(CurrentWidgetState.Widget);
 	}
 
-	return CurrentCachedElementData->AddCachedElement(CurrentWidgetState.CacheHandle, GetClippingManager(), CurrentWidgetState.Widget);
+	return CurrentCachedElementData->AddCachedElement(CurrentWidgetState.CacheHandle, GetClippingManager(), CurrentWidgetState.Widget, InElementType);
 }
 
 void FSlateWindowElementList::PushCachedElementData(FSlateCachedElementData& CachedElementData)
@@ -1155,7 +1155,7 @@ FSlateCachedElementsHandle FSlateCachedElementData::AddCache(const SWidget* Widg
 	return FSlateCachedElementsHandle(NewList);
 }
 
-FSlateDrawElement& FSlateCachedElementData::AddCachedElement(FSlateCachedElementsHandle& CacheHandle, const FSlateClippingManager& ParentClipManager, const SWidget* CurrentWidget)
+FSlateDrawElement& FSlateCachedElementData::AddCachedElement(FSlateCachedElementsHandle& CacheHandle, const FSlateClippingManager& ParentClipManager, const SWidget* CurrentWidget, EElementType InElementType)
 {
 	TSharedPtr<FSlateCachedElementList> List = CacheHandle.Ptr.Pin();
 
@@ -1164,8 +1164,7 @@ FSlateDrawElement& FSlateCachedElementData::AddCachedElement(FSlateCachedElement
 	check(CurrentWidget->GetParentWidget().IsValid());
 #endif
 	
-	// Add cached widget to NonMapped since we don't have an element type dedicated to it
-	FSlateDrawElementContainer& Container = List->DrawElements.FindOrAdd(EElementType::ET_NonMapped);
+	FSlateDrawElementContainer& Container = List->DrawElements.FindOrAdd(InElementType);
 	FSlateDrawElement& NewElement = Container.Elements.AddDefaulted_GetRef();
 	NewElement.SetIsCached(true);
 
diff --git a/Engine/Source/Runtime/SlateCore/Public/Rendering/DrawElements.h b/Engine/Source/Runtime/SlateCore/Public/Rendering/DrawElements.h
--- a/Engine/Source/Runtime/SlateCore/Public/Rendering/DrawElements.h
+++ b/Engine/Source/Runtime/SlateCore/Public/Rendering/DrawElements.h
@@ -453,7 +453,7 @@ struct FSlateCachedElementData
 
 	FSlateCachedElementsHandle AddCache(const SWidget* Widget);
 
-	FSlateDrawElement& AddCachedElement(FSlateCachedElementsHandle& CacheHandle, const FSlateClippingManager& ParentClipManager, const SWidget* Widget);
+	FSlateDrawElement& AddCachedElement(FSlateCachedElementsHandle& CacheHandle, const FSlateClippingManager& ParentClipManager, const SWidget* Widget, EElementType InElementType);
 
 	FSlateRenderBatch& AddCachedRenderBatch(FSlateRenderBatch&& NewBatch, int32& OutIndex);
 	void RemoveCachedRenderBatches(const TArray<int32>& CachedRenderBatchIndices);
@@ -677,7 +677,7 @@ public:
 		}
 	}
 private:
-	FSlateDrawElement& AddCachedElement();
+	FSlateDrawElement& AddCachedElement(EElementType InElementType);
 
 	TArrayView<FSlateCachedElementData* const> GetCachedElementDataList() const { return MakeArrayView(CachedElementDataList.GetData(), CachedElementDataList.Num()); }
 
