Fixed various ellipsis bugs:
- ellipsis appearing on multiple lines on multline text, and for each block in each line for multiline + rich text or bidirectionnal text, instead of having a single ellipsis by text.
- ellipsis not always shown on text with multiple runs (like rich text or bidirectionnal text). Ellipsis could be misplaced. It was more likely to happen on left_to_right text.
- ellipsis not shown (or show the wrong side) for multiline right justified text, or right to left text that was left justified.

When using texts with multiples blocks, it's possible to have the ellipsis partially clipped, but it would require deeper changes in the text system to fix this issue.

Tests:
-New test scene in EngineTest to cover all fixed bugs, and much more ellipsis usage.
-Validated in Unreal Ed that all tested editors where working as expected.

diff --git a/Engine/Source/Runtime/Slate/Private/Framework/Text/SlateTextLayout.cpp b/Engine/Source/Runtime/Slate/Private/Framework/Text/SlateTextLayout.cpp
--- a/Engine/Source/Runtime/Slate/Private/Framework/Text/SlateTextLayout.cpp
+++ b/Engine/Source/Runtime/Slate/Private/Framework/Text/SlateTextLayout.cpp
@@ -69,39 +69,36 @@ int32 FSlateTextLayout::OnPaint( const FPaintArgs& Args, const FGeometry& Allott
 
 	int32 HighestLayerId = LayerId;
 
-	auto IsLineViewVisible = [&AllottedGeometry, &InverseScale, &MyCullingRect](const FTextLayout::FLineView& LineView, bool bTestForVerticallyClippedLine = false)
+	auto IsVisible = [&AllottedGeometry, &InverseScale, &MyCullingRect](FVector2D Offset, FVector2D Size)
 	{
 		// Is this line visible?  This checks if the culling rect, which represents the AABB around the last clipping rect, intersects the 
 		// line of text, this requires that we get the text line into render space.
 		// TODO perhaps save off this line view rect during text layout?
-		const FVector2D LocalLineOffset = LineView.Offset * InverseScale;
-		FSlateRect LineViewRect(AllottedGeometry.GetRenderBoundingRect(FSlateRect(LocalLineOffset, LocalLineOffset + (LineView.Size * InverseScale))));
+		const FVector2D LocalLineOffset = Offset * InverseScale;
+		FSlateRect Rect(AllottedGeometry.GetRenderBoundingRect(FSlateRect(LocalLineOffset, LocalLineOffset + (Size * InverseScale))));
 
-		if (!bTestForVerticallyClippedLine)
-		{
-			return FSlateRect::DoRectanglesIntersect(LineViewRect, MyCullingRect);
-		}
-		else
-		{
-			// We only care about vertically clipped lines. Horizontal clipping will be replaced with an ellipsis
-			// So make the left and right side the same as the clip rect to avoid the test failing due to being larger horizontally
+		return FSlateRect::DoRectanglesIntersect(Rect, MyCullingRect);
+	};
 
-			FSlateRect GeometryRect = AllottedGeometry.GetRenderBoundingRect();
-			LineViewRect.Left = GeometryRect.Left;
-			LineViewRect.Right = GeometryRect.Right;
 
-			return FSlateRect::IsRectangleContained(GeometryRect, LineViewRect);
+	int32 LastVisibleLineIndex = -1;
+	for (int32 LineIndex = 0; LineIndex < LineViews.Num(); ++LineIndex)
+	{	//Depending on their length, some lines can be invisible, but following ones could be visible, so we have to parse all of them to get the real last one.
+		if (IsVisible(LineViews[LineIndex].Offset, LineViews[LineIndex].Size))
+		{
+			LastVisibleLineIndex = LineIndex;
 		}
+	}
 
-	};
-
+	const bool bIsMultiline = LineViews.Num() > 1;
 	for (int32 LineIndex = 0; LineIndex < LineViews.Num(); ++LineIndex)
 	{
 		const FTextLayout::FLineView& LineView = LineViews[LineIndex];
 
 		ETextOverflowPolicy OverflowPolicy = TextOverflowPolicyOverride.Get(DefaultTextStyle.OverflowPolicy);
+		ETextOverflowDirection OverflowDirection = ETextOverflowDirection::NoOverflow;
 
-		if (!IsLineViewVisible(LineView))
+		if (!IsVisible(LineView.Offset, LineView.Size))
 		{
 			continue;
 		}
@@ -113,22 +110,70 @@ int32 FSlateTextLayout::OnPaint( const FPaintArgs& Args, const FGeometry& Allott
 		const int32 TextLayer = BlockDebugLayer + 1;
 		int32 HighestBlockLayerId = TextLayer;
 
-		const ETextJustify::Type VisualJustification = CalculateLineViewVisualJustification(LineView);
-		ETextOverflowDirection OverflowDirection = VisualJustification == ETextJustify::Left ? ETextOverflowDirection::LeftToRight : (VisualJustification == ETextJustify::Right ? ETextOverflowDirection::RightToLeft : ETextOverflowDirection::NoOverflow);
+		bool bIsLastVisibleLine = false;
+		bool bIsNextLineClipped = false;
+		if (OverflowPolicy == ETextOverflowPolicy::Ellipsis)
+		{
+			if (bIsMultiline)
+			{
+				bIsLastVisibleLine = LastVisibleLineIndex == LineIndex;
+				bIsNextLineClipped = LineViews.IsValidIndex(LineIndex + 1) ? bIsLastVisibleLine : false;
+
+				//When wrapping/multiline text, we have to use the reading direction of the text, not the justification.
+				OverflowDirection = LineView.TextBaseDirection == TextBiDi::ETextDirection::LeftToRight ? ETextOverflowDirection::LeftToRight : ETextOverflowDirection::RightToLeft;
+			}
+			else
+			{
+				bIsLastVisibleLine = true;
+				const ETextJustify::Type VisualJustification = CalculateLineViewVisualJustification(LineView);
+				OverflowDirection = VisualJustification == ETextJustify::Left ? ETextOverflowDirection::LeftToRight : (VisualJustification == ETextJustify::Right ? ETextOverflowDirection::RightToLeft : ETextOverflowDirection::NoOverflow);
+			}
+		}
+
+		int32 LastVisibleBlockIndex = -1;
+		int32 NextBlockOffset = 0;
+		if (bIsLastVisibleLine)
+		{
+			int32 StartValue = 0;
+			int32 EndValue = 0;
+			if (OverflowDirection == ETextOverflowDirection::LeftToRight)
+			{
+				EndValue = LineView.Blocks.Num();
+				NextBlockOffset = 1;
+			}
+			else if (OverflowDirection == ETextOverflowDirection::RightToLeft)
+			{
+				StartValue = LineView.Blocks.Num() - 1;
+				EndValue = -1;
+				NextBlockOffset = -1;
+			}
 
-		bool bForceEllipsisDueToClippedLine = false;
-	
-		if (OverflowPolicy == ETextOverflowPolicy::Ellipsis && LineViews.Num() > 1)
+			for (int32 BlockIndex = StartValue; BlockIndex != EndValue; BlockIndex += NextBlockOffset)
+			{	//Depending on their position, some blocks can be invisible, but a following one could be visible, so we have to parse all of them.
+				const TSharedRef< ILayoutBlock >& Block = LineView.Blocks[BlockIndex];
+				if (IsVisible(Block->GetLocationOffset(), Block->GetSize()))
+				{
+					LastVisibleBlockIndex = BlockIndex;
+				}
+			}
+		}
+		else
 		{
-			// Force the ellipsis to be on when the next line in a multi line text layout is clipped. This forces an ellipsis on this line even when its not clipped to indicate that an entire line or more is clipped
-			bForceEllipsisDueToClippedLine = LineViews.IsValidIndex(LineIndex + 1) ? !IsLineViewVisible(LineViews[LineIndex + 1], true) : false;
+			OverflowDirection = ETextOverflowDirection::NoOverflow;
+			OverflowPolicy = ETextOverflowPolicy::Clip;
 		}
 
 		// Render every block for this line
-		for (const TSharedRef< ILayoutBlock >& Block : LineView.Blocks)
+		for (int32 BlockIndex = 0; BlockIndex < LineView.Blocks.Num(); ++BlockIndex)
 		{
+			const TSharedRef< ILayoutBlock >& Block = LineView.Blocks[BlockIndex];
+			if (!IsVisible(Block->GetLocationOffset(), Block->GetSize()))
+			{
+				continue;
+			}
+
 #if !(UE_BUILD_SHIPPING || UE_BUILD_TEST)
-			if ( GShowTextDebugging )
+			if ( GShowTextDebugging)
 			{
 				BlockDebugHue.R += 50.0f;
 
@@ -144,7 +189,8 @@ int32 FSlateTextLayout::OnPaint( const FPaintArgs& Args, const FGeometry& Allott
 #endif
 
 			FTextArgs TextArgs(LineView, Block, DefaultTextStyle, OverflowPolicy, OverflowDirection);
-			TextArgs.bForceEllipsisDueToClippedLine = bForceEllipsisDueToClippedLine;
+			TextArgs.bIsLastVisibleBlock = bIsLastVisibleLine && BlockIndex == LastVisibleBlockIndex;
+			TextArgs.bIsNextBlockClipped = bIsNextLineClipped && (LineView.Blocks.IsValidIndex(BlockIndex + NextBlockOffset) ? TextArgs.bIsLastVisibleBlock : true);
 
 			const TSharedRef< ISlateRun > Run = StaticCastSharedRef< ISlateRun >( Block->GetRun() );
 
diff --git a/Engine/Source/Runtime/Slate/Private/Framework/Text/SlateTextRun.cpp b/Engine/Source/Runtime/Slate/Private/Framework/Text/SlateTextRun.cpp
--- a/Engine/Source/Runtime/Slate/Private/Framework/Text/SlateTextRun.cpp
+++ b/Engine/Source/Runtime/Slate/Private/Framework/Text/SlateTextRun.cpp
@@ -123,7 +123,8 @@ int32 FSlateTextRun::OnPaint(const FPaintArgs& PaintArgs, const FTextArgs& TextA
 	{
 		OverflowArgs.OverflowTextPtr = BlockTextContext.ShapedTextCache->FindOrAddOverflowEllipsisText(AllottedGeometry.GetAccumulatedLayoutTransform().GetScale(), BlockTextContext, Style.Font);
 		OverflowArgs.OverflowDirection = TextArgs.OverflowDirection;
-		OverflowArgs.bForceEllipsisDueToClippedLine = TextArgs.bForceEllipsisDueToClippedLine;
+		OverflowArgs.bIsLastVisibleBlock = TextArgs.bIsLastVisibleBlock;
+		OverflowArgs.bIsNextBlockClipped = TextArgs.bIsNextBlockClipped;
 	}
 
 	// Draw the optional shadow
diff --git a/Engine/Source/Runtime/Slate/Public/Framework/Text/ISlateRun.h b/Engine/Source/Runtime/Slate/Public/Framework/Text/ISlateRun.h
--- a/Engine/Source/Runtime/Slate/Public/Framework/Text/ISlateRun.h
+++ b/Engine/Source/Runtime/Slate/Public/Framework/Text/ISlateRun.h
@@ -20,7 +20,8 @@ struct FTextArgs
 		, DefaultStyle(InDefaultStyle)
 		, OverflowPolicy(InOverflowPolicy)
 		, OverflowDirection(InOverflowDirection)
-		, bForceEllipsisDueToClippedLine(false)
+		, bIsLastVisibleBlock(false)
+		, bIsNextBlockClipped(false)
 	{}
 
 	const FTextLayout::FLineView& Line;
@@ -28,7 +29,8 @@ struct FTextArgs
 	const FTextBlockStyle& DefaultStyle;
 	ETextOverflowPolicy OverflowPolicy;
 	ETextOverflowDirection OverflowDirection;
-	bool bForceEllipsisDueToClippedLine;
+	bool bIsLastVisibleBlock;
+	bool bIsNextBlockClipped;
 
 };
 
diff --git a/Engine/Source/Runtime/SlateCore/Private/Rendering/ElementBatcher.cpp b/Engine/Source/Runtime/SlateCore/Private/Rendering/ElementBatcher.cpp
--- a/Engine/Source/Runtime/SlateCore/Private/Rendering/ElementBatcher.cpp
+++ b/Engine/Source/Runtime/SlateCore/Private/Rendering/ElementBatcher.cpp
@@ -1368,7 +1368,7 @@ void FSlateElementBatcher::AddTextElement(const FSlateDrawElement& DrawElement)
 					}
 					else
 					{
-						RenderBatch->ReserveIndices(GlyphsLeft * 6);
+					RenderBatch->ReserveIndices(GlyphsLeft * 6);
 					}
 
 					InvTextureSizeX = 1.0f / FontAtlasTexture->GetWidth();
@@ -1439,15 +1439,15 @@ void FSlateElementBatcher::AddTextElement(const FSlateDrawElement& DrawElement)
 						}
 						else
 						{
-							RenderBatch->AddIndex(IndexStart + 0);
-							RenderBatch->AddIndex(IndexStart + 1);
-							RenderBatch->AddIndex(IndexStart + 2);
-							RenderBatch->AddIndex(IndexStart + 1);
-							RenderBatch->AddIndex(IndexStart + 3);
-							RenderBatch->AddIndex(IndexStart + 2);
-						}
+						RenderBatch->AddIndex(IndexStart + 0);
+						RenderBatch->AddIndex(IndexStart + 1);
+						RenderBatch->AddIndex(IndexStart + 2);
+						RenderBatch->AddIndex(IndexStart + 1);
+						RenderBatch->AddIndex(IndexStart + 3);
+						RenderBatch->AddIndex(IndexStart + 2);
 					}
 				}
+				}
 
 				LineX += Entry.XAdvance;
 			}
@@ -1581,7 +1581,7 @@ void FSlateElementBatcher::AddShapedTextElement( const FSlateDrawElement& DrawEl
 		// Optimize by culling
 		// Todo: this doesn't work with cached clipping
 		BuildContext.bEnableCulling = false;
-		BuildContext.bForceEllipsis = DrawElementPayload.OverflowArgs.bForceEllipsisDueToClippedLine;
+		BuildContext.bForceEllipsis = DrawElementPayload.OverflowArgs.bIsLastVisibleBlock && DrawElementPayload.OverflowArgs.bIsNextBlockClipped;
 		BuildContext.OverflowDirection = DrawElementPayload.OverflowArgs.OverflowDirection;
 
 		if (ShapedGlyphSequence->GetGlyphsToRender().Num() > 200 || (OverflowGlyphSequence && BuildContext.OverflowDirection != ETextOverflowDirection::NoOverflow))
@@ -1597,12 +1597,13 @@ void FSlateElementBatcher::AddShapedTextElement( const FSlateDrawElement& DrawEl
 				// checked that the render transform is axis-aligned as well
 				const FSlateRect LocalClipBoundingBox = TransformRect(RenderTransform.Inverse(), ScissorRectBox);
 				BuildContext.LocalClipBoundingBoxLeft = LocalClipBoundingBox.Left;
-				BuildContext.LocalClipBoundingBoxRight = LocalClipBoundingBox.Right - (BuildContext.bForceEllipsis ? OverflowGlyphSequence->GetMeasuredWidth() : 0);
+				BuildContext.LocalClipBoundingBoxRight = LocalClipBoundingBox.Right;
 
 				// In checks below, ignore floating-point differences caused by transforming and untransforming the clip rect
-				if (OverflowGlyphSequence && FMath::FloorToInt(BuildContext.LocalClipBoundingBoxLeft) <= 0 && FMath::CeilToInt(BuildContext.LocalClipBoundingBoxRight) >= ShapedGlyphSequence->GetMeasuredWidth())
+				const bool NeedLeftEllipsis = FMath::FloorToInt(BuildContext.LocalClipBoundingBoxLeft) > 0 && BuildContext.OverflowDirection == ETextOverflowDirection::RightToLeft;
+				const bool NeedRightEllipsis = ShapedGlyphSequence->GetMeasuredWidth() > FMath::CeilToInt(BuildContext.LocalClipBoundingBoxRight) && BuildContext.OverflowDirection == ETextOverflowDirection::LeftToRight;
+				if (!NeedLeftEllipsis && !NeedRightEllipsis && !DrawElementPayload.OverflowArgs.bIsNextBlockClipped)
 				{
-					// Override overflow if the text is smaller than (or is the same size as) the clipping rect and wont be clipped
 					BuildContext.OverflowDirection = ETextOverflowDirection::NoOverflow;
 				}
 				else if(!OverflowGlyphSequence)
@@ -2960,7 +2961,7 @@ void FSlateElementBatcher::BuildShapedTextSequence(const FShapedTextBuildContext
 	float EllipsisLineX = 0;
 	float EllipsisLineY = 0;
 	bool bNeedEllipsis = false;
-	bool bChararcterWasClipped = false;
+	bool bNeedSpaceForEllipsis = false;
 
 	// For left to right overflow direction - Sum of total whitespace we're currently advancing through. Once a non-whitespace glyph is detected this will return to 0
 	// For right to left this value is unused. We just skip all leading whitespace
@@ -2970,7 +2971,7 @@ void FSlateElementBatcher::BuildShapedTextSequence(const FShapedTextBuildContext
 	const int32 MaxPreAllocatedGlyphIndicies = 4096 + 1024;
 	const bool bUseStaticIndicies = GlyphSequenceToRender->GetGlyphsToRender().Num() < MaxPreAllocatedGlyphIndicies;
 
-	int32 NumGlyphs = GlyphSequenceToRender->GetGlyphsToRender().Num();
+	const int32 NumGlyphs = GlyphSequenceToRender->GetGlyphsToRender().Num();
 	const TArray<FShapedGlyphEntry>& GlyphsToRender = GlyphSequenceToRender->GetGlyphsToRender();
 	for (int32 GlyphIndex = 0; GlyphIndex < NumGlyphs; ++GlyphIndex)
 	{
@@ -3064,7 +3065,7 @@ void FSlateElementBatcher::BuildShapedTextSequence(const FShapedTextBuildContext
 					}
 					else
 					{
-						RenderBatch->ReserveIndices(GlyphsLeft * 6);
+					RenderBatch->ReserveIndices(GlyphsLeft * 6);
 					}
 
 					InvTextureSizeX = 1.0f / FontAtlasTexture->GetWidth();
@@ -3093,7 +3094,7 @@ void FSlateElementBatcher::BuildShapedTextSequence(const FShapedTextBuildContext
 		}
 
 		// Overflow Detection
-		// First figure out the size of the glyph. If the glyph contains multiple characters we have to measure all of them and if clicked, omit them all. This is common in complex languages with lots of diacritics
+		// First figure out the size of the glyph. If the glyph contains multiple characters we have to measure all of them and if clipped, omit them all. This is common in complex languages with lots of diacritics
 		float OverflowTestWidth = SizeX;
 		if (OverflowDirection != ETextOverflowDirection::NoOverflow && (GlyphToRender.NumGraphemeClustersInGlyph > 1 || GlyphToRender.NumCharactersInGlyph > 1))
 		{
@@ -3115,7 +3116,7 @@ void FSlateElementBatcher::BuildShapedTextSequence(const FShapedTextBuildContext
 		// Left to right overflow - If the current pen position + the ellipsis cannot fit, we have reached the end of the possible area for drawing this text. 
 		if (OverflowDirection == ETextOverflowDirection::LeftToRight)
 		{
-			// If we are on the last glyph dont bother checking if the ellipsis can fit. If the last glyph can fit there is no need for ellipsis
+			// If we are on the last glyph don't bother checking if the ellipsis can fit. If the last glyph can fit there is no need for ellipsis
 			float OverflowSequenceNeededSize = GlyphIndex < NumGlyphs - 1 ? Context.OverflowGlyphSequence->GetMeasuredWidth() : 0;
 			if(X + OverflowTestWidth + OverflowSequenceNeededSize >= Context.LocalClipBoundingBoxRight)
 			{
@@ -3136,33 +3137,36 @@ void FSlateElementBatcher::BuildShapedTextSequence(const FShapedTextBuildContext
 			{
 				// This glyph is in the clipped region or is not visible so just advance. It cannot be shown
 				bClipped = true;
-				bChararcterWasClipped = true;
+				bNeedSpaceForEllipsis = true;
 			}
-			else if(bChararcterWasClipped)
+			if (bNeedSpaceForEllipsis || !GlyphToRender.bIsVisible)
 			{
+				bClipped = true;
+
 				// Can the ellipsis fit in the free spot by skipping the previous glyph(s)
-				float EllipsisWidth = Context.OverflowGlyphSequence->GetMeasuredWidth();
-				float AvailableX = X-Context.LocalClipBoundingBoxLeft;
+				const float EllipsisWidth = Context.OverflowGlyphSequence->GetMeasuredWidth();
+				const float AvailableX = X + SizeX - Context.LocalClipBoundingBoxLeft;
 				if (AvailableX >= EllipsisWidth)
 				{
 					// The available area can fit the ellipsis. Mark that we need an ellipsis and stop checking for overflow. The rest of the text can be built normally
+					bNeedSpaceForEllipsis = false;
+				}
+
+				//Always try to put the ellipsis, wether it fits or not: it's better to have an ellipsis a bit clipped than no feedback at all.
 					bNeedEllipsis = true;
-					EllipsisLineX = (LineX - EllipsisWidth);
+				EllipsisLineX = (LineX + SizeX - EllipsisWidth);
 					EllipsisLineY = LineY;
-					OverflowDirection = ETextOverflowDirection::NoOverflow;
 				}
 				else
 				{
-					bClipped = true;
-					bChararcterWasClipped = true;
+				OverflowDirection = ETextOverflowDirection::NoOverflow;
 				}
 
-			}
 			// If we just clipped a glyph omit all characters in said glyph. Otherwise floating diacritics would be visible above the ellipsis. This is common in complex languages.
 			if (bClipped && GlyphToRender.NumCharactersInGlyph > 1)
 			{
 				GlyphIndex += GlyphToRender.NumCharactersInGlyph - 1;
-				LineX += OverflowTestWidth;
+				LineX += GlyphToRender.XAdvance;
 				continue;
 			}
 
@@ -3205,12 +3209,12 @@ void FSlateElementBatcher::BuildShapedTextSequence(const FShapedTextBuildContext
 			}
 			else
 			{
-				RenderBatch->AddIndex(IndexStart + 0);
-				RenderBatch->AddIndex(IndexStart + 1);
-				RenderBatch->AddIndex(IndexStart + 2);
-				RenderBatch->AddIndex(IndexStart + 1);
-				RenderBatch->AddIndex(IndexStart + 3);
-				RenderBatch->AddIndex(IndexStart + 2);
+			RenderBatch->AddIndex(IndexStart + 0);
+			RenderBatch->AddIndex(IndexStart + 1);
+			RenderBatch->AddIndex(IndexStart + 2);
+			RenderBatch->AddIndex(IndexStart + 1);
+			RenderBatch->AddIndex(IndexStart + 3);
+			RenderBatch->AddIndex(IndexStart + 2);
 			}
 
 			// Reset whitespace advance to 0, this is a visible character
diff --git a/Engine/Source/Runtime/SlateCore/Public/Rendering/DrawElementTextOverflowArgs.h b/Engine/Source/Runtime/SlateCore/Public/Rendering/DrawElementTextOverflowArgs.h
--- a/Engine/Source/Runtime/SlateCore/Public/Rendering/DrawElementTextOverflowArgs.h
+++ b/Engine/Source/Runtime/SlateCore/Public/Rendering/DrawElementTextOverflowArgs.h
@@ -19,18 +19,20 @@ struct FTextOverflowArgs
 	FTextOverflowArgs(FShapedGlyphSequencePtr& InOverflowText, ETextOverflowDirection InOverflowDirection)
 		: OverflowTextPtr(InOverflowText)
 		, OverflowDirection(InOverflowDirection)
-		, bForceEllipsisDueToClippedLine(false)
+		, bIsLastVisibleBlock(false)
+		, bIsNextBlockClipped(false)
 		
 	{}
 
 	FTextOverflowArgs()
 		: OverflowDirection(ETextOverflowDirection::NoOverflow)
-		, bForceEllipsisDueToClippedLine(false)
+		, bIsLastVisibleBlock(false)
+		, bIsNextBlockClipped(false)
 	{}
 
 	/** Sequence that represents the ellipsis glyph */
 	FShapedGlyphSequencePtr OverflowTextPtr;
 	ETextOverflowDirection OverflowDirection;
-	bool bForceEllipsisDueToClippedLine;
-	
+	bool bIsLastVisibleBlock;
+	bool bIsNextBlockClipped;
 };
