Add ability to add a widget inside the global window overlay

diff --git a/Engine/Source/Runtime/Engine/Classes/Engine/GameViewportClient.h b/Engine/Source/Runtime/Engine/Classes/Engine/GameViewportClient.h
--- a/Engine/Source/Runtime/Engine/Classes/Engine/GameViewportClient.h
+++ b/Engine/Source/Runtime/Engine/Classes/Engine/GameViewportClient.h
@@ -292,6 +292,21 @@ public:
 	 */
 	void RemoveAllViewportWidgets();
 
+	/**
+	 * Adds a widget to the Slate viewport's game overlay (i.e for top-level debug tools) at the specified Z-order.
+	 *
+	 * @param  ViewportContent	The widget to add.  Must be valid.
+	 * @param  ZOrder  The Z-order index for this widget.  Larger values will cause the widget to appear on top of widgets with lower values.
+	 */
+	virtual void AddGameLayerWidget(TSharedRef<SWidget> ViewportContent, const int32 ZOrder);
+
+	/**
+	 * Removes a previously-added widget from the Slate viewport's game overlay.
+	 *
+	 * @param	ViewportContent  The widget to remove.  Must be valid.
+	 */
+	virtual void RemoveGameLayerWidget(TSharedRef<SWidget> ViewportContent);
+
 	/**
 	 * Recreates cursor widgets from UISettings class.
 	 */
diff --git a/Engine/Source/Runtime/Engine/Private/GameViewportClient.cpp b/Engine/Source/Runtime/Engine/Private/GameViewportClient.cpp
--- a/Engine/Source/Runtime/Engine/Private/GameViewportClient.cpp
+++ b/Engine/Source/Runtime/Engine/Private/GameViewportClient.cpp
@@ -2892,6 +2892,22 @@ void UGameViewportClient::RemoveAllViewportWidgets()
 	}
 }
 
+void UGameViewportClient::AddGameLayerWidget(TSharedRef<SWidget> ViewportContent, const int32 ZOrder)
+{
+	if (const TSharedPtr<IGameLayerManager> GameLayerManager = GameLayerManagerPtr.Pin())
+	{
+		GameLayerManager->AddGameLayer(ViewportContent, ZOrder);
+	}
+}
+
+void UGameViewportClient::RemoveGameLayerWidget(TSharedRef<SWidget> ViewportContent)
+{
+	if (const TSharedPtr<IGameLayerManager> GameLayerManager = GameLayerManagerPtr.Pin())
+	{
+		GameLayerManager->RemoveGameLayer(ViewportContent);
+	}
+}
+
 void UGameViewportClient::VerifyPathRenderingComponents()
 {
 	const bool bShowPaths = !!EngineShowFlags.Navigation;
diff --git a/Engine/Source/Runtime/Engine/Private/Slate/SGameLayerManager.cpp b/Engine/Source/Runtime/Engine/Private/Slate/SGameLayerManager.cpp
--- a/Engine/Source/Runtime/Engine/Private/Slate/SGameLayerManager.cpp
+++ b/Engine/Source/Runtime/Engine/Private/Slate/SGameLayerManager.cpp
@@ -119,14 +119,14 @@ void SGameLayerManager::Construct(const SGameLayerManager::FArguments& InArgs)
 
 			+ SVerticalBox::Slot()
 			[
-				SNew(SOverlay)
+				SAssignNew(WindowOverlay, SOverlay)
 
-				+ SOverlay::Slot()
+				+ SOverlay::Slot(static_cast<int32>(EGameLayerOrder::Player))
 				[
 					SAssignNew(PlayerCanvas, SCanvas)
 				]
 
-				+ SOverlay::Slot()
+				+ SOverlay::Slot(static_cast<int32>(EGameLayerOrder::Viewport))
 				[
 					SAssignNew(ViewportSlotContainer, SBox)
 					[
@@ -134,7 +134,7 @@ void SGameLayerManager::Construct(const SGameLayerManager::FArguments& InArgs)
 					]
 				]
 
-				+ SOverlay::Slot()
+				+ SOverlay::Slot(static_cast<int32>(EGameLayerOrder::TitleBar))
 				[
 					SNew(SVerticalBox)
 
@@ -148,14 +148,14 @@ void SGameLayerManager::Construct(const SGameLayerManager::FArguments& InArgs)
 					]
 				]
 
-				+ SOverlay::Slot()
+				+ SOverlay::Slot(static_cast<int32>(EGameLayerOrder::Tooltip))
 				[
 					SNew(SPopup)
 					[
 						SAssignNew(TooltipPresenter, STooltipPresenter)
 					]
 				]
-				+ SOverlay::Slot()
+				+ SOverlay::Slot(static_cast<int32>(EGameLayerOrder::Debug))
 				[
 					SAssignNew(DebugCanvas, SDebugCanvas)
 					.SceneViewport(InArgs._SceneViewport)
@@ -279,7 +279,9 @@ void SGameLayerManager::NotifyPlayerRemoved(int32 PlayerIndex, ULocalPlayer* Rem
 void SGameLayerManager::AddWidgetForPlayer(ULocalPlayer* Player, TSharedRef<SWidget> ViewportContent, const int32 ZOrder)
 {
 	TSharedPtr<FPlayerLayer> PlayerLayer = FindOrCreatePlayerLayer(Player);
-	
+
+	ensure(!PlayerLayer->Widget->HasSlotWithZOrder(ZOrder));
+
 	PlayerLayer->Widget->AddSlot(ZOrder)
 	[
 		ViewportContent
@@ -372,6 +374,21 @@ void SGameLayerManager::ClearWidgets()
 	SetWindowTitleBarState(nullptr, EWindowTitleBarMode::Overlay, false, false, false);
 }
 
+void SGameLayerManager::AddGameLayer(TSharedRef<SWidget> ViewportContent, int32 ZOrder)
+{
+	ensure(!WindowOverlay->HasSlotWithZOrder(ZOrder));
+
+	WindowOverlay->AddSlot(ZOrder)
+	[
+		ViewportContent
+	];
+}
+
+void SGameLayerManager::RemoveGameLayer(TSharedRef<SWidget> ViewportContent)
+{
+	WindowOverlay->RemoveSlot(ViewportContent);
+}
+
 void SGameLayerManager::Tick(const FGeometry& AllottedGeometry, const double InCurrentTime, const float InDeltaTime)
 {
 	//if (AllottedGeometry != CachedGeometry)
diff --git a/Engine/Source/Runtime/Engine/Public/Slate/SGameLayerManager.h b/Engine/Source/Runtime/Engine/Public/Slate/SGameLayerManager.h
--- a/Engine/Source/Runtime/Engine/Public/Slate/SGameLayerManager.h
+++ b/Engine/Source/Runtime/Engine/Public/Slate/SGameLayerManager.h
@@ -74,6 +74,9 @@ public:
 
 	virtual void ClearWidgets() = 0;
 
+	virtual void AddGameLayer(TSharedRef<SWidget> ViewportContent, int32 ZOrder) = 0;
+	virtual void RemoveGameLayer(TSharedRef<SWidget> ViewportContent) = 0;
+
 	virtual void SetDefaultWindowTitleBarHeight(float Height) = 0;
 	virtual void SetWindowTitleBarState(const TSharedPtr<SWidget>& TitleBarContent, EWindowTitleBarMode Mode, bool bTitleBarDragEnabled, bool bWindowButtonsVisible, bool bTitleBarVisible) = 0;
 	virtual void RestorePreviousWindowTitleBarState() = 0;
@@ -99,6 +102,17 @@ public:
 
 	SLATE_END_ARGS()
 
+	enum class EGameLayerOrder : int32
+	{
+		Player = 100,
+		Viewport = 200,
+		TitleBar = 300,
+		Tooltip = 400,
+		Debug = 500,
+
+		// Add others with an increment of 100 to have enough space for injecting new layers in-between
+	};
+
 	SGameLayerManager();
 
 	virtual ~SGameLayerManager();
@@ -127,6 +141,9 @@ public:
 
 	virtual void ClearWidgets() override;
 
+	virtual void AddGameLayer(TSharedRef<SWidget> ViewportContent, int32 ZOrder) override;
+	virtual void RemoveGameLayer(TSharedRef<SWidget> ViewportContent) override;
+
 	virtual void SetDefaultWindowTitleBarHeight(float Height);
 	virtual void SetWindowTitleBarState(const TSharedPtr<SWidget>& TitleBarContent, EWindowTitleBarMode Mode, bool bTitleBarDragEnabled, bool bWindowButtonsVisible, bool bTitleBarVisible);
 	virtual void RestorePreviousWindowTitleBarState();
@@ -201,6 +218,7 @@ private:
 	TSharedPtr<SWindowTitleBarArea> TitleBarAreaVerticalBox;
 	TSharedPtr<SBox> WindowTitleBarVerticalBox;
 	TSharedPtr<SBox> WindowTitleBarOverlay;
+	TSharedPtr<SOverlay> WindowOverlay;
 
 	struct FWindowTitleBarState
 	{
diff --git a/Engine/Source/Runtime/SlateCore/Private/Widgets/SOverlay.cpp b/Engine/Source/Runtime/SlateCore/Private/Widgets/SOverlay.cpp
--- a/Engine/Source/Runtime/SlateCore/Private/Widgets/SOverlay.cpp
+++ b/Engine/Source/Runtime/SlateCore/Private/Widgets/SOverlay.cpp
@@ -13,12 +13,13 @@ void SOverlay::PrivateRegisterAttributes(FSlateAttributeInitializer& AttributeIn
 	FOverlaySlot::RegisterAttributes(Initializer);
 }
 
+PRAGMA_DISABLE_DEPRECATION_WARNINGS
 void SOverlay::FOverlaySlot::Construct(const FChildren& SlotOwner, FSlotArguments&& InArgs)
 {
+	ZOrder = InArgs._ZOrder.Get(0);
 	TBasicLayoutWidgetSlot<FOverlaySlot>::Construct(SlotOwner, MoveTemp(InArgs));
 }
 
-PRAGMA_DISABLE_DEPRECATION_WARNINGS
 void SOverlay::FOverlaySlot::SetZOrder(int32 InZOrder)
 {
 	if (SWidget* OwnerWidget = FSlotBase::GetOwnerWidget())
@@ -114,9 +115,9 @@ SOverlay::SOverlay()
 	bCanSupportFocus = false;
 }
 
-SOverlay::FOverlaySlot::FSlotArguments SOverlay::Slot()
+SOverlay::FOverlaySlot::FSlotArguments SOverlay::Slot( int32 ZOrder )
 {
-	return FOverlaySlot::FSlotArguments(MakeUnique<FOverlaySlot>(0));
+	return FOverlaySlot::FSlotArguments(MakeUnique<FOverlaySlot>(ZOrder));
 }
 
 void SOverlay::Construct( const SOverlay::FArguments& InArgs )
@@ -125,7 +126,13 @@ void SOverlay::Construct( const SOverlay::FArguments& InArgs )
 	//The Slot has a unique_ptr. It can't be copied.
 	//Previously, the Children.Add(), was wrong if we added the same slot twice (it would create a lot of issues).
 	//Because of that, it doesn't matter if we steal the slot from the FArguments and it enforces that a slot cannot be added twice.
-	Children.AddSlots(MoveTemp(const_cast<TArray<FOverlaySlot::FSlotArguments>&>(InArgs._Slots)));
+	TArray<FOverlaySlot::FSlotArguments>& SlotArguments = const_cast<TArray<FOverlaySlot::FSlotArguments>&>(InArgs._Slots);
+	SlotArguments.StableSort([](const FOverlaySlot::FSlotArguments& A, const FOverlaySlot::FSlotArguments& B)
+	{
+		return A._ZOrder.Get(0) < B._ZOrder.Get(0);
+	});
+
+	Children.AddSlots(MoveTemp(SlotArguments));
 }
 
 void SOverlay::OnArrangeChildren( const FGeometry& AllottedGeometry, FArrangedChildren& ArrangedChildren ) const
@@ -258,17 +265,33 @@ SOverlay::FScopedWidgetSlotArguments SOverlay::AddSlot( int32 ZOrder )
 	return MoveTemp(Result);
 }
 
+bool SOverlay::HasSlotWithZOrder( int32 ZOrder ) const
+{
+	return GetChildIndexByZOrder( ZOrder ) != INDEX_NONE;
+}
+
+int32 SOverlay::GetChildIndexByZOrder( int32 ZOrder ) const
+{
+	for( int32 ChildIndex=0; ChildIndex < Children.Num(); ++ChildIndex )
+	{
+		if ( Children[ChildIndex].GetZOrder() == ZOrder )
+		{
+			return ChildIndex;
+		}
+	}
+
+	return INDEX_NONE;
+}
+
 void SOverlay::RemoveSlot( int32 ZOrder )
 {
 	if (ZOrder != INDEX_NONE)
 	{
-		for( int32 ChildIndex=0; ChildIndex < Children.Num(); ++ChildIndex )
+		const int32 ChildIndex = GetChildIndexByZOrder( ZOrder );
+		if (ChildIndex != INDEX_NONE)
 		{
-			if ( Children[ChildIndex].GetZOrder() == ZOrder )
-			{
-				Children.RemoveAt( ChildIndex );
-				return;
-			}
+			Children.RemoveAt( ChildIndex );
+			return;
 		}
 
 		ensureMsgf(false, TEXT("Could not remove slot. There are no children with ZOrder %d."));
diff --git a/Engine/Source/Runtime/SlateCore/Public/Widgets/SOverlay.h b/Engine/Source/Runtime/SlateCore/Public/Widgets/SOverlay.h
--- a/Engine/Source/Runtime/SlateCore/Public/Widgets/SOverlay.h
+++ b/Engine/Source/Runtime/SlateCore/Public/Widgets/SOverlay.h
@@ -47,6 +47,7 @@ public:
 		{ }
 
 		SLATE_SLOT_BEGIN_ARGS(FOverlaySlot, TBasicLayoutWidgetSlot<FOverlaySlot>)
+			SLATE_ARGUMENT(TOptional<int32>, ZOrder)
 		SLATE_SLOT_END_ARGS()
 
 		void Construct(const FChildren& SlotOwner, FSlotArguments&& InArgs);
@@ -98,14 +99,20 @@ public:
 	/** Adds a slot at the specified location (ignores Z-order) */
 	FScopedWidgetSlotArguments AddSlot(int32 ZOrder=INDEX_NONE);
 
+	/** Returns true if there is a child slot with the specified z-order */
+	bool HasSlotWithZOrder(int32 ZOrder) const;
+
 	/** Removes a slot at the specified location */
 	void RemoveSlot(int32 ZOrder=INDEX_NONE);
 
 	/** Removes all children from the overlay */
 	void ClearChildren();
 
-	/** @return a new slot. Slots contain children for SOverlay */
-	static FOverlaySlot::FSlotArguments Slot();
+	/**
+	 * @param ZOrder Default z-order of the slot
+	 * @return a new slot. Slots contain children for SOverlay
+	 */
+	static FOverlaySlot::FSlotArguments Slot(int32 ZOrder = 0);
 
 	//~ Begin of SWidget interface
 	virtual void OnArrangeChildren( const FGeometry& AllottedGeometry, FArrangedChildren& ArrangedChildren ) const override;
@@ -118,6 +125,9 @@ protected:
 	virtual FVector2D ComputeDesiredSize(float) const override;
 	//~ End SWidget overrides.
 
+	/** Returns the index of the child slot at the specified z-order */
+	int32 GetChildIndexByZOrder(int32 ZOrder) const;
+
 protected:
 	/** The SOverlay's slots; each slot contains a child widget. */
 	TPanelChildren<FOverlaySlot> Children;
