Changing FPointerEvent to export only 2 methods instead everything which included all its getters.

diff --git a/Engine/Source/Runtime/SlateCore/Public/Input/Events.h b/Engine/Source/Runtime/SlateCore/Public/Input/Events.h
--- a/Engine/Source/Runtime/SlateCore/Public/Input/Events.h
+++ b/Engine/Source/Runtime/SlateCore/Public/Input/Events.h
@@ -681,7 +681,7 @@ public:
  * It is passed to event handlers dealing with pointer-based input.
  */
 USTRUCT(BlueprintType)
-struct SLATECORE_API FPointerEvent
+struct FPointerEvent
 	: public FInputEvent
 {
 	GENERATED_USTRUCT_BODY()
@@ -998,9 +998,9 @@ public:
 	/** Returns the full set of pressed buttons */
 	const TSet<FKey>& GetPressedButtons() const { return *PressedButtons; }
 
-	virtual FText ToText() const override;
+	SLATECORE_API virtual FText ToText() const override;
 
-	virtual bool IsPointerEvent() const override;
+	SLATECORE_API virtual bool IsPointerEvent() const override;
 
 	template<typename PointerEventType>
 	static PointerEventType MakeTranslatedEvent( const PointerEventType& InPointerEvent, const FVirtualPointerPosition& VirtualPosition )
