Remove the use of BULKDATA_SerializeCompressedBitWindow as it has not done anything for some time.

#rnx

diff --git a/Engine/Source/Runtime/SlateCore/Private/Fonts/FontBulkData.cpp b/Engine/Source/Runtime/SlateCore/Private/Fonts/FontBulkData.cpp
--- a/Engine/Source/Runtime/SlateCore/Private/Fonts/FontBulkData.cpp
+++ b/Engine/Source/Runtime/SlateCore/Private/Fonts/FontBulkData.cpp
@@ -101,7 +101,7 @@ void UFontBulkData::Serialize(FArchive& Ar)
 
 	if (Ar.IsSaving())
 	{
-		BulkData.SetBulkDataFlags(BULKDATA_SerializeCompressed | BULKDATA_SerializeCompressedBitWindow);
+		BulkData.SetBulkDataFlags(BULKDATA_SerializeCompressed);
 	}
 	
 	BulkData.Serialize(Ar, this, INDEX_NONE, false);
