Inline slate element iteration

diff --git a/Engine/Source/Runtime/SlateCore/Public/Layout/ChildrenBase.h b/Engine/Source/Runtime/SlateCore/Public/Layout/ChildrenBase.h
--- a/Engine/Source/Runtime/SlateCore/Public/Layout/ChildrenBase.h
+++ b/Engine/Source/Runtime/SlateCore/Public/Layout/ChildrenBase.h
@@ -63,7 +63,7 @@ public:
 
 	/** Applies the predicate to all the widgets contained by the FChildren. */
 	template<typename Predicate>
-	void ForEachWidget(Predicate Pred)
+	FORCEINLINE void ForEachWidget(Predicate&& Pred)
 	{
 #if UE_WITH_SLATE_CHILDREN_DEBUGGING 
 		Debug_TestDestroyTag();
@@ -80,7 +80,7 @@ public:
 
 	/** Applies the predicate to all the widgets contained by the FChildren. */
 	template<typename Predicate>
-	void ForEachWidget(Predicate Pred) const
+	FORCEINLINE void ForEachWidget(Predicate&& Pred) const
 	{
 #if UE_WITH_SLATE_CHILDREN_DEBUGGING 
 		Debug_TestDestroyTag();
