Slate batching improvements for large element sets
---
Group similar elements at compile time using tuples instead of maps
Merge similar renderbatches in same-typed element buckets (Currently for box and lines only).
Reserve verticies / elements in a batch before using them
Add some comments to invalidation.

diff --git a/Engine/Source/Developer/SlateReflector/Private/VisualTreeCapture.cpp b/Engine/Source/Developer/SlateReflector/Private/VisualTreeCapture.cpp
--- a/Engine/Source/Developer/SlateReflector/Private/VisualTreeCapture.cpp
+++ b/Engine/Source/Developer/SlateReflector/Private/VisualTreeCapture.cpp
@@ -59,20 +59,25 @@ void FVisualEntry::Resolve(const FSlateWindowElementList& ElementList)
 		return;
 	}
 
-	if (const FSlateDrawElementContainer* Container = ElementList.GetUncachedDrawElements().Find(ElementType))
+	auto ResolveBounds = [&](const FSlateDrawElementContainer& Container, uint8 InElementType)
 	{
-		const FSlateDrawElement& Element = Container->Elements[ElementIndex];
-		const FSlateRenderTransform& Transform = Element.GetRenderTransform();
-		const FVector2D LocalSize = Element.GetLocalSize();
+		if (InElementType == (uint8)ElementType)
+		{
+			const FSlateDrawElement& Element = Container.Elements[ElementIndex];
+			const FSlateRenderTransform& Transform = Element.GetRenderTransform();
+			const FVector2D LocalSize = Element.GetLocalSize();
 
-		TopLeft = Transform.TransformPoint(FVector2D(0, 0));
-		TopRight = Transform.TransformPoint(FVector2D(LocalSize.X, 0));
-		BottomLeft = Transform.TransformPoint(FVector2D(0, LocalSize.Y));
-		BottomRight = Transform.TransformPoint(LocalSize);
+			TopLeft = Transform.TransformPoint(FVector2D(0, 0));
+			TopRight = Transform.TransformPoint(FVector2D(LocalSize.X, 0));
+			BottomLeft = Transform.TransformPoint(FVector2D(0, LocalSize.Y));
+			BottomRight = Transform.TransformPoint(LocalSize);
 
-		LayerId = Element.GetLayer();
-		ClippingIndex = Element.GetPrecachedClippingIndex();
-	}
+			LayerId = Element.GetLayer();
+			ClippingIndex = Element.GetPrecachedClippingIndex();
+		}
+	};
+
+	VisitTupleElements(ResolveBounds, ElementList.GetUncachedDrawElements(), UE::Slate::MakeTupleIndicies<uint8, (uint8)EElementType::ET_Count>());
 }
 
 bool FVisualEntry::IsPointInside(const FVector2D& Point) const
@@ -182,9 +187,9 @@ void FVisualTreeCapture::AddInvalidationRootCachedEntries(TSharedRef<FVisualTree
 	const TArray<TSharedPtr<FSlateCachedElementList>>& CachedElements = Data.GetCachedElementLists();
 	for (const TSharedPtr<FSlateCachedElementList>& CachedElement : CachedElements)
 	{
-		for (TPair<EElementType, FSlateDrawElementContainer>& DrawElement : CachedElement->DrawElements)
+		auto AddTypedEntries = [&](FSlateDrawElementContainer& Container)
 		{
-			const FSlateDrawElementArray& DrawElementArray = DrawElement.Value.Elements;
+			const FSlateDrawElementArray& DrawElementArray = Container.Elements;
 			const SWidget* Widget = CachedElement->OwningWidget;
 			// todo, should check if parents has the metadata also
 			if (Widget && !Widget->GetMetaData<FInvisibleToWidgetReflectorMetaData>())
@@ -205,7 +210,9 @@ void FVisualTreeCapture::AddInvalidationRootCachedEntries(TSharedRef<FVisualTree
 					}
 				}
 			}
-		}
+		};
+
+		VisitTupleElements(AddTypedEntries, CachedElement->DrawElements);
 	}
 }
 
@@ -297,20 +304,20 @@ void FVisualTreeCapture::EndWidgetPaint(const SWidget* Widget, const FSlateWindo
 	}
 }
 
-void FVisualTreeCapture::ElementTypeAdded(const FSlateWindowElementList& ElementList, int32 InElementIndex, EElementType InElementType)
+void FVisualTreeCapture::ElementTypeAdded(const FSlateDebuggingElementTypeAddedEventArgs& ElementTypeAddedArgs)
 {
 	if (WindowIsInvalidationRootCounter > 0 || WidgetIsInvalidationRootCounter > 0 || WidgetIsInvisibleToWidgetReflectorCounter > 0)
 	{
 		return;
 	}
 
-	TSharedPtr<FVisualTreeSnapshot> Tree = VisualTrees.FindRef(ElementList.GetPaintWindow());
+	TSharedPtr<FVisualTreeSnapshot> Tree = VisualTrees.FindRef(ElementTypeAddedArgs.ElementList.GetPaintWindow());
 	if (Tree.IsValid())
 	{
 		if (Tree->WidgetStack.Num() > 0)
 		{
 			// Ignore any element added from a widget that's invisible to the widget reflector.
-			Tree->Entries.Emplace(Tree->WidgetStack.Top(), InElementIndex, InElementType);
+			Tree->Entries.Emplace(Tree->WidgetStack.Top(), ElementTypeAddedArgs.ElementIndex, ElementTypeAddedArgs.ElementType);
 		}
 	}
 }
diff --git a/Engine/Source/Developer/SlateReflector/Private/VisualTreeCapture.h b/Engine/Source/Developer/SlateReflector/Private/VisualTreeCapture.h
--- a/Engine/Source/Developer/SlateReflector/Private/VisualTreeCapture.h
+++ b/Engine/Source/Developer/SlateReflector/Private/VisualTreeCapture.h
@@ -13,6 +13,7 @@ class FPaintArgs;
 struct FGeometry;
 class FSlateRect;
 class FSlateInvalidationRoot;
+struct FSlateDebuggingElementTypeAddedEventArgs;
 enum class EElementType : uint8;
 
 class FVisualEntry
@@ -81,7 +82,7 @@ private:
 	void EndWidgetPaint(const SWidget* Widget, const FSlateWindowElementList& ElementList, int32 LayerId);
 
 	/**  */
-	void ElementTypeAdded(const FSlateWindowElementList& ElementList, int32 InElementIndex, EElementType InElementType);
+	void ElementTypeAdded(const FSlateDebuggingElementTypeAddedEventArgs& ElementTypeAddedArgs);
 
 	void OnWindowBeingDestroyed(const SWindow& WindowBeingDestoyed);
 private:
diff --git a/Engine/Source/Runtime/SlateCore/Private/Debugging/SlateDebugging.cpp b/Engine/Source/Runtime/SlateCore/Private/Debugging/SlateDebugging.cpp
--- a/Engine/Source/Runtime/SlateCore/Private/Debugging/SlateDebugging.cpp
+++ b/Engine/Source/Runtime/SlateCore/Private/Debugging/SlateDebugging.cpp
@@ -70,7 +70,7 @@ FSlateDebuggingFocusEventArgs::FSlateDebuggingFocusEventArgs(
 	, NewFocusedWidgetPath(InNewFocusedWidgetPath)
 	, NewFocusedWidget(InNewFocusedWidget)
 {
-}
+} 
 
 FText FSlateDebuggingFocusEventArgs::ToText() const
 {
@@ -350,6 +350,16 @@ FSlateDebuggingWidgetUpdatedEventArgs::FSlateDebuggingWidgetUpdatedEventArgs(
 {
 }
 
+FSlateDebuggingElementTypeAddedEventArgs::FSlateDebuggingElementTypeAddedEventArgs(
+	const FSlateWindowElementList& InElementList,
+	int32 InElementIndex,
+	EElementType InElementType)
+	: ElementList(InElementList)
+	, ElementIndex(InElementIndex)
+	, ElementType(InElementType)
+{
+}
+
 FText FSlateDebuggingWidgetUpdatedEventArgs::ToText() const
 {
 	TArray<FText> UpdateText;
diff --git a/Engine/Source/Runtime/SlateCore/Private/FastUpdate/SlateInvalidationRoot.cpp b/Engine/Source/Runtime/SlateCore/Private/FastUpdate/SlateInvalidationRoot.cpp
--- a/Engine/Source/Runtime/SlateCore/Private/FastUpdate/SlateInvalidationRoot.cpp
+++ b/Engine/Source/Runtime/SlateCore/Private/FastUpdate/SlateInvalidationRoot.cpp
@@ -394,6 +394,7 @@ FSlateInvalidationResult FSlateInvalidationRoot::PaintInvalidationRoot(const FSl
 				BuildFastPathWidgetList(RootWidget);
 			}
 
+			// Repopulates cached element lists
 			CachedMaxLayerId = PaintSlowPath(Context);
 #if WITH_SLATE_DEBUGGING
 			SetLastPaintType(ESlateInvalidationPaintType::Slow);
diff --git a/Engine/Source/Runtime/SlateCore/Private/Rendering/DrawElements.cpp b/Engine/Source/Runtime/SlateCore/Private/Rendering/DrawElements.cpp
--- a/Engine/Source/Runtime/SlateCore/Private/Rendering/DrawElements.cpp
+++ b/Engine/Source/Runtime/SlateCore/Private/Rendering/DrawElements.cpp
@@ -273,7 +273,7 @@ void FSlateDrawElement::MakeDebugQuad( FSlateWindowElementList& ElementList, uin
 		return;
 	}
 
-	FSlateDrawElement& Element = ElementList.AddUninitialized(EElementType::ET_DebugQuad);
+	FSlateDrawElement& Element = ElementList.AddUninitialized<EElementType::ET_DebugQuad>();
 	FSlateBoxPayload& BoxPayload = ElementList.CreatePayload<FSlateBoxPayload>(Element);
 
 	BoxPayload.SetTint(Tint);
@@ -294,7 +294,7 @@ FSlateDrawElement& FSlateDrawElement::MakeBoxInternal(
 	// Cast to Rounded Rect to get the internal parameters 
 	// New payload type - inherit from BoxPayload 
 
-	FSlateDrawElement& Element = ElementList.AddUninitialized(ElementType);
+	FSlateDrawElement& Element = ElementList.AddUninitializedLookup(ElementType);
 
 	FSlateBoxPayload* BoxPayload;
 	if ( ElementType == EElementType::ET_RoundedBox )
@@ -405,7 +405,7 @@ void FSlateDrawElement::MakeText( FSlateWindowElementList& ElementList, uint32 I
 		UE_LOG(LogSlate, Log, TEXT("MakeText: '%s'."), *InText);
 	}
 #endif
-	FSlateDrawElement& Element = ElementList.AddUninitialized(EElementType::ET_Text);
+	FSlateDrawElement& Element = ElementList.AddUninitialized<EElementType::ET_Text>();
 	FSlateTextPayload& DataPayload = ElementList.CreatePayload<FSlateTextPayload>(Element);
 
 	DataPayload.SetTint(InTint);
@@ -441,7 +441,7 @@ void FSlateDrawElement::MakeText( FSlateWindowElementList& ElementList, uint32 I
 		UE_LOG(LogSlate, Log, TEXT("MakeText: '%s'."), *InText);
 	}
 #endif
-	FSlateDrawElement& Element = ElementList.AddUninitialized(EElementType::ET_Text);
+	FSlateDrawElement& Element = ElementList.AddUninitialized<EElementType::ET_Text>();
 
 	FSlateTextPayload& DataPayload = ElementList.CreatePayload<FSlateTextPayload>(Element);
 
@@ -515,7 +515,7 @@ void FSlateDrawElement::MakeShapedText(FSlateWindowElementList& ElementList, uin
 	SlateDrawElement::CheckInvalidUMaterial(InShapedGlyphSequence->GetFontOutlineSettings().OutlineMaterial, TEXT("Outline Material"));
 #endif
 
-	FSlateDrawElement& Element = ElementList.AddUninitialized(EElementType::ET_ShapedText);
+	FSlateDrawElement& Element = ElementList.AddUninitialized<EElementType::ET_ShapedText>();
 
 	FSlateShapedTextPayload& DataPayload = ElementList.CreatePayload<FSlateShapedTextPayload>(Element);
 	DataPayload.SetTint(BaseTint);
@@ -534,7 +534,7 @@ void FSlateDrawElement::MakeGradient( FSlateWindowElementList& ElementList, uint
 		return;
 	}
 
-	FSlateDrawElement& Element = ElementList.AddUninitialized(EElementType::ET_Gradient);
+	FSlateDrawElement& Element = ElementList.AddUninitialized<EElementType::ET_Gradient>();
 
 	FSlateGradientPayload& DataPayload = ElementList.CreatePayload<FSlateGradientPayload>(Element);
 
@@ -551,7 +551,7 @@ void FSlateDrawElement::MakeSpline(FSlateWindowElementList& ElementList, uint32
 	{
 		return;
 	}
-	FSlateDrawElement& Element = ElementList.AddUninitialized(EElementType::ET_Spline);
+	FSlateDrawElement& Element = ElementList.AddUninitialized<EElementType::ET_Spline>();
 
 	FSlateSplinePayload& DataPayload = ElementList.CreatePayload<FSlateSplinePayload>(Element);
 
@@ -568,7 +568,7 @@ void FSlateDrawElement::MakeCubicBezierSpline(FSlateWindowElementList& ElementLi
 	{
 		return;
 	}
-	FSlateDrawElement& Element = ElementList.AddUninitialized(EElementType::ET_Spline);
+	FSlateDrawElement& Element = ElementList.AddUninitialized<EElementType::ET_Spline>();
 
 	FSlateSplinePayload& DataPayload = ElementList.CreatePayload<FSlateSplinePayload>(Element);
 
@@ -609,7 +609,7 @@ void FSlateDrawElement::MakeLines(FSlateWindowElementList& ElementList, uint32 I
 		return;
 	}
 
-	FSlateDrawElement& Element = ElementList.AddUninitialized(EElementType::ET_Line);
+	FSlateDrawElement& Element = ElementList.AddUninitialized<EElementType::ET_Line>();
 
 	FSlateLinePayload& DataPayload = ElementList.CreatePayload<FSlateLinePayload>(Element);
 
@@ -655,7 +655,7 @@ void FSlateDrawElement::MakeLines( FSlateWindowElementList& ElementList, uint32
 		return;
 	}
 
-	FSlateDrawElement& Element = ElementList.AddUninitialized(EElementType::ET_Line);
+	FSlateDrawElement& Element = ElementList.AddUninitialized<EElementType::ET_Line>();
 
 	FSlateLinePayload& DataPayload = ElementList.CreatePayload<FSlateLinePayload>(Element);
 	DataPayload.SetTint(InTint);
@@ -674,7 +674,7 @@ void FSlateDrawElement::MakeViewport( FSlateWindowElementList& ElementList, uint
 		return;
 	}
 
-	FSlateDrawElement& Element = ElementList.AddUninitialized(EElementType::ET_Viewport);
+	FSlateDrawElement& Element = ElementList.AddUninitialized<EElementType::ET_Viewport>();
 	FSlateViewportPayload& DataPayload = ElementList.CreatePayload<FSlateViewportPayload>(Element);
 
 	DataPayload.SetViewport(Viewport, InTint);
@@ -691,7 +691,7 @@ void FSlateDrawElement::MakeCustom( FSlateWindowElementList& ElementList, uint32
 		return;
 	}
 
-	FSlateDrawElement& Element = ElementList.AddUninitialized(EElementType::ET_Custom);
+	FSlateDrawElement& Element = ElementList.AddUninitialized<EElementType::ET_Custom>();
 
 	FSlateCustomDrawerPayload& DataPayload = ElementList.CreatePayload<FSlateCustomDrawerPayload>(Element);
 	DataPayload.SetCustomDrawer(CustomDrawer);
@@ -710,7 +710,7 @@ void FSlateDrawElement::MakeCustomVerts(FSlateWindowElementList& ElementList, ui
 		return;
 	}
 
-	FSlateDrawElement& Element = ElementList.AddUninitialized(EElementType::ET_CustomVerts);
+	FSlateDrawElement& Element = ElementList.AddUninitialized<EElementType::ET_CustomVerts>();
 	FSlateCustomVertsPayload& DataPayload = ElementList.CreatePayload<FSlateCustomVertsPayload>(Element);
 
 	const FSlateShaderResourceProxy* RenderingProxy = InRenderResourceHandle.GetResourceProxy();
@@ -730,7 +730,7 @@ void FSlateDrawElement::MakePostProcessPass(FSlateWindowElementList& ElementList
 		return;
 	}
 
-	FSlateDrawElement& Element = ElementList.AddUninitialized(EElementType::ET_PostProcessPass);
+	FSlateDrawElement& Element = ElementList.AddUninitialized<EElementType::ET_PostProcessPass>();
 
 	FSlatePostProcessPayload& DataPayload = ElementList.CreatePayload<FSlatePostProcessPayload>(Element);
 	DataPayload.DownsampleAmount = DownsampleAmount;
@@ -806,32 +806,6 @@ bool FSlateWindowElementList::GetIsInGameLayer()
 	return bIsInGameLayer;
 }
 
-
-FSlateDrawElement& FSlateWindowElementList::AddUninitialized(EElementType InElementType)
-{
-	const bool bAllowCache = CachedElementDataListStack.Num() > 0 && WidgetDrawStack.Num() && !WidgetDrawStack.Top().bIsVolatile;
-
-	if (bAllowCache)
-	{
-		// @todo get working with slate debugging
-		return AddCachedElement(InElementType);
-	}
-	else
-	{
-		FSlateDrawElementMap& Elements = UncachedDrawElements;
-		FSlateDrawElementContainer& Container = Elements.FindOrAdd(InElementType);
-		const int32 InsertIdx = Container.Elements.AddDefaulted();
-
-#if WITH_SLATE_DEBUGGING
-		FSlateDebugging::ElementTypeAdded.Broadcast(*this, InsertIdx, InElementType);
-#endif
-
-		FSlateDrawElement& NewElement = Container.Elements[InsertIdx];
-		return NewElement;
-	}
-}
-
-
 FSlateWindowElementList::FDeferredPaint::FDeferredPaint( const TSharedRef<const SWidget>& InWidgetToPaint, const FPaintArgs& InArgs, const FGeometry InAllottedGeometry, const FWidgetStyle& InWidgetStyle, bool InParentEnabled )
 	: WidgetToPaintPtr( InWidgetToPaint )
 	, Args( InArgs )
@@ -977,22 +951,6 @@ void FSlateWindowElementList::PopBatchPriortyGroup()
 	//BatchDepthPriorityStack.Pop();
 }*/
 
-FSlateDrawElement& FSlateWindowElementList::AddCachedElement(EElementType InElementType)
-{
-	FSlateCachedElementData* CurrentCachedElementData = GetCurrentCachedElementData();
-	check(CurrentCachedElementData);
-
-	FWidgetDrawElementState& CurrentWidgetState = WidgetDrawStack.Top();
-	check(!CurrentWidgetState.bIsVolatile);
-
-	if (!CurrentWidgetState.CacheHandle.IsValid())
-	{
-		CurrentWidgetState.CacheHandle = CurrentCachedElementData->AddCache(CurrentWidgetState.Widget);
-	}
-
-	return CurrentCachedElementData->AddCachedElement(CurrentWidgetState.CacheHandle, GetClippingManager(), CurrentWidgetState.Widget, InElementType);
-}
-
 void FSlateWindowElementList::PushCachedElementData(FSlateCachedElementData& CachedElementData)
 {
 	check(&CachedElementData); 
@@ -1005,6 +963,81 @@ void FSlateWindowElementList::PopCachedElementData()
 	CachedElementDataListStack.Pop();
 }
 
+FSlateDrawElement& FSlateWindowElementList::AddUninitializedLookup(EElementType InElementType)
+{
+	// Determine what type of element to add
+	switch (InElementType)
+	{
+	case EElementType::ET_Box:
+	{
+		return AddUninitialized<EElementType::ET_Box>();
+	}
+	break;
+	case EElementType::ET_RoundedBox:
+	{
+		return AddUninitialized<EElementType::ET_RoundedBox>();
+	}
+	break;
+	case EElementType::ET_Border:
+	{
+		return AddUninitialized<EElementType::ET_Border>();
+	}
+	break;
+	case EElementType::ET_Text:
+	{
+		return AddUninitialized<EElementType::ET_Text>();
+	}
+	break;
+	case EElementType::ET_ShapedText:
+	{
+		return AddUninitialized<EElementType::ET_ShapedText>();
+	}
+	break;
+	case EElementType::ET_Line:
+	{
+		return AddUninitialized<EElementType::ET_Line>();
+	}
+	break;
+	case EElementType::ET_DebugQuad:
+	{
+		return AddUninitialized<EElementType::ET_DebugQuad>();
+	}
+	break;
+	case EElementType::ET_Spline:
+	{
+		return AddUninitialized<EElementType::ET_Spline>();
+	}
+	break;
+	case EElementType::ET_Gradient:
+	{
+		return AddUninitialized<EElementType::ET_Gradient>();
+	}
+	break;
+	case EElementType::ET_Viewport:
+	{
+		return AddUninitialized<EElementType::ET_Viewport>();
+	}
+	break;
+	case EElementType::ET_Custom:
+	{
+		return AddUninitialized<EElementType::ET_Custom>();
+	}
+	break;
+	case EElementType::ET_CustomVerts:
+	{
+		return AddUninitialized<EElementType::ET_CustomVerts>();
+	}
+	break;
+	case EElementType::ET_PostProcessPass:
+	{
+		return AddUninitialized<EElementType::ET_PostProcessPass>();
+	}
+	break;
+	default:
+		return AddUninitialized<EElementType::ET_NonMapped>();
+	}
+}
+
 int32 FSlateWindowElementList::PushClip(const FSlateClippingZone& InClipZone)
 {
 	const int32 NewClipIndex = ClippingManager.PushClip(InClipZone);
@@ -1071,7 +1104,12 @@ void FSlateWindowElementList::ResetElementList()
 
 	ClippingManager.ResetClippingState();
 
-	UncachedDrawElements.Reset();
+	auto ResetElementContainer = [](FSlateDrawElementContainer& Container)
+	{
+		Container.Elements.Reset();
+	};
+
+	VisitTupleElements(ResetElementContainer, UncachedDrawElements);
 
 #if STATS
 	const int32 DeltaMemory = MemManager.GetByteCount() - MemManagerAllocatedMemory;
@@ -1095,16 +1133,18 @@ void FSlateWindowElementList::ResetElementList()
 
 void FSlateWindowElementList::AddReferencedObjects(FReferenceCollector& Collector)
 {
-	for (TPair<EElementType, FSlateDrawElementContainer>& DrawElement : UncachedDrawElements)
+	auto AddReferencedElementContainer = [&](FSlateDrawElementContainer& Container)
 	{
-		for (FSlateDrawElement& Element : DrawElement.Value.Elements)
+		for (FSlateDrawElement& Element : Container.Elements)
 		{
 			Element.AddReferencedObjects(Collector);
 		}
-	}
+	};
+
+	VisitTupleElements(AddReferencedElementContainer, UncachedDrawElements);
 }
 
-static const FSlateClippingState* GetClipStateFromParent(const FSlateClippingManager& ParentClipManager)
+const FSlateClippingState* FSlateCachedElementData::GetClipStateFromParent(const FSlateClippingManager& ParentClipManager)
 {
 	const int32 ClippingIndex = ParentClipManager.GetClippingIndex();
 
@@ -1118,6 +1158,14 @@ static const FSlateClippingState* GetClipStateFromParent(const FSlateClippingMan
 	}
 }
 
+void FSlateCachedElementData::ValidateWidgetOwner(TSharedPtr<FSlateCachedElementList> List, const SWidget* CurrentWidget)
+{
+#if WITH_SLATE_DEBUGGING
+	checkSlow(List->OwningWidget == CurrentWidget);
+	checkSlow(CurrentWidget->GetParentWidget().IsValid());
+#endif
+}
+
 void FSlateCachedElementData::Empty()
 {
 	if (CachedElementLists.Num())
@@ -1155,35 +1203,6 @@ FSlateCachedElementsHandle FSlateCachedElementData::AddCache(const SWidget* Widg
 	return FSlateCachedElementsHandle(NewList);
 }
 
-FSlateDrawElement& FSlateCachedElementData::AddCachedElement(FSlateCachedElementsHandle& CacheHandle, const FSlateClippingManager& ParentClipManager, const SWidget* CurrentWidget, EElementType InElementType)
-{
-	TSharedPtr<FSlateCachedElementList> List = CacheHandle.Ptr.Pin();
-
-#if WITH_SLATE_DEBUGGING
-	check(List->OwningWidget == CurrentWidget);
-	check(CurrentWidget->GetParentWidget().IsValid());
-#endif
-	
-	FSlateDrawElementContainer& Container = List->DrawElements.FindOrAdd(InElementType);
-	FSlateDrawElement& NewElement = Container.Elements.AddDefaulted_GetRef();
-	NewElement.SetIsCached(true);
-
-	// Check if slow vs checking a flag on the list to see if it contains new data.
-	ListsWithNewData.AddUnique(List.Get());
-
-	const FSlateClippingState* ExistingClipState = GetClipStateFromParent(ParentClipManager);
-
-	if (ExistingClipState)
-	{
-		// We need to cache this clip state for the next time the element draws
-		FSlateCachedClipState& CachedClipState = FindOrAddCachedClipState(ExistingClipState);
-		List->AddCachedClipState(CachedClipState);
-		NewElement.SetCachedClippingState(&CachedClipState.ClippingState.Get());
-	}
-
-	return NewElement;
-}
-
 FSlateRenderBatch& FSlateCachedElementData::AddCachedRenderBatch(FSlateRenderBatch&& NewBatch, int32& OutIndex)
 {
 	// Check perf against add.  AddAtLowest makes it generally re-add elements at the same index it just removed which is nicer on the cache
@@ -1288,19 +1307,57 @@ void FSlateCachedElementList::AddCachedClipState(FSlateCachedClipState& ClipStat
 
 void FSlateCachedElementList::AddReferencedObjects(FReferenceCollector& Collector)
 {
-	for (TPair<EElementType, FSlateDrawElementContainer>& DrawElement : DrawElements)
+	auto AddReferencedElementContainer = [&](FSlateDrawElementContainer& Container)
 	{
-		for (FSlateDrawElement& Element : DrawElement.Value.Elements)
+		for (FSlateDrawElement& Element : Container.Elements)
 		{
 			Element.AddReferencedObjects(Collector);
 		}
-	}
+	};
+
+	VisitTupleElements(AddReferencedElementContainer, DrawElements);
+}
+
+bool FSlateCachedElementList::IsEmpty()
+{
+	bool bElementsFound = false;
+	
+	auto CheckElementsExist = [&](FSlateDrawElementContainer& Container)
+	{
+		if (!bElementsFound)
+		{
+			bElementsFound = Container.Elements.Num() > 0;
+		}
+	};
+
+	VisitTupleElements(CheckElementsExist, DrawElements);
+
+	return bElementsFound;
+}
+
+int32 FSlateCachedElementList::NumElements()
+{
+	int32 NumElements = 0;
+	
+	auto CountNumElements = [&](FSlateDrawElementContainer& Container)
+	{
+		NumElements += Container.Elements.Num();
+	};
+
+	VisitTupleElements(CountNumElements, DrawElements);
+
+	return NumElements;
 }
 
 void FSlateCachedElementList::DestroyCachedData()
 {
 	// Clear out any cached draw elements
-	DrawElements.Reset();
+	auto ResetElementContainer = [](FSlateDrawElementContainer& Container)
+	{
+		Container.Elements.Reset();
+	};
+
+	VisitTupleElements(ResetElementContainer, DrawElements);
 
 	// Clear out any cached render batches
 	if (CachedRenderBatchIndices.Num())
@@ -1363,7 +1420,19 @@ bool FSlateCachedElementsHandle::HasCachedElements() const
 {
 	if (const TSharedPtr<FSlateCachedElementList> List = Ptr.Pin())
 	{
-		return List->DrawElements.Num() > 0;
+		bool bElementsFound = false;
+
+		auto CheckElementsExist = [&](FSlateDrawElementContainer& Container)
+		{
+			if (!bElementsFound)
+			{
+				bElementsFound = Container.Elements.Num() > 0;
+			}
+		};
+
+		VisitTupleElements(CheckElementsExist, List->DrawElements);
+
+		return bElementsFound;
 	}
 
 	return false;
diff --git a/Engine/Source/Runtime/SlateCore/Private/Rendering/ElementBatcher.cpp b/Engine/Source/Runtime/SlateCore/Private/Rendering/ElementBatcher.cpp
--- a/Engine/Source/Runtime/SlateCore/Private/Rendering/ElementBatcher.cpp
+++ b/Engine/Source/Runtime/SlateCore/Private/Rendering/ElementBatcher.cpp
@@ -3,6 +3,7 @@
 #include "Rendering/ElementBatcher.h"
 #include "Fonts/SlateFontInfo.h"
 #include "Fonts/FontCache.h"
+#include "Rendering/RenderingCommon.h"
 #include "Rendering/DrawElementPayloads.h"
 #include "Rendering/DrawElements.h"
 #include "Rendering/RenderingPolicy.h"
@@ -355,7 +356,6 @@ void FSlateElementBatcher::AddElements(FSlateWindowElementList& WindowElementLis
 
 void FSlateElementBatcher::AddElementsInternal(const FSlateDrawElementMap& DrawElements, FVector2f ViewportSize)
 {
-	// @TODO: DarenC = Warning this is ugly but right thing to do.
 	auto HandleUnMapped = [&](const FSlateDrawElementArray& Elements, FVector2f ViewportSize)
 	{
 		for (const FSlateDrawElement& DrawElement : Elements)
@@ -366,9 +366,7 @@ void FSlateElementBatcher::AddElementsInternal(const FSlateDrawElementMap& DrawE
 			case EElementType::ET_Box:
 			case EElementType::ET_RoundedBox:
 			{
-				SCOPED_NAMED_EVENT_TEXT("Slate::AddBoxElement", FColor::Magenta);
-				STAT(ElementStat_Boxes++);
-				DrawElement.IsPixelSnapped() ? AddBoxElement<ESlateVertexRounding::Enabled>(DrawElement) : AddBoxElement<ESlateVertexRounding::Disabled>(DrawElement);
+				checkf(0, TEXT("All box elements should be mapped and use bulk element processing"));
 			}
 			break;
 			case EElementType::ET_Border:
@@ -394,9 +392,7 @@ void FSlateElementBatcher::AddElementsInternal(const FSlateDrawElementMap& DrawE
 			break;
 			case EElementType::ET_Line:
 			{
-				SCOPED_NAMED_EVENT_TEXT("Slate::AddLineElement", FColor::Magenta);
-				STAT(ElementStat_Line++);
-				DrawElement.IsPixelSnapped() ? AddLineElement<ESlateVertexRounding::Enabled>(DrawElement) : AddLineElement<ESlateVertexRounding::Disabled>(DrawElement);
+				checkf(0, TEXT("All line elements should be mapped and use bulk element processing"));
 			}
 			break;
 			case EElementType::ET_DebugQuad:
@@ -455,146 +451,134 @@ void FSlateElementBatcher::AddElementsInternal(const FSlateDrawElementMap& DrawE
 		}
 	};
 
-	for (const TPair<EElementType, FSlateDrawElementContainer>& DrawElementMap : DrawElements)
 	{
-		EElementType ElmentType = DrawElementMap.Key;
-		const FSlateDrawElementContainer& Container = DrawElementMap.Value;
-		// Determine what type of element to add
-		switch (ElmentType)
-		{
-		case EElementType::ET_Box:
-		case EElementType::ET_RoundedBox:	
-		{
-			SCOPED_NAMED_EVENT_TEXT("Slate::AddRoundedBoxElement", FColor::Magenta);
-			STAT(ElementStat_Other += Container.Elements.Num());
-			for (const FSlateDrawElement& DrawElement : Container.Elements)
-			{
-				DrawElement.IsPixelSnapped() ? AddBoxElement<ESlateVertexRounding::Enabled>(DrawElement) : AddBoxElement<ESlateVertexRounding::Disabled>(DrawElement);
-			}
-		}
-			break;
-		case EElementType::ET_Border:
-		{
-			SCOPED_NAMED_EVENT_TEXT("Slate::AddBorderElement", FColor::Magenta);
-			STAT(ElementStat_Other += Container.Elements.Num());
-			for (const FSlateDrawElement& DrawElement : Container.Elements)
-			{
-				DrawElement.IsPixelSnapped() ? AddBorderElement<ESlateVertexRounding::Enabled>(DrawElement) : AddBorderElement<ESlateVertexRounding::Disabled>(DrawElement);
-			}
-		}
-			break;
-		case EElementType::ET_Text:
-		{
-			SCOPED_NAMED_EVENT_TEXT("Slate::AddTextElement", FColor::Magenta);
-			STAT(ElementStat_Other += Container.Elements.Num());
-			for (const FSlateDrawElement& DrawElement : Container.Elements)
-			{
-				DrawElement.IsPixelSnapped() ? AddTextElement<ESlateVertexRounding::Enabled>(DrawElement) : AddTextElement<ESlateVertexRounding::Disabled>(DrawElement);
-			}
-		}
-			break;
-		case EElementType::ET_ShapedText:
+		SCOPED_NAMED_EVENT_TEXT("Slate::AddBoxElements", FColor::Magenta);
+		const FSlateDrawElementContainer& Container = DrawElements.Get<(uint8)EElementType::ET_Box>();
+		STAT(ElementStat_Other += Container.Elements.Num());
+		AddBoxElements(Container.Elements);
+	}
+
+	{
+		SCOPED_NAMED_EVENT_TEXT("Slate::AddRoundedBoxElements", FColor::Magenta);
+		const FSlateDrawElementContainer& Container = DrawElements.Get<(uint8)EElementType::ET_RoundedBox>();
+		STAT(ElementStat_Other += Container.Elements.Num());
+		AddBoxElements(Container.Elements);
+	}
+
+	{
+		SCOPED_NAMED_EVENT_TEXT("Slate::AddBorderElement", FColor::Magenta);
+		const FSlateDrawElementContainer& Container = DrawElements.Get<(uint8)EElementType::ET_Border>();
+		STAT(ElementStat_Other += Container.Elements.Num());
+		for (const FSlateDrawElement& DrawElement : Container.Elements)
 		{
-			SCOPED_NAMED_EVENT_TEXT("Slate::AddShapedTextElement", FColor::Magenta);
-			STAT(ElementStat_Other += Container.Elements.Num());
-			for (const FSlateDrawElement& DrawElement : Container.Elements)
-			{
-				DrawElement.IsPixelSnapped() ? AddShapedTextElement<ESlateVertexRounding::Enabled>(DrawElement) : AddShapedTextElement<ESlateVertexRounding::Disabled>(DrawElement);
-			}
+			DrawElement.IsPixelSnapped() ? AddBorderElement<ESlateVertexRounding::Enabled>(DrawElement) : AddBorderElement<ESlateVertexRounding::Disabled>(DrawElement);
 		}
-			break;
-		case EElementType::ET_Line:
+	}
+
+	{
+		SCOPED_NAMED_EVENT_TEXT("Slate::AddTextElement", FColor::Magenta);
+		const FSlateDrawElementContainer& Container = DrawElements.Get<(uint8)EElementType::ET_Text>();
+		STAT(ElementStat_Other += Container.Elements.Num());
+		for (const FSlateDrawElement& DrawElement : Container.Elements)
 		{
-			SCOPED_NAMED_EVENT_TEXT("Slate::AddLineElement", FColor::Magenta);
-			STAT(ElementStat_Other += Container.Elements.Num());
-			for (const FSlateDrawElement& DrawElement : Container.Elements)
-			{
-				DrawElement.IsPixelSnapped() ? AddLineElement<ESlateVertexRounding::Enabled>(DrawElement) : AddLineElement<ESlateVertexRounding::Disabled>(DrawElement);
-			}
+			DrawElement.IsPixelSnapped() ? AddTextElement<ESlateVertexRounding::Enabled>(DrawElement) : AddTextElement<ESlateVertexRounding::Disabled>(DrawElement);
 		}
-			break;
-		case EElementType::ET_DebugQuad:
+	}
+
+	{
+		SCOPED_NAMED_EVENT_TEXT("Slate::AddShapedTextElement", FColor::Magenta);
+		const FSlateDrawElementContainer& Container = DrawElements.Get<(uint8)EElementType::ET_ShapedText>();
+		STAT(ElementStat_Other += Container.Elements.Num());
+		for (const FSlateDrawElement& DrawElement : Container.Elements)
 		{
-			SCOPED_NAMED_EVENT_TEXT("Slate::AddDebugQuadElement", FColor::Magenta);
-			STAT(ElementStat_Other += Container.Elements.Num());
-			for (const FSlateDrawElement& DrawElement : Container.Elements)
-			{
-				DrawElement.IsPixelSnapped() ? AddDebugQuadElement<ESlateVertexRounding::Enabled>(DrawElement) : AddDebugQuadElement<ESlateVertexRounding::Disabled>(DrawElement);
-			}
+			DrawElement.IsPixelSnapped() ? AddShapedTextElement<ESlateVertexRounding::Enabled>(DrawElement) : AddShapedTextElement<ESlateVertexRounding::Disabled>(DrawElement);
 		}
-			break;
-		case EElementType::ET_Spline:
+	}
+
+	{
+		SCOPED_NAMED_EVENT_TEXT("Slate::AddLineElements", FColor::Magenta);
+		const FSlateDrawElementContainer& Container = DrawElements.Get<(uint8)EElementType::ET_Line>();
+		STAT(ElementStat_Other += Container.Elements.Num());
+		AddLineElements(Container.Elements);
+	}
+
+	{
+		SCOPED_NAMED_EVENT_TEXT("Slate::AddDebugQuadElement", FColor::Magenta);
+		const FSlateDrawElementContainer& Container = DrawElements.Get<(uint8)EElementType::ET_DebugQuad>();
+		STAT(ElementStat_Other += Container.Elements.Num());
+		for (const FSlateDrawElement& DrawElement : Container.Elements)
 		{
-			SCOPED_NAMED_EVENT_TEXT("Slate::AddSplineElement", FColor::Magenta);
-			// Note that we ignore pixel snapping here; see implementation for more info.
-			STAT(ElementStat_Other += Container.Elements.Num());
-			for (const FSlateDrawElement& DrawElement : Container.Elements)
-			{
-				AddSplineElement(DrawElement);
-			}
+			DrawElement.IsPixelSnapped() ? AddDebugQuadElement<ESlateVertexRounding::Enabled>(DrawElement) : AddDebugQuadElement<ESlateVertexRounding::Disabled>(DrawElement);
 		}
-			break;
-		case EElementType::ET_Gradient:
+	}
+
+	{
+		SCOPED_NAMED_EVENT_TEXT("Slate::AddSplineElement", FColor::Magenta);
+		const FSlateDrawElementContainer& Container = DrawElements.Get<(uint8)EElementType::ET_Spline>();
+
+		// Note that we ignore pixel snapping here; see implementation for more info.
+		STAT(ElementStat_Other += Container.Elements.Num());
+		for (const FSlateDrawElement& DrawElement : Container.Elements)
 		{
-			SCOPED_NAMED_EVENT_TEXT("Slate::AddGradientElement", FColor::Magenta);
-			STAT(ElementStat_Other += Container.Elements.Num());
-			for (const FSlateDrawElement& DrawElement : Container.Elements)
-			{
-				DrawElement.IsPixelSnapped() ? AddGradientElement<ESlateVertexRounding::Enabled>(DrawElement) : AddGradientElement<ESlateVertexRounding::Disabled>(DrawElement);
-			}
+			AddSplineElement(DrawElement);
 		}
-			break;
-		case EElementType::ET_Viewport:
+	}
+
+	{
+		SCOPED_NAMED_EVENT_TEXT("Slate::AddGradientElement", FColor::Magenta);
+		const FSlateDrawElementContainer& Container = DrawElements.Get<(uint8)EElementType::ET_Gradient>();
+		STAT(ElementStat_Other += Container.Elements.Num());
+		for (const FSlateDrawElement& DrawElement : Container.Elements)
 		{
-			SCOPED_NAMED_EVENT_TEXT("Slate::AddViewportElement", FColor::Magenta);
-			STAT(ElementStat_Other += Container.Elements.Num());
-			for (const FSlateDrawElement& DrawElement : Container.Elements)
-			{
-				DrawElement.IsPixelSnapped() ? AddViewportElement<ESlateVertexRounding::Enabled>(DrawElement) : AddViewportElement<ESlateVertexRounding::Disabled>(DrawElement);
-			}
+			DrawElement.IsPixelSnapped() ? AddGradientElement<ESlateVertexRounding::Enabled>(DrawElement) : AddGradientElement<ESlateVertexRounding::Disabled>(DrawElement);
 		}
-			break;
-		case EElementType::ET_Custom:
+	}
+
+	{
+		SCOPED_NAMED_EVENT_TEXT("Slate::AddViewportElement", FColor::Magenta);
+		const FSlateDrawElementContainer& Container = DrawElements.Get<(uint8)EElementType::ET_Viewport>();
+		STAT(ElementStat_Other += Container.Elements.Num());
+		for (const FSlateDrawElement& DrawElement : Container.Elements)
 		{
-			SCOPED_NAMED_EVENT_TEXT("Slate::AddCustomElement", FColor::Magenta);
-			STAT(ElementStat_Other += Container.Elements.Num());
-			for (const FSlateDrawElement& DrawElement : Container.Elements)
-			{
-				AddCustomElement(DrawElement);
-			}
+			DrawElement.IsPixelSnapped() ? AddViewportElement<ESlateVertexRounding::Enabled>(DrawElement) : AddViewportElement<ESlateVertexRounding::Disabled>(DrawElement);
 		}
-			break;
-		case EElementType::ET_CustomVerts:
+	}
+
+	{
+		SCOPED_NAMED_EVENT_TEXT("Slate::AddCustomElement", FColor::Magenta);
+		const FSlateDrawElementContainer& Container = DrawElements.Get<(uint8)EElementType::ET_Custom>();
+		STAT(ElementStat_Other += Container.Elements.Num());
+		for (const FSlateDrawElement& DrawElement : Container.Elements)
 		{
-			SCOPED_NAMED_EVENT_TEXT("Slate::AddCustomVertsElement", FColor::Magenta);
-			STAT(ElementStat_Other += Container.Elements.Num());
-			for (const FSlateDrawElement& DrawElement : Container.Elements)
-			{
-				AddCustomVerts(DrawElement);
-			}
+			AddCustomElement(DrawElement);
 		}
-			break;
-		case EElementType::ET_PostProcessPass:
+	}
+
+	{
+		SCOPED_NAMED_EVENT_TEXT("Slate::AddCustomVertsElement", FColor::Magenta);
+		const FSlateDrawElementContainer& Container = DrawElements.Get<(uint8)EElementType::ET_CustomVerts>();
+		STAT(ElementStat_Other += Container.Elements.Num());
+		for (const FSlateDrawElement& DrawElement : Container.Elements)
 		{
-			SCOPED_NAMED_EVENT_TEXT("Slate::AddPostProcessElement", FColor::Magenta);
-			STAT(ElementStat_Other += Container.Elements.Num());
-			for (const FSlateDrawElement& DrawElement : Container.Elements)
-			{
-				AddPostProcessPass(DrawElement, ViewportSize);
-			}
+			AddCustomVerts(DrawElement);
 		}
-			break;
-		case EElementType::ET_NonMapped:
+	}
+
+	{
+		SCOPED_NAMED_EVENT_TEXT("Slate::AddPostProcessElement", FColor::Magenta);
+		const FSlateDrawElementContainer& Container = DrawElements.Get<(uint8)EElementType::ET_PostProcessPass>();
+		STAT(ElementStat_Other += Container.Elements.Num());
+		for (const FSlateDrawElement& DrawElement : Container.Elements)
 		{
-			SCOPED_NAMED_EVENT_TEXT("Slate::AddNonMappedElement", FColor::Magenta);
-			HandleUnMapped(Container.Elements, ViewportSize);
-		}
-			break;
-		default:
-			checkf(0, TEXT("Invalid element type"));
-			break;
+			AddPostProcessPass(DrawElement, ViewportSize);
 		}
 	}
+
+	{
+		SCOPED_NAMED_EVENT_TEXT("Slate::AddNonMappedElement", FColor::Magenta);
+		const FSlateDrawElementContainer& Container = DrawElements.Get<(uint8)EElementType::ET_NonMapped>();
+		HandleUnMapped(Container.Elements, ViewportSize);
+	}
 }
 
 void FSlateElementBatcher::AddCachedElements(FSlateCachedElementData& CachedElementData, FVector2f ViewportSize)
@@ -610,12 +594,13 @@ void FSlateElementBatcher::AddCachedElements(FSlateCachedElementData& CachedElem
 
 	for (FSlateCachedElementList* List : CachedElementData.ListsWithNewData)
 	{
-		if (List->DrawElements.Num() > 0)
+		int32 NumElements = List->NumElements();
+		if (NumElements > 0)
 		{
-			STAT(ElementStat_RecachedElements += List->DrawElements.Num());
+			STAT(ElementStat_RecachedElements += NumElements);
 
 #if SLATE_CSV_TRACKER
-			RecachedDrawElements += List->DrawElements.Num();
+			RecachedDrawElements += NumElements;
 #endif
 
 			CurrentCachedElementList = List;
@@ -645,6 +630,58 @@ void FSlateElementBatcher::AddCachedElements(FSlateCachedElementData& CachedElem
 #endif
 }
 
+template<typename ElementAdder, typename ElementBatchParamCreator, typename ElementBatchReserver>
+void FSlateElementBatcher::GenerateIndexedVertexBatches(
+	const FSlateDrawElementArray& DrawElements
+	, ElementAdder&& InElementAdder
+	, ElementBatchParamCreator&& InElementBatchParamCreator
+	, ElementBatchReserver&& InElementBatchReserver)
+{
+	for (int32 Index = 0; Index < DrawElements.Num();)
+	{
+		// Index is set to end of last processed batch at end of loop, may not be valid
+		if (!DrawElements.IsValidIndex(Index))
+		{
+			return;
+		}
+
+		// Determine batch range
+		FSlateRenderBatchParams NewBatchParams = InElementBatchParamCreator(DrawElements[Index]);
+		int32 BatchIndexEnd = Index;
+		while (DrawElements.IsValidIndex(++BatchIndexEnd))
+		{
+			const FSlateDrawElement& NextDrawElement = DrawElements[BatchIndexEnd];
+			FSlateRenderBatchParams NextBatchParams = InElementBatchParamCreator(NextDrawElement);
+
+			if (!NextBatchParams.IsBatchableWith(NewBatchParams))
+			{
+				break;
+			}
+		}
+
+		// Process valid range
+		FSlateRenderBatch& RenderBatch = CreateRenderBatch(
+			NewBatchParams.Layer
+			, NewBatchParams.ShaderParams
+			, NewBatchParams.Resource
+			, NewBatchParams.PrimitiveType
+			, NewBatchParams.ShaderType
+			, NewBatchParams.DrawEffects
+			, NewBatchParams.DrawFlags
+			, NewBatchParams.SceneIndex
+			, NewBatchParams.ClippingState);
+		
+		InElementBatchReserver(RenderBatch, Index, BatchIndexEnd);
+
+		for (int32 BatchIndex = Index; BatchIndex < BatchIndexEnd; BatchIndex++)
+		{
+			InElementAdder(RenderBatch, DrawElements[BatchIndex]);
+		}
+
+		Index = BatchIndexEnd;
+	}
+}
+
 template<ESlateVertexRounding Rounding>
 void FSlateElementBatcher::AddDebugQuadElement(const FSlateDrawElement& DrawElement)
 {
@@ -697,486 +734,530 @@ FORCEINLINE void IndexQuad(FSlateRenderBatch& RenderBatch, int32 TopLeft, int32
 	RenderBatch.AddIndex(TopLeft);
 }
 
-template<ESlateVertexRounding Rounding>
-void FSlateElementBatcher::AddBoxElement(const FSlateDrawElement& DrawElement)
+void FSlateElementBatcher::AddBoxElements(const FSlateDrawElementArray& DrawElements)
 {
-	const FSlateBoxPayload& DrawElementPayload = DrawElement.GetDataPayload<FSlateBoxPayload>();
+	auto AddBoxElementInternal = [&](FSlateRenderBatch& RenderBatch, const FSlateDrawElement& DrawElement)
+	{
+		const ESlateVertexRounding Rounding = DrawElement.IsPixelSnapped() ? ESlateVertexRounding::Enabled : ESlateVertexRounding::Disabled;
+		const FSlateBoxPayload& DrawElementPayload = DrawElement.GetDataPayload<FSlateBoxPayload>();
 
-	const FColor Tint = PackVertexColor(DrawElementPayload.GetTint());
-	const FSlateRenderTransform& ElementRenderTransform = DrawElement.GetRenderTransform();
-	const FSlateRenderTransform RenderTransform = DrawElement.GetRenderTransform();// GetBoxRenderTransform(DrawElement);
-	const FVector2f LocalSize = DrawElement.GetLocalSize();
+		const FColor Tint = PackVertexColor(DrawElementPayload.GetTint());
+		const FSlateRenderTransform& ElementRenderTransform = DrawElement.GetRenderTransform();
+		const FSlateRenderTransform RenderTransform = DrawElement.GetRenderTransform();// GetBoxRenderTransform(DrawElement);
+		const FVector2f LocalSize = DrawElement.GetLocalSize();
 
-	const ESlateDrawEffect InDrawEffects = DrawElement.GetDrawEffects();
-	const int32 Layer = DrawElement.GetLayer();
+		const ESlateDrawEffect InDrawEffects = DrawElement.GetDrawEffects();
+		const int32 Layer = DrawElement.GetLayer();
 
-	const float DrawScale = DrawElement.GetScale();
+		const float DrawScale = DrawElement.GetScale();
 
-	// Do pixel snapping
-	FVector2f TopLeft(0.f, 0.f);
-	FVector2f BotRight(LocalSize);
+		// Do pixel snapping
+		FVector2f TopLeft(0.f, 0.f);
+		FVector2f BotRight(LocalSize);
 
-	uint32 TextureWidth = 1;
-	uint32 TextureHeight = 1;
+		uint32 TextureWidth = 1;
+		uint32 TextureHeight = 1;
 
-	// Get the default start and end UV.  If the texture is atlased this value will be a subset of this
-	FVector2f StartUV = FVector2f(0.0f, 0.0f);
-	FVector2f EndUV = FVector2f(1.0f, 1.0f);
-	FVector2f SizeUV;
+		// Get the default start and end UV.  If the texture is atlased this value will be a subset of this
+		FVector2f StartUV = FVector2f(0.0f, 0.0f);
+		FVector2f EndUV = FVector2f(1.0f, 1.0f);
+		FVector2f SizeUV;
 
-	FVector2f HalfTexel;
+		FVector2f HalfTexel;
 
-	const FSlateShaderResourceProxy* ResourceProxy = DrawElementPayload.GetResourceProxy();
-	FSlateShaderResource* Resource = nullptr;
-	if( ResourceProxy )
-	{
-		// The actual texture for rendering.  If the texture is atlased this is the atlas
-		Resource = ResourceProxy->Resource;
-		// The width and height of the texture (non-atlased size)
-		TextureWidth = ResourceProxy->ActualSize.X != 0 ? ResourceProxy->ActualSize.X : 1;
-		TextureHeight = ResourceProxy->ActualSize.Y != 0 ? ResourceProxy->ActualSize.Y : 1;
+		const FSlateShaderResourceProxy* ResourceProxy = DrawElementPayload.GetResourceProxy();
+		FSlateShaderResource* Resource = nullptr;
+		if (ResourceProxy)
+		{
+			// The actual texture for rendering.  If the texture is atlased this is the atlas
+			Resource = ResourceProxy->Resource;
+			// The width and height of the texture (non-atlased size)
+			TextureWidth = ResourceProxy->ActualSize.X != 0 ? ResourceProxy->ActualSize.X : 1;
+			TextureHeight = ResourceProxy->ActualSize.Y != 0 ? ResourceProxy->ActualSize.Y : 1;
 
-		// Texel offset
-		HalfTexel = FVector2f( PixelCenterOffset/TextureWidth, PixelCenterOffset/TextureHeight );
+			// Texel offset
+			HalfTexel = FVector2f(PixelCenterOffset / TextureWidth, PixelCenterOffset / TextureHeight);
 
-		const FBox2f& BrushUV = DrawElementPayload.GetBrushUVRegion();
-		//In case brush has valid UV region - use it instead of proxy UV
-		if (BrushUV.bIsValid)
-		{
-			SizeUV = BrushUV.GetSize();
-			StartUV = BrushUV.Min + HalfTexel;
-			EndUV = StartUV + SizeUV;
+			const FBox2f& BrushUV = DrawElementPayload.GetBrushUVRegion();
+			//In case brush has valid UV region - use it instead of proxy UV
+			if (BrushUV.bIsValid)
+			{
+				SizeUV = BrushUV.GetSize();
+				StartUV = BrushUV.Min + HalfTexel;
+				EndUV = StartUV + SizeUV;
+			}
+			else
+			{
+				SizeUV = FVector2f(ResourceProxy->SizeUV);
+				StartUV = FVector2f(ResourceProxy->StartUV) + HalfTexel;
+				EndUV = StartUV + FVector2f(ResourceProxy->SizeUV);
+			}
 		}
 		else
 		{
-			SizeUV = FVector2f(ResourceProxy->SizeUV);
-			StartUV = FVector2f(ResourceProxy->StartUV) + HalfTexel;
-			EndUV = StartUV + FVector2f(ResourceProxy->SizeUV);
+			// no texture
+			SizeUV = FVector2f(1.0f, 1.0f);
+			HalfTexel = FVector2f(PixelCenterOffset, PixelCenterOffset);
 		}
-	}
-	else
-	{
-		// no texture
-		SizeUV = FVector2f(1.0f,1.0f);
-		HalfTexel = FVector2f( PixelCenterOffset, PixelCenterOffset );
-	}
 
 
-	const ESlateBrushTileType::Type TilingRule = DrawElementPayload.GetBrushTiling();
-	const bool bTileHorizontal = (TilingRule == ESlateBrushTileType::Both || TilingRule == ESlateBrushTileType::Horizontal);
-	const bool bTileVertical = (TilingRule == ESlateBrushTileType::Both || TilingRule == ESlateBrushTileType::Vertical);
+		const ESlateBrushTileType::Type TilingRule = DrawElementPayload.GetBrushTiling();
+		const bool bTileHorizontal = (TilingRule == ESlateBrushTileType::Both || TilingRule == ESlateBrushTileType::Horizontal);
+		const bool bTileVertical = (TilingRule == ESlateBrushTileType::Both || TilingRule == ESlateBrushTileType::Vertical);
 
-	const ESlateBrushMirrorType::Type MirroringRule = DrawElementPayload.GetBrushMirroring();
-	const bool bMirrorHorizontal = (MirroringRule == ESlateBrushMirrorType::Both || MirroringRule == ESlateBrushMirrorType::Horizontal);
-	const bool bMirrorVertical = (MirroringRule == ESlateBrushMirrorType::Both || MirroringRule == ESlateBrushMirrorType::Vertical);
+		const ESlateBrushMirrorType::Type MirroringRule = DrawElementPayload.GetBrushMirroring();
+		const bool bMirrorHorizontal = (MirroringRule == ESlateBrushMirrorType::Both || MirroringRule == ESlateBrushMirrorType::Horizontal);
+		const bool bMirrorVertical = (MirroringRule == ESlateBrushMirrorType::Both || MirroringRule == ESlateBrushMirrorType::Vertical);
 
-	// Pass the tiling information as a flag so we can pick the correct texture addressing mode
-	ESlateBatchDrawFlag DrawFlags = DrawElement.GetBatchFlags();
-	DrawFlags |= ( ( bTileHorizontal ? ESlateBatchDrawFlag::TileU : ESlateBatchDrawFlag::None ) | ( bTileVertical ? ESlateBatchDrawFlag::TileV : ESlateBatchDrawFlag::None ) );
+		// Pass the tiling information as a flag so we can pick the correct texture addressing mode
+		ESlateBatchDrawFlag DrawFlags = DrawElement.GetBatchFlags();
+		DrawFlags |= ((bTileHorizontal ? ESlateBatchDrawFlag::TileU : ESlateBatchDrawFlag::None) | (bTileVertical ? ESlateBatchDrawFlag::TileV : ESlateBatchDrawFlag::None));
 
-	// Add Shader Parameters for extra RoundedBox parameters
-	ESlateShader ShaderType = ESlateShader::Default;
-	FShaderParams ShaderParams;
-	FColor SecondaryColor;
-	if (DrawElement.GetElementType() == EElementType::ET_RoundedBox)
-	{
-	
-		ShaderType = ESlateShader::RoundedBox;
-		const FSlateRoundedBoxPayload& RoundedPayload = DrawElement.GetDataPayload<FSlateRoundedBoxPayload>();
-
-		ShaderParams.PixelParams = FVector4f(0, RoundedPayload.GetOutlineWeight(), LocalSize.X, LocalSize.Y);//RadiusWeight;
-		ShaderParams.PixelParams2 = RoundedPayload.GetRadius();
-
-		SecondaryColor = PackVertexColor(RoundedPayload.OutlineColor);
-	}
-
-	FSlateRenderBatch& RenderBatch = CreateRenderBatch(Layer, ShaderParams, Resource, ESlateDrawPrimitive::TriangleList, ShaderType, InDrawEffects, DrawFlags, DrawElement);
-	 
-	float HorizontalTiling = bTileHorizontal ? LocalSize.X / TextureWidth : 1.0f;
-	float VerticalTiling = bTileVertical ? LocalSize.Y / TextureHeight : 1.0f;
+		// Add Shader Parameters for extra RoundedBox parameters
+		ESlateShader ShaderType = ESlateShader::Default;
+		FShaderParams ShaderParams;
+		FColor SecondaryColor;
+		if (DrawElement.GetElementType() == EElementType::ET_RoundedBox)
+		{
 
-	const FVector2f Tiling( HorizontalTiling, VerticalTiling );
+			ShaderType = ESlateShader::RoundedBox;
+			const FSlateRoundedBoxPayload& RoundedPayload = DrawElement.GetDataPayload<FSlateRoundedBoxPayload>();
 
-	// The start index of these vertices in the index buffer
-	const uint32 IndexStart = 0;// BatchVertices.Num();
+			ShaderParams.PixelParams = FVector4f(0, RoundedPayload.GetOutlineWeight(), LocalSize.X, LocalSize.Y);//RadiusWeight;
+			ShaderParams.PixelParams2 = RoundedPayload.GetRadius();
 
+			SecondaryColor = PackVertexColor(RoundedPayload.OutlineColor);
+		}
 
-	const FMargin& Margin = DrawElementPayload.GetBrushMargin();
+		float HorizontalTiling = bTileHorizontal ? LocalSize.X / TextureWidth : 1.0f;
+		float VerticalTiling = bTileVertical ? LocalSize.Y / TextureHeight : 1.0f;
 
-	const FVector2f TopRight = FVector2f(BotRight.X, TopLeft.Y);
-	const FVector2f BotLeft = FVector2f(TopLeft.X, BotRight.Y);
+		const FVector2f Tiling(HorizontalTiling, VerticalTiling);
 
-	const FColor FeatherColor(0, 0, 0, 0);
+		// The start index of these vertices in the index buffer
+		const uint32 IndexStart = RenderBatch.GetNumVertices();
 
-	if (DrawElementPayload.GetBrushDrawType() != ESlateBrushDrawType::Image &&
-		( Margin.Left != 0.0f || Margin.Top != 0.0f || Margin.Right != 0.0f || Margin.Bottom != 0.0f ) )
-	{
-		// Create 9 quads for the box element based on the following diagram
-		//     ___LeftMargin    ___RightMargin
-		//    /                /
-		//  +--+-------------+--+
-		//  |  |c1           |c2| ___TopMargin
-		//  +--o-------------o--+
-		//  |  |             |  |
-		//  |  |c3           |c4|
-		//  +--o-------------o--+
-		//  |  |             |  | ___BottomMargin
-		//  +--+-------------+--+
+		const FMargin& Margin = DrawElementPayload.GetBrushMargin();
+
+		const FVector2f TopRight = FVector2f(BotRight.X, TopLeft.Y);
+		const FVector2f BotLeft = FVector2f(TopLeft.X, BotRight.Y);
+
+		const FColor FeatherColor(0, 0, 0, 0);
+
+		if (DrawElementPayload.GetBrushDrawType() != ESlateBrushDrawType::Image &&
+			(Margin.Left != 0.0f || Margin.Top != 0.0f || Margin.Right != 0.0f || Margin.Bottom != 0.0f))
+		{
+			// Create 9 quads for the box element based on the following diagram
+			//     ___LeftMargin    ___RightMargin
+			//    /                /
+			//  +--+-------------+--+
+			//  |  |c1           |c2| ___TopMargin
+			//  +--o-------------o--+
+			//  |  |             |  |
+			//  |  |c3           |c4|
+			//  +--o-------------o--+
+			//  |  |             |  | ___BottomMargin
+			//  +--+-------------+--+
+
+
+			// Determine the texture coordinates for each quad
+			// These are not scaled.
+			float LeftMarginU = (Margin.Left > 0.0f)
+				? StartUV.X + Margin.Left * SizeUV.X + HalfTexel.X
+				: StartUV.X;
+			float TopMarginV = (Margin.Top > 0.0f)
+				? StartUV.Y + Margin.Top * SizeUV.Y + HalfTexel.Y
+				: StartUV.Y;
+			float RightMarginU = (Margin.Right > 0.0f)
+				? EndUV.X - Margin.Right * SizeUV.X + HalfTexel.X
+				: EndUV.X;
+			float BottomMarginV = (Margin.Bottom > 0.0f)
+				? EndUV.Y - Margin.Bottom * SizeUV.Y + HalfTexel.Y
+				: EndUV.Y;
+
+			if (bMirrorHorizontal || bMirrorVertical)
+			{
+				const FVector2f UVMin = StartUV;
+				const FVector2f UVMax = EndUV;
 
+				if (bMirrorHorizontal)
+				{
+					StartUV.X = UVMax.X - (StartUV.X - UVMin.X);
+					EndUV.X = UVMax.X - (EndUV.X - UVMin.X);
+					LeftMarginU = UVMax.X - (LeftMarginU - UVMin.X);
+					RightMarginU = UVMax.X - (RightMarginU - UVMin.X);
+				}
+				if (bMirrorVertical)
+				{
+					StartUV.Y = UVMax.Y - (StartUV.Y - UVMin.Y);
+					EndUV.Y = UVMax.Y - (EndUV.Y - UVMin.Y);
+					TopMarginV = UVMax.Y - (TopMarginV - UVMin.Y);
+					BottomMarginV = UVMax.Y - (BottomMarginV - UVMin.Y);
+				}
+			}
 
-		// Determine the texture coordinates for each quad
-		// These are not scaled.
-		float LeftMarginU = (Margin.Left > 0.0f)
-			? StartUV.X + Margin.Left * SizeUV.X + HalfTexel.X
-			: StartUV.X;
-		float TopMarginV = (Margin.Top > 0.0f)
-			? StartUV.Y + Margin.Top * SizeUV.Y + HalfTexel.Y
-			: StartUV.Y;
-		float RightMarginU = (Margin.Right > 0.0f)
-			? EndUV.X - Margin.Right * SizeUV.X + HalfTexel.X
-			: EndUV.X;
-		float BottomMarginV = (Margin.Bottom > 0.0f)
-			? EndUV.Y - Margin.Bottom * SizeUV.Y + HalfTexel.Y
-			: EndUV.Y;
+			// Determine the margins for each quad
 
-		if( bMirrorHorizontal || bMirrorVertical )
-		{
-			const FVector2f UVMin = StartUV;
-			const FVector2f UVMax = EndUV;
+			float LeftMarginX = TextureWidth * Margin.Left;
+			float TopMarginY = TextureHeight * Margin.Top;
+			float RightMarginX = LocalSize.X - TextureWidth * Margin.Right;
+			float BottomMarginY = LocalSize.Y - TextureHeight * Margin.Bottom;
 
-			if( bMirrorHorizontal )
+			// If the margins are overlapping the margins are too big or the button is too small
+			// so clamp margins to half of the box size
+			if (RightMarginX < LeftMarginX)
 			{
-				StartUV.X = UVMax.X - ( StartUV.X - UVMin.X );
-				EndUV.X = UVMax.X - ( EndUV.X - UVMin.X );
-				LeftMarginU = UVMax.X - ( LeftMarginU - UVMin.X );
-				RightMarginU = UVMax.X - ( RightMarginU - UVMin.X );
+				LeftMarginX = LocalSize.X / 2;
+				RightMarginX = LeftMarginX;
 			}
-			if( bMirrorVertical )
+
+			if (BottomMarginY < TopMarginY)
 			{
-				StartUV.Y = UVMax.Y - ( StartUV.Y - UVMin.Y );
-				EndUV.Y = UVMax.Y - ( EndUV.Y - UVMin.Y );
-				TopMarginV = UVMax.Y - ( TopMarginV - UVMin.Y );
-				BottomMarginV = UVMax.Y - ( BottomMarginV - UVMin.Y );
+				TopMarginY = LocalSize.Y / 2;
+				BottomMarginY = TopMarginY;
 			}
-		}
 
-		// Determine the margins for each quad
-
-		float LeftMarginX = TextureWidth * Margin.Left;
-		float TopMarginY = TextureHeight * Margin.Top;
-		float RightMarginX = LocalSize.X - TextureWidth * Margin.Right;
-		float BottomMarginY = LocalSize.Y - TextureHeight * Margin.Bottom;
-
-		// If the margins are overlapping the margins are too big or the button is too small
-		// so clamp margins to half of the box size
-		if( RightMarginX < LeftMarginX )
-		{
-			LeftMarginX = LocalSize.X / 2;
-			RightMarginX = LeftMarginX;
-		}
-
-		if( BottomMarginY < TopMarginY )
-		{
-			TopMarginY = LocalSize.Y / 2;
-			BottomMarginY = TopMarginY;
-		}
-
-		FVector2f Position = TopLeft;
-		FVector2f EndPos = BotRight;
-
-		RenderBatch.AddVertex( FSlateVertex::Make<Rounding>( RenderTransform, FVector2f( Position.X, Position.Y ),		LocalSize, DrawScale, FVector4f(FVector2f( StartUV ),						Tiling),	Tint, SecondaryColor) ); //0
-		RenderBatch.AddVertex( FSlateVertex::Make<Rounding>( RenderTransform, FVector2f( Position.X, TopMarginY ),		LocalSize, DrawScale, FVector4f(FVector2f( StartUV.X, TopMarginV ),			Tiling),	Tint, SecondaryColor) ); //1
-		RenderBatch.AddVertex( FSlateVertex::Make<Rounding>( RenderTransform, FVector2f( LeftMarginX, Position.Y ),		LocalSize, DrawScale, FVector4f(FVector2f( LeftMarginU, StartUV.Y ),		Tiling),	Tint, SecondaryColor) ); //2
-		RenderBatch.AddVertex( FSlateVertex::Make<Rounding>( RenderTransform, FVector2f( LeftMarginX, TopMarginY ),		LocalSize, DrawScale, FVector4f(FVector2f( LeftMarginU, TopMarginV ),		Tiling),	Tint, SecondaryColor) ); //3
-		RenderBatch.AddVertex( FSlateVertex::Make<Rounding>( RenderTransform, FVector2f( RightMarginX, Position.Y ),	LocalSize, DrawScale, FVector4f(FVector2f( RightMarginU, StartUV.Y ),		Tiling),	Tint, SecondaryColor) ); //4
-		RenderBatch.AddVertex( FSlateVertex::Make<Rounding>( RenderTransform, FVector2f( RightMarginX, TopMarginY ),	LocalSize, DrawScale, FVector4f(FVector2f( RightMarginU,TopMarginV),		Tiling),	Tint, SecondaryColor) ); //5
-		RenderBatch.AddVertex( FSlateVertex::Make<Rounding>( RenderTransform, FVector2f( EndPos.X, Position.Y ),		LocalSize, DrawScale, FVector4f(FVector2f( EndUV.X, StartUV.Y ),			Tiling),	Tint, SecondaryColor) ); //6
-		RenderBatch.AddVertex( FSlateVertex::Make<Rounding>( RenderTransform, FVector2f( EndPos.X, TopMarginY ),		LocalSize, DrawScale, FVector4f(FVector2f( EndUV.X, TopMarginV),			Tiling),	Tint, SecondaryColor) ); //7
-
-		RenderBatch.AddVertex( FSlateVertex::Make<Rounding>( RenderTransform, FVector2f( Position.X, BottomMarginY ),	LocalSize, DrawScale, FVector4f(FVector2f( StartUV.X, BottomMarginV ),		Tiling),	Tint, SecondaryColor) ); //8
-		RenderBatch.AddVertex( FSlateVertex::Make<Rounding>( RenderTransform, FVector2f( LeftMarginX, BottomMarginY ),	LocalSize, DrawScale, FVector4f(FVector2f( LeftMarginU, BottomMarginV ),	Tiling),	Tint, SecondaryColor) ); //9
-		RenderBatch.AddVertex( FSlateVertex::Make<Rounding>( RenderTransform, FVector2f( RightMarginX, BottomMarginY ), LocalSize, DrawScale, FVector4f(FVector2f( RightMarginU, BottomMarginV ),	Tiling),	Tint, SecondaryColor) ); //10
-		RenderBatch.AddVertex( FSlateVertex::Make<Rounding>( RenderTransform, FVector2f( EndPos.X, BottomMarginY ),		LocalSize, DrawScale, FVector4f(FVector2f( EndUV.X, BottomMarginV ),		Tiling),	Tint, SecondaryColor) ); //11
-		RenderBatch.AddVertex( FSlateVertex::Make<Rounding>( RenderTransform, FVector2f( Position.X, EndPos.Y ),		LocalSize, DrawScale, FVector4f(FVector2f( StartUV.X, EndUV.Y ),			Tiling),	Tint, SecondaryColor) ); //12
-		RenderBatch.AddVertex( FSlateVertex::Make<Rounding>( RenderTransform, FVector2f( LeftMarginX, EndPos.Y ),		LocalSize, DrawScale, FVector4f(FVector2f( LeftMarginU, EndUV.Y ),			Tiling),	Tint, SecondaryColor) ); //13
-		RenderBatch.AddVertex( FSlateVertex::Make<Rounding>( RenderTransform, FVector2f( RightMarginX, EndPos.Y ),		LocalSize, DrawScale, FVector4f(FVector2f( RightMarginU, EndUV.Y ),			Tiling),	Tint, SecondaryColor) ); //14
-		RenderBatch.AddVertex( FSlateVertex::Make<Rounding>( RenderTransform, FVector2f( EndPos.X, EndPos.Y ),			LocalSize, DrawScale, FVector4f(FVector2f( EndUV ),							Tiling),	Tint, SecondaryColor) ); //15
-
-		static FSlateIndexArray DefaultPaddedBoxIndiciesArray = []()
-		{
-			const int32 IndexStart = 0;
-			FSlateIndexArray DefaultPaddedBoxIndicies = {};
-			DefaultPaddedBoxIndicies.Reserve(6 * 3 * 3);
-
-			// Top
-			DefaultPaddedBoxIndicies.Add(IndexStart + 0);
-			DefaultPaddedBoxIndicies.Add(IndexStart + 1);
-			DefaultPaddedBoxIndicies.Add(IndexStart + 2);
-			DefaultPaddedBoxIndicies.Add(IndexStart + 2);
-			DefaultPaddedBoxIndicies.Add(IndexStart + 1);
-			DefaultPaddedBoxIndicies.Add(IndexStart + 3);
-
-			DefaultPaddedBoxIndicies.Add(IndexStart + 2);
-			DefaultPaddedBoxIndicies.Add(IndexStart + 3);
-			DefaultPaddedBoxIndicies.Add(IndexStart + 4);
-			DefaultPaddedBoxIndicies.Add(IndexStart + 4);
-			DefaultPaddedBoxIndicies.Add(IndexStart + 3);
-			DefaultPaddedBoxIndicies.Add(IndexStart + 5);
-
-			DefaultPaddedBoxIndicies.Add(IndexStart + 4);
-			DefaultPaddedBoxIndicies.Add(IndexStart + 5);
-			DefaultPaddedBoxIndicies.Add(IndexStart + 6);
-			DefaultPaddedBoxIndicies.Add(IndexStart + 6);
-			DefaultPaddedBoxIndicies.Add(IndexStart + 5);
-			DefaultPaddedBoxIndicies.Add(IndexStart + 7);
+			FVector2f Position = TopLeft;
+			FVector2f EndPos = BotRight;
+
+			RenderBatch.EmplaceVertex(FSlateVertex::Make(RenderTransform, FVector2f(Position.X, Position.Y), LocalSize, DrawScale, FVector4f(FVector2f(StartUV), Tiling), Tint, SecondaryColor, Rounding)); //0
+			RenderBatch.EmplaceVertex(FSlateVertex::Make(RenderTransform, FVector2f(Position.X, TopMarginY), LocalSize, DrawScale, FVector4f(FVector2f(StartUV.X, TopMarginV), Tiling), Tint, SecondaryColor, Rounding)); //1
+			RenderBatch.EmplaceVertex(FSlateVertex::Make(RenderTransform, FVector2f(LeftMarginX, Position.Y), LocalSize, DrawScale, FVector4f(FVector2f(LeftMarginU, StartUV.Y), Tiling), Tint, SecondaryColor, Rounding)); //2
+			RenderBatch.EmplaceVertex(FSlateVertex::Make(RenderTransform, FVector2f(LeftMarginX, TopMarginY), LocalSize, DrawScale, FVector4f(FVector2f(LeftMarginU, TopMarginV), Tiling), Tint, SecondaryColor, Rounding)); //3
+			RenderBatch.EmplaceVertex(FSlateVertex::Make(RenderTransform, FVector2f(RightMarginX, Position.Y), LocalSize, DrawScale, FVector4f(FVector2f(RightMarginU, StartUV.Y), Tiling), Tint, SecondaryColor, Rounding)); //4
+			RenderBatch.EmplaceVertex(FSlateVertex::Make(RenderTransform, FVector2f(RightMarginX, TopMarginY), LocalSize, DrawScale, FVector4f(FVector2f(RightMarginU, TopMarginV), Tiling), Tint, SecondaryColor, Rounding)); //5
+			RenderBatch.EmplaceVertex(FSlateVertex::Make(RenderTransform, FVector2f(EndPos.X, Position.Y), LocalSize, DrawScale, FVector4f(FVector2f(EndUV.X, StartUV.Y), Tiling), Tint, SecondaryColor, Rounding)); //6
+			RenderBatch.EmplaceVertex(FSlateVertex::Make(RenderTransform, FVector2f(EndPos.X, TopMarginY), LocalSize, DrawScale, FVector4f(FVector2f(EndUV.X, TopMarginV), Tiling), Tint, SecondaryColor, Rounding)); //7
+
+			RenderBatch.EmplaceVertex(FSlateVertex::Make(RenderTransform, FVector2f(Position.X, BottomMarginY), LocalSize, DrawScale, FVector4f(FVector2f(StartUV.X, BottomMarginV), Tiling), Tint, SecondaryColor, Rounding)); //8
+			RenderBatch.EmplaceVertex(FSlateVertex::Make(RenderTransform, FVector2f(LeftMarginX, BottomMarginY), LocalSize, DrawScale, FVector4f(FVector2f(LeftMarginU, BottomMarginV), Tiling), Tint, SecondaryColor, Rounding)); //9
+			RenderBatch.EmplaceVertex(FSlateVertex::Make(RenderTransform, FVector2f(RightMarginX, BottomMarginY), LocalSize, DrawScale, FVector4f(FVector2f(RightMarginU, BottomMarginV), Tiling), Tint, SecondaryColor, Rounding)); //10
+			RenderBatch.EmplaceVertex(FSlateVertex::Make(RenderTransform, FVector2f(EndPos.X, BottomMarginY), LocalSize, DrawScale, FVector4f(FVector2f(EndUV.X, BottomMarginV), Tiling), Tint, SecondaryColor, Rounding)); //11
+			RenderBatch.EmplaceVertex(FSlateVertex::Make(RenderTransform, FVector2f(Position.X, EndPos.Y), LocalSize, DrawScale, FVector4f(FVector2f(StartUV.X, EndUV.Y), Tiling), Tint, SecondaryColor, Rounding)); //12
+			RenderBatch.EmplaceVertex(FSlateVertex::Make(RenderTransform, FVector2f(LeftMarginX, EndPos.Y), LocalSize, DrawScale, FVector4f(FVector2f(LeftMarginU, EndUV.Y), Tiling), Tint, SecondaryColor, Rounding)); //13
+			RenderBatch.EmplaceVertex(FSlateVertex::Make(RenderTransform, FVector2f(RightMarginX, EndPos.Y), LocalSize, DrawScale, FVector4f(FVector2f(RightMarginU, EndUV.Y), Tiling), Tint, SecondaryColor, Rounding)); //14
+			RenderBatch.EmplaceVertex(FSlateVertex::Make(RenderTransform, FVector2f(EndPos.X, EndPos.Y), LocalSize, DrawScale, FVector4f(FVector2f(EndUV), Tiling), Tint, SecondaryColor, Rounding)); //15
+
+			RenderBatch.EmplaceIndex(IndexStart + 0);
+			RenderBatch.EmplaceIndex(IndexStart + 1);
+			RenderBatch.EmplaceIndex(IndexStart + 2);
+			RenderBatch.EmplaceIndex(IndexStart + 2);
+			RenderBatch.EmplaceIndex(IndexStart + 1);
+			RenderBatch.EmplaceIndex(IndexStart + 3);
+
+			RenderBatch.EmplaceIndex(IndexStart + 2);
+			RenderBatch.EmplaceIndex(IndexStart + 3);
+			RenderBatch.EmplaceIndex(IndexStart + 4);
+			RenderBatch.EmplaceIndex(IndexStart + 4);
+			RenderBatch.EmplaceIndex(IndexStart + 3);
+			RenderBatch.EmplaceIndex(IndexStart + 5);
+
+			RenderBatch.EmplaceIndex(IndexStart + 4);
+			RenderBatch.EmplaceIndex(IndexStart + 5);
+			RenderBatch.EmplaceIndex(IndexStart + 6);
+			RenderBatch.EmplaceIndex(IndexStart + 6);
+			RenderBatch.EmplaceIndex(IndexStart + 5);
+			RenderBatch.EmplaceIndex(IndexStart + 7);
 
 			// Middle
-			DefaultPaddedBoxIndicies.Add(IndexStart + 1);
-			DefaultPaddedBoxIndicies.Add(IndexStart + 8);
-			DefaultPaddedBoxIndicies.Add(IndexStart + 3);
-			DefaultPaddedBoxIndicies.Add(IndexStart + 3);
-			DefaultPaddedBoxIndicies.Add(IndexStart + 8);
-			DefaultPaddedBoxIndicies.Add(IndexStart + 9);
-
-			DefaultPaddedBoxIndicies.Add(IndexStart + 3);
-			DefaultPaddedBoxIndicies.Add(IndexStart + 9);
-			DefaultPaddedBoxIndicies.Add(IndexStart + 5);
-			DefaultPaddedBoxIndicies.Add(IndexStart + 5);
-			DefaultPaddedBoxIndicies.Add(IndexStart + 9);
-			DefaultPaddedBoxIndicies.Add(IndexStart + 10);
-
-			DefaultPaddedBoxIndicies.Add(IndexStart + 5);
-			DefaultPaddedBoxIndicies.Add(IndexStart + 10);
-			DefaultPaddedBoxIndicies.Add(IndexStart + 7);
-			DefaultPaddedBoxIndicies.Add(IndexStart + 7);
-			DefaultPaddedBoxIndicies.Add(IndexStart + 10);
-			DefaultPaddedBoxIndicies.Add(IndexStart + 11);
+			RenderBatch.EmplaceIndex(IndexStart + 1);
+			RenderBatch.EmplaceIndex(IndexStart + 8);
+			RenderBatch.EmplaceIndex(IndexStart + 3);
+			RenderBatch.EmplaceIndex(IndexStart + 3);
+			RenderBatch.EmplaceIndex(IndexStart + 8);
+			RenderBatch.EmplaceIndex(IndexStart + 9);
+
+			RenderBatch.EmplaceIndex(IndexStart + 3);
+			RenderBatch.EmplaceIndex(IndexStart + 9);
+			RenderBatch.EmplaceIndex(IndexStart + 5);
+			RenderBatch.EmplaceIndex(IndexStart + 5);
+			RenderBatch.EmplaceIndex(IndexStart + 9);
+			RenderBatch.EmplaceIndex(IndexStart + 10);
+
+			RenderBatch.EmplaceIndex(IndexStart + 5);
+			RenderBatch.EmplaceIndex(IndexStart + 10);
+			RenderBatch.EmplaceIndex(IndexStart + 7);
+			RenderBatch.EmplaceIndex(IndexStart + 7);
+			RenderBatch.EmplaceIndex(IndexStart + 10);
+			RenderBatch.EmplaceIndex(IndexStart + 11);
 
 			// Bottom
-			DefaultPaddedBoxIndicies.Add(IndexStart + 8);
-			DefaultPaddedBoxIndicies.Add(IndexStart + 12);
-			DefaultPaddedBoxIndicies.Add(IndexStart + 9);
-			DefaultPaddedBoxIndicies.Add(IndexStart + 9);
-			DefaultPaddedBoxIndicies.Add(IndexStart + 12);
-			DefaultPaddedBoxIndicies.Add(IndexStart + 13);
-
-			DefaultPaddedBoxIndicies.Add(IndexStart + 9);
-			DefaultPaddedBoxIndicies.Add(IndexStart + 13);
-			DefaultPaddedBoxIndicies.Add(IndexStart + 10);
-			DefaultPaddedBoxIndicies.Add(IndexStart + 10);
-			DefaultPaddedBoxIndicies.Add(IndexStart + 13);
-			DefaultPaddedBoxIndicies.Add(IndexStart + 14);
-
-			DefaultPaddedBoxIndicies.Add(IndexStart + 10);
-			DefaultPaddedBoxIndicies.Add(IndexStart + 14);
-			DefaultPaddedBoxIndicies.Add(IndexStart + 11);
-			DefaultPaddedBoxIndicies.Add(IndexStart + 11);
-			DefaultPaddedBoxIndicies.Add(IndexStart + 14);
-			DefaultPaddedBoxIndicies.Add(IndexStart + 15);
-
-			return DefaultPaddedBoxIndicies;
-		}();
-		RenderBatch.SourceIndices = &DefaultPaddedBoxIndiciesArray;
-		RenderBatch.NumIndices = DefaultPaddedBoxIndiciesArray.Num();
-		RenderBatch.IndexOffset = 0;
-
-		if ( GSlateFeathering && Rounding == ESlateVertexRounding::Disabled )
-		{
-			const int32 FeatherStart = RenderBatch.GetNumVertices();
-
-			// Top
-			RenderBatch.AddVertex(FSlateVertex::Make<Rounding>(RenderTransform, FVector2f(Position.X, Position.Y) + FVector2f(-1.f, -1.f) / DrawScale,		LocalSize, DrawScale, FVector4f(FVector2f(StartUV), Tiling), FeatherColor, SecondaryColor)); //0
-			RenderBatch.AddVertex(FSlateVertex::Make<Rounding>(RenderTransform, FVector2f(LeftMarginX, Position.Y) + FVector2f(0.f, -1.f) / DrawScale,		LocalSize, DrawScale, FVector4f(FVector2f(LeftMarginU, StartUV.Y), Tiling), FeatherColor, SecondaryColor)); //1
-			RenderBatch.AddVertex(FSlateVertex::Make<Rounding>(RenderTransform, FVector2f(RightMarginX, Position.Y) + FVector2f(0.f, -1.f) / DrawScale,		LocalSize, DrawScale, FVector4f(FVector2f(RightMarginU, StartUV.Y), Tiling), FeatherColor, SecondaryColor)); //2
-			RenderBatch.AddVertex(FSlateVertex::Make<Rounding>(RenderTransform, FVector2f(EndPos.X, Position.Y) + FVector2f(1.f, -1.f) / DrawScale,			LocalSize, DrawScale, FVector4f(FVector2f(EndUV.X, StartUV.Y), Tiling), FeatherColor, SecondaryColor)); //3
-
-			// Left
-			RenderBatch.AddVertex(FSlateVertex::Make<Rounding>(RenderTransform, FVector2f(Position.X, TopMarginY) + FVector2f(-1.f, 0.f) / DrawScale,		LocalSize, DrawScale, FVector4f(FVector2f(StartUV.X, TopMarginV), Tiling), FeatherColor, SecondaryColor)); //4
-			RenderBatch.AddVertex(FSlateVertex::Make<Rounding>(RenderTransform, FVector2f(Position.X, BottomMarginY) + FVector2f(-1.f, 0.f) / DrawScale,	LocalSize, DrawScale, FVector4f(FVector2f(StartUV.X, BottomMarginV), Tiling), FeatherColor, SecondaryColor)); //5
-
-			// Right
-			RenderBatch.AddVertex(FSlateVertex::Make<Rounding>(RenderTransform, FVector2f(EndPos.X, TopMarginY) + FVector2f(1.f, 0.f) / DrawScale,			LocalSize, DrawScale, FVector4f(FVector2f(EndUV.X, TopMarginV), Tiling), FeatherColor, SecondaryColor)); //6
-			RenderBatch.AddVertex(FSlateVertex::Make<Rounding>(RenderTransform, FVector2f(EndPos.X, BottomMarginY) + FVector2f(1.f, 0.f) / DrawScale,		LocalSize, DrawScale, FVector4f(FVector2f(EndUV.X, BottomMarginV), Tiling), FeatherColor, SecondaryColor)); //7
+			RenderBatch.EmplaceIndex(IndexStart + 8);
+			RenderBatch.EmplaceIndex(IndexStart + 12);
+			RenderBatch.EmplaceIndex(IndexStart + 9);
+			RenderBatch.EmplaceIndex(IndexStart + 9);
+			RenderBatch.EmplaceIndex(IndexStart + 12);
+			RenderBatch.EmplaceIndex(IndexStart + 13);
+
+			RenderBatch.EmplaceIndex(IndexStart + 9);
+			RenderBatch.EmplaceIndex(IndexStart + 13);
+			RenderBatch.EmplaceIndex(IndexStart + 10);
+			RenderBatch.EmplaceIndex(IndexStart + 10);
+			RenderBatch.EmplaceIndex(IndexStart + 13);
+			RenderBatch.EmplaceIndex(IndexStart + 14);
+
+			RenderBatch.EmplaceIndex(IndexStart + 10);
+			RenderBatch.EmplaceIndex(IndexStart + 14);
+			RenderBatch.EmplaceIndex(IndexStart + 11);
+			RenderBatch.EmplaceIndex(IndexStart + 11);
+			RenderBatch.EmplaceIndex(IndexStart + 14);
+			RenderBatch.EmplaceIndex(IndexStart + 15);
+
+			if (GSlateFeathering && Rounding == ESlateVertexRounding::Disabled)
+			{
+				const int32 FeatherStart = RenderBatch.GetNumVertices();
 
-			// Bottom
-			RenderBatch.AddVertex(FSlateVertex::Make<Rounding>(RenderTransform, FVector2f(Position.X, EndPos.Y) + FVector2f(-1.f, 1.f) / DrawScale,			LocalSize, DrawScale, FVector4f(FVector2f(StartUV.X, EndUV.Y), Tiling), FeatherColor, SecondaryColor)); //8
-			RenderBatch.AddVertex(FSlateVertex::Make<Rounding>(RenderTransform, FVector2f(LeftMarginX, EndPos.Y) + FVector2f(0.f, 1.f) / DrawScale,			LocalSize, DrawScale, FVector4f(FVector2f(LeftMarginU, EndUV.Y), Tiling), FeatherColor, SecondaryColor)); //9
-			RenderBatch.AddVertex(FSlateVertex::Make<Rounding>(RenderTransform, FVector2f(RightMarginX, EndPos.Y) + FVector2f(0.f, 1.f) / DrawScale,		LocalSize, DrawScale, FVector4f(FVector2f(RightMarginU, EndUV.Y), Tiling), FeatherColor, SecondaryColor)); //10
-			RenderBatch.AddVertex(FSlateVertex::Make<Rounding>(RenderTransform, FVector2f(EndPos.X, EndPos.Y) + FVector2f(1.f, 1.f) / DrawScale,			LocalSize, DrawScale, FVector4f(FVector2f(EndUV), Tiling), FeatherColor, SecondaryColor)); //11
+				// Top
+				RenderBatch.EmplaceVertex(FSlateVertex::Make(RenderTransform, FVector2f(Position.X, Position.Y) + FVector2f(-1.f, -1.f) / DrawScale, LocalSize, DrawScale, FVector4f(FVector2f(StartUV), Tiling), FeatherColor, SecondaryColor, Rounding)); //0
+				RenderBatch.EmplaceVertex(FSlateVertex::Make(RenderTransform, FVector2f(LeftMarginX, Position.Y) + FVector2f(0.f, -1.f) / DrawScale, LocalSize, DrawScale, FVector4f(FVector2f(LeftMarginU, StartUV.Y), Tiling), FeatherColor, SecondaryColor, Rounding)); //1
+				RenderBatch.EmplaceVertex(FSlateVertex::Make(RenderTransform, FVector2f(RightMarginX, Position.Y) + FVector2f(0.f, -1.f) / DrawScale, LocalSize, DrawScale, FVector4f(FVector2f(RightMarginU, StartUV.Y), Tiling), FeatherColor, SecondaryColor, Rounding)); //2
+				RenderBatch.EmplaceVertex(FSlateVertex::Make(RenderTransform, FVector2f(EndPos.X, Position.Y) + FVector2f(1.f, -1.f) / DrawScale, LocalSize, DrawScale, FVector4f(FVector2f(EndUV.X, StartUV.Y), Tiling), FeatherColor, SecondaryColor, Rounding)); //3
 
-			static FSlateIndexArray FeatheredPaddedBoxIndiciesArray = [DefaultPaddedBoxIndiciesArray = DefaultPaddedBoxIndiciesArray]()
-			{
-				const int32 IndexStart = 0;
-				const int32 FeatherStart = 16;
-				FSlateIndexArray FeatheredPaddedBoxIndicies = DefaultPaddedBoxIndiciesArray;
-				FeatheredPaddedBoxIndicies.Reserve(FeatheredPaddedBoxIndicies.Num() + 6 * 3 * 4);
+				// Left
+				RenderBatch.EmplaceVertex(FSlateVertex::Make(RenderTransform, FVector2f(Position.X, TopMarginY) + FVector2f(-1.f, 0.f) / DrawScale, LocalSize, DrawScale, FVector4f(FVector2f(StartUV.X, TopMarginV), Tiling), FeatherColor, SecondaryColor, Rounding)); //4
+				RenderBatch.EmplaceVertex(FSlateVertex::Make(RenderTransform, FVector2f(Position.X, BottomMarginY) + FVector2f(-1.f, 0.f) / DrawScale, LocalSize, DrawScale, FVector4f(FVector2f(StartUV.X, BottomMarginV), Tiling), FeatherColor, SecondaryColor, Rounding)); //5
 
-				auto IndexQuadArray = [](FSlateIndexArray& IndexArray, int32 TopLeft, int32 TopRight, int32 BottomRight, int32 BottomLeft)
+				// Right
+				RenderBatch.EmplaceVertex(FSlateVertex::Make(RenderTransform, FVector2f(EndPos.X, TopMarginY) + FVector2f(1.f, 0.f) / DrawScale, LocalSize, DrawScale, FVector4f(FVector2f(EndUV.X, TopMarginV), Tiling), FeatherColor, SecondaryColor, Rounding)); //6
+				RenderBatch.EmplaceVertex(FSlateVertex::Make(RenderTransform, FVector2f(EndPos.X, BottomMarginY) + FVector2f(1.f, 0.f) / DrawScale, LocalSize, DrawScale, FVector4f(FVector2f(EndUV.X, BottomMarginV), Tiling), FeatherColor, SecondaryColor, Rounding)); //7
+
+				// Bottom
+				RenderBatch.EmplaceVertex(FSlateVertex::Make(RenderTransform, FVector2f(Position.X, EndPos.Y) + FVector2f(-1.f, 1.f) / DrawScale, LocalSize, DrawScale, FVector4f(FVector2f(StartUV.X, EndUV.Y), Tiling), FeatherColor, SecondaryColor, Rounding)); //8
+				RenderBatch.EmplaceVertex(FSlateVertex::Make(RenderTransform, FVector2f(LeftMarginX, EndPos.Y) + FVector2f(0.f, 1.f) / DrawScale, LocalSize, DrawScale, FVector4f(FVector2f(LeftMarginU, EndUV.Y), Tiling), FeatherColor, SecondaryColor, Rounding)); //9
+				RenderBatch.EmplaceVertex(FSlateVertex::Make(RenderTransform, FVector2f(RightMarginX, EndPos.Y) + FVector2f(0.f, 1.f) / DrawScale, LocalSize, DrawScale, FVector4f(FVector2f(RightMarginU, EndUV.Y), Tiling), FeatherColor, SecondaryColor, Rounding)); //10
+				RenderBatch.EmplaceVertex(FSlateVertex::Make(RenderTransform, FVector2f(EndPos.X, EndPos.Y) + FVector2f(1.f, 1.f) / DrawScale, LocalSize, DrawScale, FVector4f(FVector2f(EndUV), Tiling), FeatherColor, SecondaryColor, Rounding)); //11
+
+				auto IndexQuadArray = [](FSlateRenderBatch& RenderBatch, int32 TopLeft, int32 TopRight, int32 BottomRight, int32 BottomLeft)
 				{
-					IndexArray.Add(TopLeft);
-					IndexArray.Add(TopRight);
-					IndexArray.Add(BottomRight);
+					RenderBatch.EmplaceIndex(TopLeft);
+					RenderBatch.EmplaceIndex(TopRight);
+					RenderBatch.EmplaceIndex(BottomRight);
 
-					IndexArray.Add(BottomRight);
-					IndexArray.Add(BottomLeft);
-					IndexArray.Add(TopLeft);
+					RenderBatch.EmplaceIndex(BottomRight);
+					RenderBatch.EmplaceIndex(BottomLeft);
+					RenderBatch.EmplaceIndex(TopLeft);
 				};
 
 				// Top Left
-				IndexQuadArray(FeatheredPaddedBoxIndicies, FeatherStart + 0, FeatherStart + 1, IndexStart + 2, IndexStart + 0);
+				IndexQuadArray(RenderBatch, FeatherStart + 0, FeatherStart + 1, IndexStart + 2, IndexStart + 0);
 				// Top Middle
-				IndexQuadArray(FeatheredPaddedBoxIndicies, FeatherStart + 1, FeatherStart + 2, IndexStart + 4, IndexStart + 2);
+				IndexQuadArray(RenderBatch, FeatherStart + 1, FeatherStart + 2, IndexStart + 4, IndexStart + 2);
 				// Top Right
-				IndexQuadArray(FeatheredPaddedBoxIndicies, FeatherStart + 2, FeatherStart + 3, IndexStart + 6, IndexStart + 4);
+				IndexQuadArray(RenderBatch, FeatherStart + 2, FeatherStart + 3, IndexStart + 6, IndexStart + 4);
 
 				//-----------------------------------------------------------
 
 				// Left Top
-				IndexQuadArray(FeatheredPaddedBoxIndicies, FeatherStart + 0, IndexStart + 0, IndexStart + 1, FeatherStart + 4);
+				IndexQuadArray(RenderBatch, FeatherStart + 0, IndexStart + 0, IndexStart + 1, FeatherStart + 4);
 				// Left Middle
-				IndexQuadArray(FeatheredPaddedBoxIndicies, FeatherStart + 4, IndexStart + 1, IndexStart + 8, FeatherStart + 5);
+				IndexQuadArray(RenderBatch, FeatherStart + 4, IndexStart + 1, IndexStart + 8, FeatherStart + 5);
 				// Left Bottom
-				IndexQuadArray(FeatheredPaddedBoxIndicies, FeatherStart + 5, IndexStart + 8, IndexStart + 12, FeatherStart + 8);
+				IndexQuadArray(RenderBatch, FeatherStart + 5, IndexStart + 8, IndexStart + 12, FeatherStart + 8);
 
 				//-----------------------------------------------------------
 
 				// Right Top
-				IndexQuadArray(FeatheredPaddedBoxIndicies, IndexStart + 6, FeatherStart + 3, FeatherStart + 6, IndexStart + 7);
+				IndexQuadArray(RenderBatch, IndexStart + 6, FeatherStart + 3, FeatherStart + 6, IndexStart + 7);
 				// Right Middle
-				IndexQuadArray(FeatheredPaddedBoxIndicies, IndexStart + 7, FeatherStart + 6, FeatherStart + 7, IndexStart + 11);
+				IndexQuadArray(RenderBatch, IndexStart + 7, FeatherStart + 6, FeatherStart + 7, IndexStart + 11);
 				// Right Bottom
-				IndexQuadArray(FeatheredPaddedBoxIndicies, IndexStart + 11, FeatherStart + 7, FeatherStart + 11, IndexStart + 15);
+				IndexQuadArray(RenderBatch, IndexStart + 11, FeatherStart + 7, FeatherStart + 11, IndexStart + 15);
 
 				//-----------------------------------------------------------
 
 				// Bottom Left
-				IndexQuadArray(FeatheredPaddedBoxIndicies, IndexStart + 12, IndexStart + 13, FeatherStart + 9, FeatherStart + 8);
+				IndexQuadArray(RenderBatch, IndexStart + 12, IndexStart + 13, FeatherStart + 9, FeatherStart + 8);
 				// Bottom Middle
-				IndexQuadArray(FeatheredPaddedBoxIndicies, IndexStart + 13, IndexStart + 14, FeatherStart + 10, FeatherStart + 9);
+				IndexQuadArray(RenderBatch, IndexStart + 13, IndexStart + 14, FeatherStart + 10, FeatherStart + 9);
 				// Bottom Right
-				IndexQuadArray(FeatheredPaddedBoxIndicies, IndexStart + 14, IndexStart + 15, FeatherStart + 11, FeatherStart + 10);
-
-				return FeatheredPaddedBoxIndicies;
-			}();
-			RenderBatch.SourceIndices = &FeatheredPaddedBoxIndiciesArray;
-			RenderBatch.NumIndices = FeatheredPaddedBoxIndiciesArray.Num();
-			RenderBatch.IndexOffset = 0;
-		}
-	}
-	else
-	{
-		if( bMirrorHorizontal || bMirrorVertical )
-		{
-			const FVector2f UVMin = StartUV;
-			const FVector2f UVMax = EndUV;
-
-			if( bMirrorHorizontal )
-			{
-				StartUV.X = UVMax.X - ( StartUV.X - UVMin.X );
-				EndUV.X = UVMax.X - ( EndUV.X - UVMin.X );
-			}
-			if( bMirrorVertical )
-			{
-				StartUV.Y = UVMax.Y - ( StartUV.Y - UVMin.Y );
-				EndUV.Y = UVMax.Y - ( EndUV.Y - UVMin.Y );
+				IndexQuadArray(RenderBatch, IndexStart + 14, IndexStart + 15, FeatherStart + 11, FeatherStart + 10);
 			}
 		}
-
-		// Add four vertices to the list of verts to be added to the vertex buffer
-		RenderBatch.AddVertex(FSlateVertex::Make<Rounding>(RenderTransform, TopLeft, LocalSize, DrawScale, FVector4f(FVector2f(StartUV), Tiling), Tint, SecondaryColor));
-		RenderBatch.AddVertex(FSlateVertex::Make<Rounding>(RenderTransform, TopRight, LocalSize, DrawScale, FVector4f(FVector2f(EndUV.X, StartUV.Y), Tiling), Tint, SecondaryColor));
-		RenderBatch.AddVertex(FSlateVertex::Make<Rounding>(RenderTransform, BotLeft, LocalSize, DrawScale, FVector4f(FVector2f(StartUV.X, EndUV.Y), Tiling), Tint, SecondaryColor));
-		RenderBatch.AddVertex(FSlateVertex::Make<Rounding>(RenderTransform, BotRight, LocalSize, DrawScale, FVector4f(FVector2f(EndUV), Tiling), Tint, SecondaryColor));
-
-		static FSlateIndexArray DefaultBoxIndiciesArray = []()
+		else
 		{
-			const int32 IndexStart = 0;
-			FSlateIndexArray DefaultBoxIndicies = {};
-			DefaultBoxIndicies.Reserve(6);
-
-			DefaultBoxIndicies.Add(IndexStart + 0);
-			DefaultBoxIndicies.Add(IndexStart + 1);
-			DefaultBoxIndicies.Add(IndexStart + 2);
-
-			DefaultBoxIndicies.Add(IndexStart + 2);
-			DefaultBoxIndicies.Add(IndexStart + 1);
-			DefaultBoxIndicies.Add(IndexStart + 3);
+			if (bMirrorHorizontal || bMirrorVertical)
+			{
+				const FVector2f UVMin = StartUV;
+				const FVector2f UVMax = EndUV;
 
-			return DefaultBoxIndicies;
-		}();
-		RenderBatch.SourceIndices = &DefaultBoxIndiciesArray;
-		RenderBatch.NumIndices = DefaultBoxIndiciesArray.Num();
-		RenderBatch.IndexOffset = 0;
+				if (bMirrorHorizontal)
+				{
+					StartUV.X = UVMax.X - (StartUV.X - UVMin.X);
+					EndUV.X = UVMax.X - (EndUV.X - UVMin.X);
+				}
+				if (bMirrorVertical)
+				{
+					StartUV.Y = UVMax.Y - (StartUV.Y - UVMin.Y);
+					EndUV.Y = UVMax.Y - (EndUV.Y - UVMin.Y);
+				}
+			}
 
-		if ( GSlateFeathering && Rounding == ESlateVertexRounding::Disabled )
-		{
-			const int32 FeatherStart = RenderBatch.GetNumVertices();
+			// Add four vertices to the list of verts to be added to the vertex buffer
+			RenderBatch.EmplaceVertex(FSlateVertex::Make(RenderTransform, TopLeft, LocalSize, DrawScale, FVector4f(FVector2f(StartUV), Tiling), Tint, SecondaryColor, Rounding));
+			RenderBatch.EmplaceVertex(FSlateVertex::Make(RenderTransform, TopRight, LocalSize, DrawScale, FVector4f(FVector2f(EndUV.X, StartUV.Y), Tiling), Tint, SecondaryColor, Rounding));
+			RenderBatch.EmplaceVertex(FSlateVertex::Make(RenderTransform, BotLeft, LocalSize, DrawScale, FVector4f(FVector2f(StartUV.X, EndUV.Y), Tiling), Tint, SecondaryColor, Rounding));
+			RenderBatch.EmplaceVertex(FSlateVertex::Make(RenderTransform, BotRight, LocalSize, DrawScale, FVector4f(FVector2f(EndUV), Tiling), Tint, SecondaryColor, Rounding));
+	
+			RenderBatch.EmplaceIndex(IndexStart + 0);
+			RenderBatch.EmplaceIndex(IndexStart + 1);
+			RenderBatch.EmplaceIndex(IndexStart + 2);
 
-			RenderBatch.AddVertex(FSlateVertex::Make<Rounding>(RenderTransform, TopLeft + FVector2f(-1.f, -1.f) / DrawScale, LocalSize, DrawScale, FVector4f(FVector2f(StartUV), Tiling), FeatherColor, SecondaryColor));
-			RenderBatch.AddVertex(FSlateVertex::Make<Rounding>(RenderTransform, TopRight + FVector2f(1.f, -1.f) / DrawScale, LocalSize, DrawScale, FVector4f(FVector2f(EndUV.X, StartUV.Y), Tiling), FeatherColor, SecondaryColor));
-			RenderBatch.AddVertex(FSlateVertex::Make<Rounding>(RenderTransform, BotLeft + FVector2f(-1.f, 1.f) / DrawScale, LocalSize, DrawScale, FVector4f(FVector2f(StartUV.X, EndUV.Y), Tiling), FeatherColor, SecondaryColor));
-			RenderBatch.AddVertex(FSlateVertex::Make<Rounding>(RenderTransform, BotRight + FVector2f(1.f, 1.f) / DrawScale, LocalSize, DrawScale, FVector4f(FVector2f(EndUV), Tiling), FeatherColor, SecondaryColor));
+			RenderBatch.EmplaceIndex(IndexStart + 2);
+			RenderBatch.EmplaceIndex(IndexStart + 1);
+			RenderBatch.EmplaceIndex(IndexStart + 3);
 
-			static FSlateIndexArray FeatheredBoxIndiciesArray = [DefaultBoxIndiciesArray = DefaultBoxIndiciesArray]()
+			if (GSlateFeathering && Rounding == ESlateVertexRounding::Disabled)
 			{
-				const int32 IndexStart = 0;
-				const int32 FeatherStart = 4;
+				const int32 FeatherStart = RenderBatch.GetNumVertices();
+
+				RenderBatch.EmplaceVertex(FSlateVertex::Make(RenderTransform, TopLeft + FVector2f(-1.f, -1.f) / DrawScale, LocalSize, DrawScale, FVector4f(FVector2f(StartUV), Tiling), FeatherColor, SecondaryColor, Rounding));
+				RenderBatch.EmplaceVertex(FSlateVertex::Make(RenderTransform, TopRight + FVector2f(1.f, -1.f) / DrawScale, LocalSize, DrawScale, FVector4f(FVector2f(EndUV.X, StartUV.Y), Tiling), FeatherColor, SecondaryColor, Rounding));
+				RenderBatch.EmplaceVertex(FSlateVertex::Make(RenderTransform, BotLeft + FVector2f(-1.f, 1.f) / DrawScale, LocalSize, DrawScale, FVector4f(FVector2f(StartUV.X, EndUV.Y), Tiling), FeatherColor, SecondaryColor, Rounding));
+				RenderBatch.EmplaceVertex(FSlateVertex::Make(RenderTransform, BotRight + FVector2f(1.f, 1.f) / DrawScale, LocalSize, DrawScale, FVector4f(FVector2f(EndUV), Tiling), FeatherColor, SecondaryColor, Rounding));
 
 				const int32 TopLeftIndex = IndexStart + 0;
 				const int32 TopRightIndex = IndexStart + 1;
 				const int32 BottomLeftIndex = IndexStart + 2;
 				const int32 BottomRightIndex = IndexStart + 3;
 
-				FSlateIndexArray FeatheredBoxIndicies = DefaultBoxIndiciesArray;
-				FeatheredBoxIndicies.Reserve(FeatheredBoxIndicies.Num() + 3 * 8);
-
 				// Top-Top
-				FeatheredBoxIndicies.Add(FeatherStart + 0);
-				FeatheredBoxIndicies.Add(FeatherStart + 1);
-				FeatheredBoxIndicies.Add(TopRightIndex);
+				RenderBatch.EmplaceIndex(FeatherStart + 0);
+				RenderBatch.EmplaceIndex(FeatherStart + 1);
+				RenderBatch.EmplaceIndex(TopRightIndex);
 
 				// Top-Bottom
-				FeatheredBoxIndicies.Add(FeatherStart + 0);
-				FeatheredBoxIndicies.Add(TopRightIndex);
-				FeatheredBoxIndicies.Add(TopLeftIndex);
+				RenderBatch.EmplaceIndex(FeatherStart + 0);
+				RenderBatch.EmplaceIndex(TopRightIndex);
+				RenderBatch.EmplaceIndex(TopLeftIndex);
 
 				// Left-Top
-				FeatheredBoxIndicies.Add(FeatherStart + 0);
-				FeatheredBoxIndicies.Add(BottomLeftIndex);
-				FeatheredBoxIndicies.Add(FeatherStart + 2);
+				RenderBatch.EmplaceIndex(FeatherStart + 0);
+				RenderBatch.EmplaceIndex(BottomLeftIndex);
+				RenderBatch.EmplaceIndex(FeatherStart + 2);
 
 				// Left-Bottom
-				FeatheredBoxIndicies.Add(FeatherStart + 0);
-				FeatheredBoxIndicies.Add(TopLeftIndex);
-				FeatheredBoxIndicies.Add(BottomLeftIndex);
+				RenderBatch.EmplaceIndex(FeatherStart + 0);
+				RenderBatch.EmplaceIndex(TopLeftIndex);
+				RenderBatch.EmplaceIndex(BottomLeftIndex);
 
 				// Right-Top
-				FeatheredBoxIndicies.Add(TopRightIndex);
-				FeatheredBoxIndicies.Add(FeatherStart + 1);
-				FeatheredBoxIndicies.Add(FeatherStart + 3);
+				RenderBatch.EmplaceIndex(TopRightIndex);
+				RenderBatch.EmplaceIndex(FeatherStart + 1);
+				RenderBatch.EmplaceIndex(FeatherStart + 3);
 
 				// Right-Bottom
-				FeatheredBoxIndicies.Add(TopRightIndex);
-				FeatheredBoxIndicies.Add(FeatherStart + 3);
-				FeatheredBoxIndicies.Add(BottomRightIndex);
+				RenderBatch.EmplaceIndex(TopRightIndex);
+				RenderBatch.EmplaceIndex(FeatherStart + 3);
+				RenderBatch.EmplaceIndex(BottomRightIndex);
 
 				// Bottom-Top
-				FeatheredBoxIndicies.Add(BottomLeftIndex);
-				FeatheredBoxIndicies.Add(BottomRightIndex);
-				FeatheredBoxIndicies.Add(FeatherStart + 3);
+				RenderBatch.EmplaceIndex(BottomLeftIndex);
+				RenderBatch.EmplaceIndex(BottomRightIndex);
+				RenderBatch.EmplaceIndex(FeatherStart + 3);
 
 				// Bottom-Bottom
-				FeatheredBoxIndicies.Add(FeatherStart + 3);
-				FeatheredBoxIndicies.Add(FeatherStart + 2);
-				FeatheredBoxIndicies.Add(BottomLeftIndex);
+				RenderBatch.EmplaceIndex(FeatherStart + 3);
+				RenderBatch.EmplaceIndex(FeatherStart + 2);
+				RenderBatch.EmplaceIndex(BottomLeftIndex);
+			}
+		}
+	};
+	
+	auto GenerateBoxElementBatchParams = [&](const FSlateDrawElement& InDrawElement)
+	{
+		const FSlateBoxPayload& DrawElementPayload = InDrawElement.GetDataPayload<FSlateBoxPayload>();
+		const FVector2f LocalSize = InDrawElement.GetLocalSize();
+		const ESlateDrawEffect DrawEffects = InDrawElement.GetDrawEffects();
+
+		// Shader type and params
+		ESlateShader ShaderType = ESlateShader::Default;
+		FShaderParams ShaderParams;
+		if (InDrawElement.GetElementType() == EElementType::ET_RoundedBox)
+		{
+			ShaderType = ESlateShader::RoundedBox;
+			const FSlateRoundedBoxPayload& RoundedPayload = InDrawElement.GetDataPayload<FSlateRoundedBoxPayload>();
 
-				return FeatheredBoxIndicies;
-			}();
-			RenderBatch.SourceIndices = &FeatheredBoxIndiciesArray;
-			RenderBatch.NumIndices = FeatheredBoxIndiciesArray.Num();
-			RenderBatch.IndexOffset = 0;
+			ShaderParams.PixelParams = FVector4f(0, RoundedPayload.GetOutlineWeight(), LocalSize.X, LocalSize.Y);//RadiusWeight;
+			ShaderParams.PixelParams2 = RoundedPayload.GetRadius();
 		}
-	}
+
+		// Shader Resource
+		const FSlateShaderResourceProxy* ResourceProxy = DrawElementPayload.GetResourceProxy();
+		FSlateShaderResource* Resource = nullptr;
+		if (ResourceProxy)
+		{
+			// The actual texture for rendering.  If the texture is atlased this is the atlas
+			Resource = ResourceProxy->Resource;
+		}
+
+		// Draw Flags
+		const ESlateBrushTileType::Type TilingRule = DrawElementPayload.GetBrushTiling();
+		const bool bTileHorizontal = (TilingRule == ESlateBrushTileType::Both || TilingRule == ESlateBrushTileType::Horizontal);
+		const bool bTileVertical = (TilingRule == ESlateBrushTileType::Both || TilingRule == ESlateBrushTileType::Vertical);
+
+		const ESlateBrushMirrorType::Type MirroringRule = DrawElementPayload.GetBrushMirroring();
+		const bool bMirrorHorizontal = (MirroringRule == ESlateBrushMirrorType::Both || MirroringRule == ESlateBrushMirrorType::Horizontal);
+		const bool bMirrorVertical = (MirroringRule == ESlateBrushMirrorType::Both || MirroringRule == ESlateBrushMirrorType::Vertical);
+
+		// Pass the tiling information as a flag so we can pick the correct texture addressing mode
+		ESlateBatchDrawFlag DrawFlags = InDrawElement.GetBatchFlags();
+		DrawFlags |= ((bTileHorizontal ? ESlateBatchDrawFlag::TileU : ESlateBatchDrawFlag::None) | (bTileVertical ? ESlateBatchDrawFlag::TileV : ESlateBatchDrawFlag::None));
+
+		return FSlateRenderBatchParams(
+			InDrawElement.GetLayer()
+			, ShaderParams
+			, Resource
+			, ESlateDrawPrimitive::TriangleList
+			, ShaderType
+			, DrawEffects
+			, DrawFlags
+			, InDrawElement.GetSceneIndex()
+			, ResolveClippingState(InDrawElement));
+	};
+
+	auto ReserveBoxElementBatch = [&](FSlateRenderBatch& RenderBatch, uint32 InBatchStart, uint32 InBatchEnd)
+	{
+		uint32 NumVertexes = 0;
+		uint32 NumIndices = 0;
+
+		// For source of magic numbers, see internal method above
+		for (uint32 Index = InBatchStart; Index < InBatchEnd; Index++)
+		{
+			const FSlateDrawElement& DrawElement = DrawElements[Index];
+			const ESlateVertexRounding Rounding = DrawElement.IsPixelSnapped() ? ESlateVertexRounding::Enabled : ESlateVertexRounding::Disabled;
+			const FSlateBoxPayload& DrawElementPayload = DrawElement.GetDataPayload<FSlateBoxPayload>();
+			const FMargin& Margin = DrawElementPayload.GetBrushMargin();
+
+			if (DrawElementPayload.GetBrushDrawType() != ESlateBrushDrawType::Image &&
+				(Margin.Left != 0.0f || Margin.Top != 0.0f || Margin.Right != 0.0f || Margin.Bottom != 0.0f))
+			{
+				NumVertexes += 8 + 8;
+				NumIndices += 6 * 3 * 3;
+
+				if (GSlateFeathering && Rounding == ESlateVertexRounding::Disabled)
+				{
+					NumVertexes += 4 + 2 + 2 + 4;
+					NumIndices += 6 * 3 * 4;
+				}
+			}
+			else
+			{
+				NumVertexes += 4;
+				NumIndices += 6;
+
+				if (GSlateFeathering && Rounding == ESlateVertexRounding::Disabled)
+				{
+					NumVertexes += 4;
+					NumIndices += 3 * 8;
+				}
+			}
+		}
+
+		RenderBatch.SourceVertices->Reserve(RenderBatch.SourceVertices->Num() + NumVertexes);
+		RenderBatch.SourceIndices->Reserve(RenderBatch.SourceIndices->Num() + NumVertexes);
+	};
+
+	GenerateIndexedVertexBatches(DrawElements, AddBoxElementInternal, GenerateBoxElementBatchParams, ReserveBoxElementBatch);
 }
 
 namespace SlateElementBatcher
@@ -1368,7 +1449,7 @@ void FSlateElementBatcher::AddTextElement(const FSlateDrawElement& DrawElement)
 					}
 					else
 					{
-					RenderBatch->ReserveIndices(GlyphsLeft * 6);
+						RenderBatch->ReserveIndices(GlyphsLeft * 6);
 					}
 
 					InvTextureSizeX = 1.0f / FontAtlasTexture->GetWidth();
@@ -1439,15 +1520,15 @@ void FSlateElementBatcher::AddTextElement(const FSlateDrawElement& DrawElement)
 						}
 						else
 						{
-						RenderBatch->AddIndex(IndexStart + 0);
-						RenderBatch->AddIndex(IndexStart + 1);
-						RenderBatch->AddIndex(IndexStart + 2);
-						RenderBatch->AddIndex(IndexStart + 1);
-						RenderBatch->AddIndex(IndexStart + 3);
-						RenderBatch->AddIndex(IndexStart + 2);
+							RenderBatch->AddIndex(IndexStart + 0);
+							RenderBatch->AddIndex(IndexStart + 1);
+							RenderBatch->AddIndex(IndexStart + 2);
+							RenderBatch->AddIndex(IndexStart + 1);
+							RenderBatch->AddIndex(IndexStart + 3);
+							RenderBatch->AddIndex(IndexStart + 2);
+						}
 					}
 				}
-				}
 
 				LineX += Entry.XAdvance;
 			}
@@ -2111,35 +2192,90 @@ static FVector2f GetRotated90(const FVector2f V)
 struct FLineBuilder
 {
 	/** Constructor
-	*
-	* @param ElementScale		Element layout scale, used to convert screenspace thickness/radius.
-	* @param HalfThickness		Half thickness of the lines in screenspace pixels.
-	* @param FilterRadius		Antialiasing filter radius in screenspace pixels.
-	* @param MiterAngleLimit	Corners sharper than this will be split instead of mitered (0-180 degrees).
-	*/
+	 *
+	 * @param ElementScale		Element layout scale, used to convert screenspace thickness/radius.
+	 * @param HalfThickness		Half thickness of the lines in screenspace pixels.
+	 * @param FilterRadius		Antialiasing filter radius in screenspace pixels.
+	 * @param AngleCosineLimit	Miter Angle Limit after being passed through AngleCosine.
+	 */
 	FLineBuilder(FSlateRenderBatch& InRenderBatch, const FSlateRenderTransform& InRenderTransform, float ElementScale, float HalfThickness, float FilterRadius, float MiterAngleLimit) :
 		RenderBatch(InRenderBatch),
 		RenderTransform(InRenderTransform),
 		LocalHalfThickness((HalfThickness + FilterRadius) / ElementScale),
 		LocalFilterRadius(FilterRadius / ElementScale),
 		LocalCapLength((FilterRadius / ElementScale) * 2.0f),
-		AngleCosineLimit(FMath::Cos(FMath::DegreesToRadians((180.0f - MiterAngleLimit) * 0.5f)))
+		AngleCosineLimit(FMath::DegreesToRadians((180.0f - MiterAngleLimit) * 0.5f))
 	{
 	}
 
 	/**
-	* Build geometry for each line segment between the passed-in points. Segments
-	* are represented by trapezoids with a pair of edges parallel to the segment.
-	* Each trapezoid connects to those for neighbouring segments if it's possible to
-	* miter the join cleanly. Rectangular caps are added to antialias any open ends.
-	*
-	* All quads are built using two triangles. The UV coordinates are based on
-	* topology only, not modified by vertex position. This is why each antialiased
-	* edge is kept parallel to the opposite side. Not doing so would result in the
-	* UVs shearing apart at the diagonal where the two triangles meet.
-	*/
-	template<ESlateVertexRounding Rounding>
-	void BuildLineGeometry(const TArray<FVector2f>& Points, const TArray<FColor>& PackedColors, const FColor& PackedTint) const
+	 * Calculate num elements per line
+	 */
+	void NumElements(const TArray<FVector2f>& Points, uint32& OutNumVertex, uint32& OutNumIndex)
+	{
+		FVector2f Position = Points[0];
+		FVector2f NextPosition = Points[1];
+
+		FVector2f Direction;
+		float Length;
+		(NextPosition - Position).ToDirectionAndLength(Direction, Length);
+		FVector2f Up = GetRotated90(Direction) * LocalHalfThickness;
+
+		NumStartCap(Length, OutNumVertex, OutNumIndex);
+
+		// Build the intermediate points and their incoming segments
+		const int32 LastPointIndex = Points.Num() - 1;
+		for (int32 Point = 1; Point < LastPointIndex; ++Point)
+		{
+			const FVector2f LastDirection = Direction;
+			const float LastLength = Length;
+
+			Position = NextPosition;
+			NextPosition = Points[Point + 1];
+
+			(NextPosition - Position).ToDirectionAndLength(Direction, Length);
+			Up = GetRotated90(Direction) * LocalHalfThickness;
+
+			// Project "up" onto the miter normal to get perpendicular distance to the miter line
+			const FVector2f MiterNormal = GetMiterNormal(LastDirection, Direction);
+			const float DistanceToMiterLine = FVector2f::DotProduct(Up, MiterNormal);
+
+			// Get the component of Direction perpendicular to the miter line
+			const float DirDotMiterNormal = FVector2f::DotProduct(Direction, MiterNormal);
+
+			// Break the strip at zero-length segments, if the miter angle is
+			// too tight or the miter would cross either segment's bisector
+			const float MinSegmentLength = FMath::Min(LastLength, Length);
+			if (MinSegmentLength > SMALL_NUMBER &&
+				DirDotMiterNormal >= AngleCosineLimit &&
+				(MinSegmentLength * 0.5f * DirDotMiterNormal) >= FMath::Abs(DistanceToMiterLine))
+			{
+				OutNumVertex += 2;
+				OutNumIndex += 6;
+			}
+			else
+			{
+				// Can't miter, so end the current strip and start a new one
+				NumEndCap(Length, OutNumVertex, OutNumIndex);
+				NumStartCap(LastLength, OutNumVertex, OutNumIndex);
+			}
+		}
+
+		NumEndCap(Length, OutNumVertex, OutNumIndex);
+	}
+
+	/**
+	 * Build geometry for each line segment between the passed-in points. Segments
+	 * are represented by trapezoids with a pair of edges parallel to the segment.
+	 * Each trapezoid connects to those for neighbouring segments if it's possible to
+	 * miter the join cleanly. Rectangular caps are added to antialias any open ends.
+	 *
+	 * All quads are built using two triangles. The UV coordinates are based on
+	 * topology only, not modified by vertex position. This is why each antialiased
+	 * edge is kept parallel to the opposite side. Not doing so would result in the
+	 * UVs shearing apart at the diagonal where the two triangles meet.
+	 */
+	void BuildLineGeometry(const TArray<FVector2f>& Points, const TArray<FColor>& PackedColors, const FColor& PackedTint, ESlateVertexRounding Rounding) const
 	{
 		FColor PointColor = PackedColors.Num() ? PackedColors[0] : PackedTint;
 		FVector2f Position = Points[0];
@@ -2151,7 +2287,7 @@ struct FLineBuilder
 		FVector2f Up = GetRotated90(Direction) * LocalHalfThickness;
 
 		// Build the start cap at the first point
-		MakeStartCap<Rounding>(Position, Direction, Length, Up, PointColor);
+		MakeStartCap(Position, Direction, Length, Up, PointColor, Rounding);
 
 		// @TODO: Since the vertex - index relationship is not homogenous whenever 
 		// a miter angle break occurs we cannot pre-allocate indicies here
@@ -2191,21 +2327,21 @@ struct FLineBuilder
 				const float ParallelDistance = DistanceToMiterLine / DirDotMiterNormal;
 				const FVector2f MiterUp = Up - (Direction * ParallelDistance);
 
-				RenderBatch.AddVertex(FSlateVertex::Make<Rounding>(RenderTransform, FVector2f(Position + MiterUp), FVector2f(1.0f, 0.0f), PointColor));
-				RenderBatch.AddVertex(FSlateVertex::Make<Rounding>(RenderTransform, FVector2f(Position - MiterUp), FVector2f(-1.0f, 0.0f), PointColor));
+				RenderBatch.EmplaceVertex(FSlateVertex::Make(RenderTransform, FVector2f(Position + MiterUp), FVector2f(1.0f, 0.0f), PointColor, {}, Rounding));
+				RenderBatch.EmplaceVertex(FSlateVertex::Make(RenderTransform, FVector2f(Position - MiterUp), FVector2f(-1.0f, 0.0f), PointColor, {}, Rounding));
 				AddQuadIndices(RenderBatch);
 			}
 			else
 			{
 				// Can't miter, so end the current strip and start a new one
-				MakeEndCap<Rounding>(Position, LastDirection, LastLength, LastUp, PointColor);
-				MakeStartCap<Rounding>(Position, Direction, Length, Up, PointColor);
+				MakeEndCap(Position, LastDirection, LastLength, LastUp, PointColor, Rounding);
+				MakeStartCap(Position, Direction, Length, Up, PointColor, Rounding);
 			}
 		}
 
 		// Build the last point's incoming segment and end cap
 		PointColor = PackedColors.Num() ? PackedColors[LastPointIndex] : PackedTint;
-		MakeEndCap<Rounding>(NextPosition, Direction, Length, Up, PointColor);
+		MakeEndCap(NextPosition, Direction, Length, Up, PointColor, Rounding);
 	}
 
 	/**
@@ -2217,23 +2353,59 @@ struct FLineBuilder
 	*    | \|
 	*    0--2
 	*/
-	static void AddQuadIndices(FSlateRenderBatch& RenderBatch)
+	static void AddQuadIndices(FSlateRenderBatch& InRenderBatch)
 	{
-		const int32 NumVerts = RenderBatch.GetNumVertices();
-		IndexQuad(RenderBatch, NumVerts - 3, NumVerts - 1, NumVerts - 2, NumVerts - 4);
+		const int32 NumVerts = InRenderBatch.GetNumVertices();
+
+		auto IndexQuad = [](FSlateRenderBatch& InRenderBatch, int32 TopLeft, int32 TopRight, int32 BottomRight, int32 BottomLeft)
+		{
+			InRenderBatch.EmplaceIndex(TopLeft);
+			InRenderBatch.EmplaceIndex(TopRight);
+			InRenderBatch.EmplaceIndex(BottomRight);
+
+			InRenderBatch.EmplaceIndex(BottomRight);
+			InRenderBatch.EmplaceIndex(BottomLeft);
+			InRenderBatch.EmplaceIndex(TopLeft);
+		};
+
+		IndexQuad(InRenderBatch, NumVerts - 3, NumVerts - 1, NumVerts - 2, NumVerts - 4);
 	}
 
 private:
 	/**
-	* Create a rectangular cap to antialias the start of a segment
-	*/
-	template<ESlateVertexRounding Rounding>
+	 * Calculate num elements in start cap
+	 */
+	void NumStartCap(float SegmentLength, uint32& OutNumVertex, uint32& OutNumIndex) const
+	{
+		if (SegmentLength > SMALL_NUMBER)
+		{
+			OutNumVertex += 4;
+			OutNumIndex += 6;
+		}
+	}
+
+	/**
+	 * Calculate num elements in start cap
+	 */
+	void NumEndCap(float SegmentLength, uint32& OutNumVertex, uint32& OutNumIndex) const
+	{
+		if (SegmentLength > SMALL_NUMBER)
+		{
+			OutNumVertex += 4;
+			OutNumIndex += 12;
+		}
+	}
+
+	/**
+	 * Create a rectangular cap to antialias the start of a segment
+	 */
 	void MakeStartCap(
 		const FVector2f Position,
 		const FVector2f Direction,
 		float SegmentLength,
 		const FVector2f Up,
-		const FColor& Color) const
+		const FColor& Color,
+		ESlateVertexRounding Rounding) const
 	{
 		// Don't build a cap for a zero-length segment
 		// Up would be zero-length anyway, so it would be invisible
@@ -2246,24 +2418,24 @@ private:
 			const FVector2f CapInward = Direction * InwardDistance;
 			const FVector2f CapOutward = Direction * (InwardDistance - LocalCapLength);
 
-			RenderBatch.AddVertex(FSlateVertex::Make<Rounding>(RenderTransform, FVector2f(Position + CapOutward + Up), FVector2f(1.0f, -1.0f), Color));
-			RenderBatch.AddVertex(FSlateVertex::Make<Rounding>(RenderTransform, FVector2f(Position + CapOutward - Up), FVector2f(-1.0f, -1.0f), Color));
-			RenderBatch.AddVertex(FSlateVertex::Make<Rounding>(RenderTransform, FVector2f(Position + CapInward + Up), FVector2f(1.0f, 0.0f), Color));
-			RenderBatch.AddVertex(FSlateVertex::Make<Rounding>(RenderTransform, FVector2f(Position + CapInward - Up), FVector2f(-1.0f, 0.0f), Color));
+			RenderBatch.EmplaceVertex(FSlateVertex::Make(RenderTransform, FVector2f(Position + CapOutward + Up), FVector2f(1.0f, -1.0f), Color, {}, Rounding));
+			RenderBatch.EmplaceVertex(FSlateVertex::Make(RenderTransform, FVector2f(Position + CapOutward - Up), FVector2f(-1.0f, -1.0f), Color, {}, Rounding));
+			RenderBatch.EmplaceVertex(FSlateVertex::Make(RenderTransform, FVector2f(Position + CapInward + Up), FVector2f(1.0f, 0.0f), Color, {}, Rounding));
+			RenderBatch.EmplaceVertex(FSlateVertex::Make(RenderTransform, FVector2f(Position + CapInward - Up), FVector2f(-1.0f, 0.0f), Color, {}, Rounding));
 			AddQuadIndices(RenderBatch);
 		}
 	}
 
 	/**
-	* Create a rectangular cap to antialias the end of a segment
-	*/
-	template<ESlateVertexRounding Rounding>
+	 * Create a rectangular cap to antialias the end of a segment
+	 */
 	void MakeEndCap(
 		const FVector2f Position,
 		const FVector2f Direction,
 		float SegmentLength,
 		const FVector2f Up,
-		const FColor& Color) const
+		const FColor& Color,
+		ESlateVertexRounding Rounding) const
 	{
 		// Don't build a cap for a zero-length segment
 		// Up would be zero-length anyway, so it would be invisible
@@ -2276,12 +2448,12 @@ private:
 			const FVector2f CapInward = Direction * -InwardDistance;
 			const FVector2f CapOutward = Direction * (LocalCapLength - InwardDistance);
 
-			RenderBatch.AddVertex(FSlateVertex::Make<Rounding>(RenderTransform, FVector2f(Position + CapInward + Up), FVector2f(1.0f, 0.0f), Color));
-			RenderBatch.AddVertex(FSlateVertex::Make<Rounding>(RenderTransform, FVector2f(Position + CapInward - Up), FVector2f(-1.0f, 0.0f), Color));
+			RenderBatch.EmplaceVertex(FSlateVertex::Make(RenderTransform, FVector2f(Position + CapInward + Up), FVector2f(1.0f, 0.0f), Color, {}, Rounding));
+			RenderBatch.EmplaceVertex(FSlateVertex::Make(RenderTransform, FVector2f(Position + CapInward - Up), FVector2f(-1.0f, 0.0f), Color, {}, Rounding));
 			AddQuadIndices(RenderBatch);
 
-			RenderBatch.AddVertex(FSlateVertex::Make<Rounding>(RenderTransform, FVector2f(Position + CapOutward + Up), FVector2f(1.0f, 1.0f), Color));
-			RenderBatch.AddVertex(FSlateVertex::Make<Rounding>(RenderTransform, FVector2f(Position + CapOutward - Up), FVector2f(-1.0f, 1.0f), Color));
+			RenderBatch.EmplaceVertex(FSlateVertex::Make(RenderTransform, FVector2f(Position + CapOutward + Up), FVector2f(1.0f, 1.0f), Color, {}, Rounding));
+			RenderBatch.EmplaceVertex(FSlateVertex::Make(RenderTransform, FVector2f(Position + CapOutward - Up), FVector2f(-1.0f, 1.0f), Color, {}, Rounding));
 			AddQuadIndices(RenderBatch);
 		}
 	}
@@ -2300,7 +2472,7 @@ private:
 
 		if (SizeSquared > SMALL_NUMBER)
 		{
-			return DirSum * FMath::InvSqrt(SizeSquared);
+			return DirSum * FMath::InvSqrtEst(SizeSquared);
 		}
 
 		// Inputs are antiparallel, so any perpendicular vector suffices
@@ -2317,140 +2489,221 @@ private:
 };
 
 
-template<ESlateVertexRounding Rounding>
-void FSlateElementBatcher::AddLineElement( const FSlateDrawElement& DrawElement )
+void FSlateElementBatcher::AddLineElements( const FSlateDrawElementArray& DrawElements )
 {
-	const FSlateLinePayload& DrawElementPayload = DrawElement.GetDataPayload<FSlateLinePayload>();
-	const FSlateRenderTransform& RenderTransform = DrawElement.GetRenderTransform();
-	const ESlateDrawEffect DrawEffects = DrawElement.GetDrawEffects();
-	const int32 Layer = DrawElement.GetLayer();
-
-	const TArray<FVector2f>& Points = DrawElementPayload.GetPoints();
-	const TArray<FLinearColor>& PointColors = DrawElementPayload.GetPointColors();
-
-	const int32 NumPoints = Points.Num();
-	if (NumPoints < 2)
+	auto AddLineElementInternal = [&](FSlateRenderBatch& RenderBatch, const FSlateDrawElement& DrawElement)
 	{
-		return;
-	}
+		const ESlateVertexRounding Rounding = DrawElement.IsPixelSnapped() ? ESlateVertexRounding::Enabled : ESlateVertexRounding::Disabled;
+		const FSlateLinePayload& DrawElementPayload = DrawElement.GetDataPayload<FSlateLinePayload>();
+		const FSlateRenderTransform& RenderTransform = DrawElement.GetRenderTransform();
+		const ESlateDrawEffect DrawEffects = DrawElement.GetDrawEffects();
+		const int32 Layer = DrawElement.GetLayer();
 
-	// Pack element tint and per-point colors
-	const FColor PackedTint = PackVertexColor(DrawElementPayload.GetTint());
-	TArray<FColor> PackedColors;
-	PackedColors.Reserve(PointColors.Num());
-	for (const FLinearColor& PointColor : PointColors)
-	{
-		PackedColors.Add(PackVertexColor(PointColor * DrawElementPayload.GetTint()));
-	}
+		const TArray<FVector2f>& Points = DrawElementPayload.GetPoints();
+		const TArray<FLinearColor>& PointColors = DrawElementPayload.GetPointColors();
 
-	if (DrawElementPayload.IsAntialiased())
-	{
-		// Filter size to use for anti-aliasing.
-		// Increasing this value will increase the fuzziness of line edges.
-		const float FilterRadius = 1.0f;
+		const int32 NumPoints = Points.Num();
+		if (NumPoints < 2)
+		{
+			return;
+		}
 
-		// Corners sharper than this will be split instead of mitered
-		// A split can also be forced by repeating points in the input
-		const float MiterAngleLimit = 90.0f - KINDA_SMALL_NUMBER;
+		// Pack element tint and per-point colors
+		const FColor PackedTint = PackVertexColor(DrawElementPayload.GetTint());
+		TArray<FColor> PackedColors;
+		PackedColors.Reserve(PointColors.Num());
+		for (const FLinearColor& PointColor : PointColors)
+		{
+			PackedColors.Add(PackVertexColor(PointColor * DrawElementPayload.GetTint()));
+		}
 
-		const float HalfThickness = 0.5f * FMath::Max(0.0f, DrawElementPayload.GetThickness());
+		if (DrawElementPayload.IsAntialiased())
+		{
+			// Filter size to use for anti-aliasing.
+			// Increasing this value will increase the fuzziness of line edges.
+			const float FilterRadius = 1.0f;
 
-		// The amount we increase each side of the line to generate enough pixels
-		const float FilteredHalfThickness = HalfThickness + FilterRadius;
+			// Corners sharper than this will be split instead of mitered
+			// A split can also be forced by repeating points in the input
+			const float MiterAngleLimit = 90.0f - KINDA_SMALL_NUMBER;
 
-		const float InsideFilterU = (HalfThickness - FilterRadius) / FilteredHalfThickness;
-		const FVector4f ShaderParams = FVector4f(InsideFilterU, 0.0f, 0.0f, 0.0f);
-		FSlateRenderBatch& RenderBatch = CreateRenderBatch(Layer, FShaderParams::MakePixelShaderParams(ShaderParams), nullptr, ESlateDrawPrimitive::TriangleList, ESlateShader::LineSegment, DrawEffects, ESlateBatchDrawFlag::None, DrawElement);
+			const float HalfThickness = 0.5f * FMath::Max(0.0f, DrawElementPayload.GetThickness());
 
-		FLineBuilder LineBuilder(
-			RenderBatch,
-			RenderTransform,
-			DrawElement.GetScale(),
-			HalfThickness,
-			FilterRadius,
-			MiterAngleLimit);
+			FLineBuilder LineBuilder(
+				RenderBatch,
+				RenderTransform,
+				DrawElement.GetScale(),
+				HalfThickness,
+				FilterRadius,
+				MiterAngleLimit);
 
-		LineBuilder.BuildLineGeometry<Rounding>(Points, PackedColors, PackedTint);
-	}
-	else
-	{
-		if (DrawElementPayload.GetThickness() == 1.0f)
+			LineBuilder.BuildLineGeometry(Points, PackedColors, PackedTint, Rounding);
+		}
+		else
 		{
-			// Generate the line segments using the native line rendering of the platform.
-			FSlateRenderBatch& RenderBatch = CreateRenderBatch(Layer, FShaderParams(), nullptr, ESlateDrawPrimitive::LineList, ESlateShader::Default, DrawEffects, ESlateBatchDrawFlag::None, DrawElement);
-
-			const FColor StartColor = PackedColors.Num() ? PackedColors[0] : PackedTint;
-			RenderBatch.AddVertex(FSlateVertex::Make<Rounding>(RenderTransform, Points[0], FVector2f::ZeroVector, StartColor));
-
-			for (int32 Point = 1; Point < NumPoints; ++Point)
+			if (DrawElementPayload.GetThickness() == 1.0f)
 			{
-				const FColor PointColor = PackedColors.Num() ? PackedColors[Point] : PackedTint;
-				RenderBatch.AddVertex(FSlateVertex::Make<Rounding>(RenderTransform, Points[Point], FVector2f::ZeroVector, PointColor));
-			}
-
-			// The number of points can technically be unbound, so pick some number that if smaller we use the preallocated index array
-			const int32 MaxPreAllocatedPointIndicies = 4096 + 1024;
+				const FColor StartColor = PackedColors.Num() ? PackedColors[0] : PackedTint;
+				RenderBatch.EmplaceVertex(FSlateVertex::Make(RenderTransform, Points[0], FVector2f::ZeroVector, StartColor, {}, Rounding));
 
-			if (NumPoints < MaxPreAllocatedPointIndicies)
-			{
-				static FSlateIndexArray DefaultLineIndiciesArray = [MaxPreAllocatedPointIndicies = MaxPreAllocatedPointIndicies]()
+				for (int32 Point = 1; Point < NumPoints; ++Point)
 				{
-					FSlateIndexArray DefaultLineIndicies = {};
-					DefaultLineIndicies.Reserve(MaxPreAllocatedPointIndicies * 2 - 2);
-
-					for (int32 Point = 1; Point < MaxPreAllocatedPointIndicies; ++Point)
-					{
-						const uint32 IndexStart = 1 + Point;
-						DefaultLineIndicies.Add(IndexStart - 2);
-						DefaultLineIndicies.Add(IndexStart - 1);
-					}
+					const FColor PointColor = PackedColors.Num() ? PackedColors[Point] : PackedTint;
+					RenderBatch.EmplaceVertex(FSlateVertex::Make(RenderTransform, Points[Point], FVector2f::ZeroVector, PointColor, {}, Rounding));
 
-					return DefaultLineIndicies;
-				}();
+					const uint32 IndexStart = RenderBatch.GetNumVertices();
+					RenderBatch.EmplaceIndex(IndexStart - 2);
+					RenderBatch.EmplaceIndex(IndexStart - 1);
+				}
 
-				RenderBatch.SourceIndices = &DefaultLineIndiciesArray;
-				RenderBatch.NumIndices = NumPoints * 2 - 2;
-				RenderBatch.IndexOffset = 0;
 			}
 			else
 			{
-				for (int32 Point = 1; Point < NumPoints; ++Point)
+
+				// Thickness is given in screenspace, convert it to local space for geometry build
+				const float LocalHalfThickness = (DrawElementPayload.GetThickness() * 0.5f) / DrawElement.GetScale();
+
+				for (int32 Point = 0; Point < NumPoints - 1; ++Point)
 				{
-					const uint32 IndexStart = 1 + Point;
-					RenderBatch.AddIndex(IndexStart - 2);
-					RenderBatch.AddIndex(IndexStart - 1);
+					const FVector2f StartPos = Points[Point];
+					const FVector2f EndPos = Points[Point + 1];
+
+					const FColor StartColor = PackedColors.Num() ? PackedColors[Point] : PackedTint;
+					const FColor EndColor = PackedColors.Num() ? PackedColors[Point + 1] : PackedTint;
+
+					const FVector2f SegmentDir = (EndPos - StartPos).GetSafeNormal();
+					const FVector2f Up = GetRotated90(SegmentDir) * LocalHalfThickness;
+
+					RenderBatch.EmplaceVertex(FSlateVertex::Make(RenderTransform, FVector2f(StartPos + Up), FVector2f::ZeroVector, StartColor, {}, Rounding));
+					RenderBatch.EmplaceVertex(FSlateVertex::Make(RenderTransform, FVector2f(StartPos - Up), FVector2f::ZeroVector, StartColor, {}, Rounding));
+					RenderBatch.EmplaceVertex(FSlateVertex::Make(RenderTransform, FVector2f(EndPos + Up), FVector2f::ZeroVector, EndColor, {}, Rounding));
+					RenderBatch.EmplaceVertex(FSlateVertex::Make(RenderTransform, FVector2f(EndPos - Up), FVector2f::ZeroVector, EndColor, {}, Rounding));
+
+					FLineBuilder::AddQuadIndices(RenderBatch);
 				}
 			}
+		}
+	};
+
+	auto GenerateLineElementBatchParams = [&](const FSlateDrawElement& InDrawElement)
+	{
+		const FSlateLinePayload& DrawElementPayload = InDrawElement.GetDataPayload<FSlateLinePayload>();
+		const ESlateDrawEffect DrawEffects = InDrawElement.GetDrawEffects();
+
+		// Shader type and params
+		ESlateShader ShaderType = ESlateShader::Default;
+		FShaderParams ShaderParams = FShaderParams();
+		if (DrawElementPayload.IsAntialiased())
+		{
+			// Filter size to use for anti-aliasing.
+			// Increasing this value will increase the fuzziness of line edges.
+			const float FilterRadius = 1.0f;
+
+			const float HalfThickness = 0.5f * FMath::Max(0.0f, DrawElementPayload.GetThickness());
+
+			// The amount we increase each side of the line to generate enough pixels
+			const float FilteredHalfThickness = HalfThickness + FilterRadius;
+
+			const float InsideFilterU = (HalfThickness - FilterRadius) / FilteredHalfThickness;
+			const FVector4f PixelShaderParams = FVector4f(InsideFilterU, 0.0f, 0.0f, 0.0f);
+			ShaderParams = FShaderParams::MakePixelShaderParams(PixelShaderParams);
+			ShaderType = ESlateShader::LineSegment;
+		}
+
+		// Primitive type
+		ESlateDrawPrimitive PrimitiveType = ESlateDrawPrimitive::LineList;
 
+		if (DrawElementPayload.IsAntialiased())
+		{
+			PrimitiveType = ESlateDrawPrimitive::TriangleList;
 		}
 		else
 		{
-			// Generate the line segments using non-aa'ed quads.
-			FSlateRenderBatch& RenderBatch = CreateRenderBatch(Layer, FShaderParams(), nullptr, ESlateDrawPrimitive::TriangleList, ESlateShader::Default, DrawEffects, ESlateBatchDrawFlag::None, DrawElement);
+			if (DrawElementPayload.GetThickness() == 1.0f)
+			{
+				PrimitiveType = ESlateDrawPrimitive::LineList;
+			}
+			else
+			{
+				PrimitiveType = ESlateDrawPrimitive::TriangleList;
+			}
+		}
+
+		return FSlateRenderBatchParams(
+			InDrawElement.GetLayer()
+			, ShaderParams
+			, nullptr
+			, PrimitiveType
+			, ShaderType
+			, DrawEffects
+			, ESlateBatchDrawFlag::None
+			, InDrawElement.GetSceneIndex()
+			, ResolveClippingState(InDrawElement));
+	};
+
+	auto ReserveLineElementBatch = [&](FSlateRenderBatch& RenderBatch, uint32 InBatchStart, uint32 InBatchEnd)
+	{
+		uint32 NumVertexes = 0;
+		uint32 NumIndices = 0;
+
+		// For source of magic numbers, see internal method above
+		for (uint32 Index = InBatchStart; Index < InBatchEnd; Index++)
+		{
 
-			// Thickness is given in screenspace, convert it to local space for geometry build
-			const float LocalHalfThickness = (DrawElementPayload.GetThickness() * 0.5f) / DrawElement.GetScale();
+			const FSlateDrawElement& DrawElement = DrawElements[Index];
+			const ESlateVertexRounding Rounding = DrawElement.IsPixelSnapped() ? ESlateVertexRounding::Enabled : ESlateVertexRounding::Disabled;
+			const FSlateLinePayload& DrawElementPayload = DrawElement.GetDataPayload<FSlateLinePayload>();
+			const FSlateRenderTransform& RenderTransform = DrawElement.GetRenderTransform();
+			const ESlateDrawEffect DrawEffects = DrawElement.GetDrawEffects();
+			const int32 Layer = DrawElement.GetLayer();
 
-			for (int32 Point = 0; Point < NumPoints - 1; ++Point)
+			const TArray<FVector2f>& Points = DrawElementPayload.GetPoints();
+
+			const int32 NumPoints = Points.Num();
+			if (NumPoints < 2)
+			{
+				return;
+			}
+
+			if (DrawElementPayload.IsAntialiased())
 			{
-				const FVector2f StartPos = Points[Point];
-				const FVector2f EndPos = Points[Point + 1];
+				const float FilterRadius = 1.0f;
 
-				const FColor StartColor = PackedColors.Num() ? PackedColors[Point] : PackedTint;
-				const FColor EndColor = PackedColors.Num() ? PackedColors[Point+1] : PackedTint;
+				const float MiterAngleLimit = 90.0f - KINDA_SMALL_NUMBER;
 
-				const FVector2f SegmentDir = (EndPos - StartPos).GetSafeNormal();
-				const FVector2f Up = GetRotated90(SegmentDir) * LocalHalfThickness;
+				const float HalfThickness = 0.5f * FMath::Max(0.0f, DrawElementPayload.GetThickness());
 
-				RenderBatch.AddVertex(FSlateVertex::Make<Rounding>(RenderTransform, FVector2f(StartPos + Up), FVector2f::ZeroVector, StartColor));
-				RenderBatch.AddVertex(FSlateVertex::Make<Rounding>(RenderTransform, FVector2f(StartPos - Up), FVector2f::ZeroVector, StartColor));
-				RenderBatch.AddVertex(FSlateVertex::Make<Rounding>(RenderTransform, FVector2f(EndPos + Up), FVector2f::ZeroVector, EndColor));
-				RenderBatch.AddVertex(FSlateVertex::Make<Rounding>(RenderTransform, FVector2f(EndPos - Up), FVector2f::ZeroVector, EndColor));
+				FLineBuilder LineBuilder(
+					RenderBatch,
+					RenderTransform,
+					DrawElement.GetScale(),
+					HalfThickness,
+					FilterRadius,
+					MiterAngleLimit);
 
-				FLineBuilder::AddQuadIndices(RenderBatch);
+				LineBuilder.NumElements(Points, NumVertexes, NumIndices);
+			}
+			else
+			{
+				if (DrawElementPayload.GetThickness() == 1.0f)
+				{
+					NumVertexes += NumPoints;
+					NumIndices += (NumPoints - 1) * 2;
+
+				}
+				else
+				{
+					NumVertexes += (NumPoints - 1) * 4;
+					NumIndices += (NumPoints - 1) * 6;
+				}
 			}
 		}
-	}
+
+		RenderBatch.SourceVertices->Reserve(RenderBatch.SourceVertices->Num() + NumVertexes);
+		RenderBatch.SourceIndices->Reserve(RenderBatch.SourceIndices->Num() + NumIndices);
+	};
+
+	GenerateIndexedVertexBatches(DrawElements, AddLineElementInternal, GenerateLineElementBatchParams, ReserveLineElementBatch);
 }
 
 template<ESlateVertexRounding Rounding>
@@ -2894,6 +3147,27 @@ void FSlateElementBatcher::AddPostProcessPass(const FSlateDrawElement& DrawEleme
 	}
 }
 
+FSlateRenderBatch& FSlateElementBatcher::CreateRenderBatch(
+	FSlateBatchData* SlateBatchData
+	, int32 Layer
+	, const FShaderParams& ShaderParams
+	, const FSlateShaderResource* InResource
+	, ESlateDrawPrimitive PrimitiveType
+	, ESlateShader ShaderType
+	, ESlateDrawEffect DrawEffects
+	, ESlateBatchDrawFlag DrawFlags
+	, int8 SceneIndex
+	, const FSlateClippingState* ClippingState)
+{
+	FSlateRenderBatch& NewBatch = CurrentCachedElementList
+		? CurrentCachedElementList->AddRenderBatch(Layer, ShaderParams, InResource, PrimitiveType, ShaderType, DrawEffects, DrawFlags, SceneIndex)
+		: SlateBatchData->AddRenderBatch(Layer, ShaderParams, InResource, PrimitiveType, ShaderType, DrawEffects, DrawFlags, SceneIndex);
+
+	NewBatch.ClippingState = ClippingState;
+
+	return NewBatch;
+}
+
 FSlateRenderBatch& FSlateElementBatcher::CreateRenderBatch(
 	FSlateBatchData* SlateBatchData,
 	int32 Layer, 
@@ -3065,7 +3339,7 @@ void FSlateElementBatcher::BuildShapedTextSequence(const FShapedTextBuildContext
 					}
 					else
 					{
-					RenderBatch->ReserveIndices(GlyphsLeft * 6);
+						RenderBatch->ReserveIndices(GlyphsLeft * 6);
 					}
 
 					InvTextureSizeX = 1.0f / FontAtlasTexture->GetWidth();
@@ -3154,12 +3428,12 @@ void FSlateElementBatcher::BuildShapedTextSequence(const FShapedTextBuildContext
 
 				//Always try to put the ellipsis, wether it fits or not: it's better to have an ellipsis a bit clipped than no feedback at all.
 					bNeedEllipsis = true;
-				EllipsisLineX = (LineX + SizeX - EllipsisWidth);
+					EllipsisLineX = (LineX + SizeX - EllipsisWidth);
 					EllipsisLineY = LineY;
 				}
 				else
 				{
-				OverflowDirection = ETextOverflowDirection::NoOverflow;
+					OverflowDirection = ETextOverflowDirection::NoOverflow;
 				}
 
 			// If we just clipped a glyph omit all characters in said glyph. Otherwise floating diacritics would be visible above the ellipsis. This is common in complex languages.
@@ -3209,12 +3483,12 @@ void FSlateElementBatcher::BuildShapedTextSequence(const FShapedTextBuildContext
 			}
 			else
 			{
-			RenderBatch->AddIndex(IndexStart + 0);
-			RenderBatch->AddIndex(IndexStart + 1);
-			RenderBatch->AddIndex(IndexStart + 2);
-			RenderBatch->AddIndex(IndexStart + 1);
-			RenderBatch->AddIndex(IndexStart + 3);
-			RenderBatch->AddIndex(IndexStart + 2);
+				RenderBatch->AddIndex(IndexStart + 0);
+				RenderBatch->AddIndex(IndexStart + 1);
+				RenderBatch->AddIndex(IndexStart + 2);
+				RenderBatch->AddIndex(IndexStart + 1);
+				RenderBatch->AddIndex(IndexStart + 3);
+				RenderBatch->AddIndex(IndexStart + 2);
 			}
 
 			// Reset whitespace advance to 0, this is a visible character
diff --git a/Engine/Source/Runtime/SlateCore/Private/Rendering/ShaderResourceManager.cpp b/Engine/Source/Runtime/SlateCore/Private/Rendering/ShaderResourceManager.cpp
--- a/Engine/Source/Runtime/SlateCore/Private/Rendering/ShaderResourceManager.cpp
+++ b/Engine/Source/Runtime/SlateCore/Private/Rendering/ShaderResourceManager.cpp
@@ -8,10 +8,6 @@ FSlateResourceHandle::FSlateResourceHandle()
 {
 }
 
-FSlateResourceHandle::~FSlateResourceHandle()
-{
-}
-
 bool FSlateResourceHandle::IsValid() const
 {
 	return Data.IsValid() && Data->Proxy;
diff --git a/Engine/Source/Runtime/SlateCore/Private/Rendering/SlateRenderBatch.cpp b/Engine/Source/Runtime/SlateCore/Private/Rendering/SlateRenderBatch.cpp
--- a/Engine/Source/Runtime/SlateCore/Private/Rendering/SlateRenderBatch.cpp
+++ b/Engine/Source/Runtime/SlateCore/Private/Rendering/SlateRenderBatch.cpp
@@ -1,6 +1,7 @@
 // Copyright Epic Games, Inc. All Rights Reserved.
 
 #include "Rendering/SlateRenderBatch.h"
+#include "Rendering/DrawElements.h"
 #include "Rendering/ElementBatcher.h"
 #include "Textures/SlateShaderResource.h"
 
@@ -43,5 +44,3 @@ FSlateRenderBatch::FSlateRenderBatch(
 {
 	check(ShaderResource == nullptr || !ShaderResource->Debug_IsDestroyed());
 }
-
-
diff --git a/Engine/Source/Runtime/SlateCore/Public/Debugging/SlateDebugging.h b/Engine/Source/Runtime/SlateCore/Public/Debugging/SlateDebugging.h
--- a/Engine/Source/Runtime/SlateCore/Public/Debugging/SlateDebugging.h
+++ b/Engine/Source/Runtime/SlateCore/Public/Debugging/SlateDebugging.h
@@ -255,6 +255,25 @@ public:
 	FText ToText() const;
 };
 
+struct SLATECORE_API FSlateDebuggingElementTypeAddedEventArgs
+{
+public:
+	FSlateDebuggingElementTypeAddedEventArgs(
+		const FSlateWindowElementList& InElementList,
+		int32 InElementIndex,
+		EElementType InElementType
+	);
+
+	/** Element list containing the element */
+	const FSlateWindowElementList& ElementList;
+
+	/** Index into the element within the corresponding typed container */
+	int32 ElementIndex;
+
+	/** Type of the element, element will reside in a container dedicated for that type */
+	EElementType ElementType;
+};
+
 /**
  * 
  */
@@ -295,7 +314,7 @@ PRAGMA_ENABLE_DEPRECATION_WARNINGS
 	 * Called as soon as the element is added to the element list.
 	 * NOTE: These elements are not valid until the widget finishes painting, or you can resolve them all after the window finishes painting.
 	 */
-	DECLARE_MULTICAST_DELEGATE_ThreeParams(FDrawElementType, const FSlateWindowElementList& /*ElementList*/, int32 /*ElementIndex*/, EElementType /*ElementType*/);
+	DECLARE_MULTICAST_DELEGATE_OneParam(FDrawElementType, const FSlateDebuggingElementTypeAddedEventArgs& /*ElementTypeAddedArgs*/);
 	static FDrawElementType ElementTypeAdded;
 
 public:
diff --git a/Engine/Source/Runtime/SlateCore/Public/Rendering/DrawElementPayloads.h b/Engine/Source/Runtime/SlateCore/Public/Rendering/DrawElementPayloads.h
--- a/Engine/Source/Runtime/SlateCore/Public/Rendering/DrawElementPayloads.h
+++ b/Engine/Source/Runtime/SlateCore/Public/Rendering/DrawElementPayloads.h
@@ -72,7 +72,7 @@ struct FSlateBoxPayload : public FSlateDataPayload, public FSlateTintableElement
 		Tiling = InBrush->GetTiling();
 		Mirroring = InBrush->GetMirroring();
 		DrawType = InBrush->GetDrawType();
-		FSlateResourceHandle Handle = InBrush->GetRenderingResource(LocalSize, DrawScale);
+		const FSlateResourceHandle& Handle = InBrush->GetRenderingResource(LocalSize, DrawScale);
 		if (Handle.IsValid())
 		{
 			ResourceProxy = Handle.GetResourceProxy();
diff --git a/Engine/Source/Runtime/SlateCore/Public/Rendering/DrawElementTypes.h b/Engine/Source/Runtime/SlateCore/Public/Rendering/DrawElementTypes.h
new file mode 100644
--- /dev/null
+++ b/Engine/Source/Runtime/SlateCore/Public/Rendering/DrawElementTypes.h
@@ -0,0 +1,31 @@
+// Copyright Epic Games, Inc. All Rights Reserved.
+
+#pragma once
+
+/**
+ * Slate higher level abstract element types
+ */
+enum class EElementType : uint8
+{
+	ET_Box,
+	ET_DebugQuad,
+	ET_Text,
+	ET_ShapedText,
+	ET_Spline,
+	ET_Line,
+	ET_Gradient,
+	ET_Viewport,
+	ET_Border,
+	ET_Custom,
+	ET_CustomVerts,
+	ET_PostProcessPass,
+	ET_RoundedBox,
+	/**
+	 * We map draw elements by type on add for better cache coherency if possible,
+	 * this type is used when that grouping is not possible.
+	 * Grouping is also planned to be used for bulk element type processing.
+	 */
+	ET_NonMapped,
+	/** Total number of draw commands */
+	ET_Count,
+};
\ No newline at end of file
diff --git a/Engine/Source/Runtime/SlateCore/Public/Rendering/DrawElements.h b/Engine/Source/Runtime/SlateCore/Public/Rendering/DrawElements.h
--- a/Engine/Source/Runtime/SlateCore/Public/Rendering/DrawElements.h
+++ b/Engine/Source/Runtime/SlateCore/Public/Rendering/DrawElements.h
@@ -38,30 +38,6 @@ struct FSlateBrush;
 struct FSlateDataPayload;
 struct FSlateGradientStop;
 
-enum class EElementType : uint8
-{
-	ET_Box,
-	ET_DebugQuad,
-	ET_Text,
-	ET_ShapedText,
-	ET_Spline,
-	ET_Line,
-	ET_Gradient,
-	ET_Viewport,
-	ET_Border,
-	ET_Custom,
-	ET_CustomVerts,
-	ET_PostProcessPass,
-	ET_RoundedBox,
-	/** 
-	 * We map draw elements by type on add for better cache coherency if possible, 
-	 * this type is used when that grouping is not possible. 
-	 * Grouping is also planned to be used for bulk element type processing.
-	 */
-	ET_NonMapped,
-	/** Total number of draw commands */
-	ET_Count,
-};
 
 
 /**
@@ -361,6 +337,9 @@ private:
 	uint8 bIsCached : 1;
 };
 
+/**
+ * Data holder for info used in cached renderbatches. Primarily clipping state, vertices and indices
+ */
 struct FSlateCachedFastPathRenderingData
 {
 	~FSlateCachedFastPathRenderingData()
@@ -375,6 +354,11 @@ struct FSlateCachedFastPathRenderingData
 
 struct FSlateCachedElementData;
 
+/**
+ * Cached list of elements corresponding to a particular widget. 
+ * This class is used as an interfaces for renderbatches to interact with cached rendering data with
+ * In particular verticies and indicies
+ */
 struct FSlateCachedElementList
 {
 	FSlateCachedElementList(FSlateCachedElementData* InParentData, const SWidget* InWidget)
@@ -401,6 +385,16 @@ struct FSlateCachedElementList
 
 	void AddReferencedObjects(FReferenceCollector& Collector);
 
+	/**
+	 * Returns true if all typed containers are empty, else false
+	 */
+	SLATECORE_API bool IsEmpty();
+
+	/**
+	 * Returns number of elements in all containers summed
+	 */
+	SLATECORE_API int32 NumElements();
+	
 private:
 	SLATECORE_API void DestroyCachedData();
 
@@ -418,6 +412,12 @@ public:
 	FSlateCachedFastPathRenderingData* CachedRenderingData;
 };
 
+
+/**
+ * Handle used to uniquely identify a cached element list with some widget
+ * Used when a widget gets invalidated / additional elements added
+ * So we can modify that widget's existing cached element list
+ */
 struct FSlateCachedElementsHandle
 {
 	friend struct FSlateCachedElementData;
@@ -444,16 +444,57 @@ private:
 	TWeakPtr<FSlateCachedElementList> Ptr;
 };
 
+/**
+ * Top level class responsible for cached elements within a particular invalidation root.
+ * Equivalent to the plain FSlateBatchData in that it manages multiple renderbatches,
+ * element lists (Per widget), & new element lists (Per invalidated widgets).
+ * 
+ * Note: Just as each window may have multiple invalidation roots, each window element list 
+ * may have multiple 'FSlateCachedElementData' for each root
+ */
 struct FSlateCachedElementData
 {
 	friend class FSlateElementBatcher;
 	friend struct FSlateCachedElementsHandle;
 
+	/** Reset all cached data, except num elements, call EmptyCachedNumElements for that */
 	void Empty();
 
+	/** Create a new CachedElementList, occurs when adding a cached element whose top level draw widget is not already part of some cache */
 	FSlateCachedElementsHandle AddCache(const SWidget* Widget);
 
-	FSlateDrawElement& AddCachedElement(FSlateCachedElementsHandle& CacheHandle, const FSlateClippingManager& ParentClipManager, const SWidget* Widget, EElementType InElementType);
+	static const FSlateClippingState* GetClipStateFromParent(const FSlateClippingManager& ParentClipManager);
+
+	FORCEINLINE void ValidateWidgetOwner(TSharedPtr<FSlateCachedElementList> List, const SWidget* CurrentWidget);
+
+	template<EElementType ElementType>
+	FSlateDrawElement& AddCachedElement(FSlateCachedElementsHandle& CacheHandle, const FSlateClippingManager& ParentClipManager, const SWidget* CurrentWidget)
+	{
+		TSharedPtr<FSlateCachedElementList> List = CacheHandle.Ptr.Pin(); 
+
+#if WITH_SLATE_DEBUGGING
+		ValidateWidgetOwner(List, CurrentWidget);
+#endif
+
+		FSlateDrawElementContainer& Container = List->DrawElements.Get<(uint8)ElementType>();
+		FSlateDrawElement& NewElement = Container.Elements.AddDefaulted_GetRef();
+		NewElement.SetIsCached(true);
+
+		// Check if slow vs checking a flag on the list to see if it contains new data.
+		ListsWithNewData.AddUnique(List.Get());
+
+		const FSlateClippingState* ExistingClipState = FSlateCachedElementData::GetClipStateFromParent(ParentClipManager);
+
+		if (ExistingClipState)
+		{
+			// We need to cache this clip state for the next time the element draws
+			FSlateCachedClipState& CachedClipState = FindOrAddCachedClipState(ExistingClipState);
+			List->AddCachedClipState(CachedClipState);
+			NewElement.SetCachedClippingState(&CachedClipState.ClippingState.Get());
+		}
+
+		return NewElement;
+	}
 
 	FSlateRenderBatch& AddCachedRenderBatch(FSlateRenderBatch&& NewBatch, int32& OutIndex);
 	void RemoveCachedRenderBatches(const TArray<int32>& CachedRenderBatchIndices);
@@ -476,10 +517,13 @@ private:
 	/** List of cached batches to submit for drawing */
 	TSparseArray<FSlateRenderBatch> CachedBatches;
 
+	/** List of cached element lists used to redraw when no invalidation occurs, each list corresponds to a particular widget. See 'FSlateCachedElementsHandle' */
 	TArray<TSharedPtr<FSlateCachedElementList>> CachedElementLists;
 
+	/** List of pointers to instances of the element lists above which have new data or have been invalidated */
 	TArray<FSlateCachedElementList*, TInlineAllocator<50>> ListsWithNewData;
 
+	/** List of clip states used later when rendering */
 	TArray<FSlateCachedClipState> CachedClipStates;
 };
 
@@ -556,9 +600,39 @@ public:
 	}
 
 	/**
-	 * Creates an uninitialized draw element
+	 * Creates an uninitialized draw element if using caching will create a new cached draw list
+	 * if needed (Whenever a top level draw widget's cache handle doesn't match the current cached handle).
+	 */
+	template<EElementType ElementType = EElementType::ET_NonMapped>
+	FSlateDrawElement& AddUninitialized()
+	{
+		const bool bAllowCache = CachedElementDataListStack.Num() > 0 && WidgetDrawStack.Num() && !WidgetDrawStack.Top().bIsVolatile;
+
+		if (bAllowCache)
+		{
+			// @todo get working with slate debugging
+			return AddCachedElement<ElementType>();
+		}
+		else
+		{
+			FSlateDrawElementMap& Elements = UncachedDrawElements;
+			FSlateDrawElementContainer& Container = Elements.Get<(uint8)ElementType>();
+			const int32 InsertIdx = Container.Elements.AddDefaulted();
+
+#if WITH_SLATE_DEBUGGING
+			FSlateDebuggingElementTypeAddedEventArgs ElementTypeAddedArgs{*this, InsertIdx, ElementType};
+			FSlateDebugging::ElementTypeAdded.Broadcast(ElementTypeAddedArgs);
+#endif
+
+			FSlateDrawElement& NewElement = Container.Elements[InsertIdx];
+			return NewElement;
+		}
+	}
+
+	/**
+	 * Calls AddUninitialized, resolving template version based on enum
 	 */
-	SLATECORE_API FSlateDrawElement& AddUninitialized(EElementType InElementType = EElementType::ET_NonMapped);
+	SLATECORE_API FSlateDrawElement& AddUninitializedLookup(EElementType InElementType = EElementType::ET_NonMapped);
 
 
 	//--------------------------------------------------------------------------
@@ -677,7 +751,23 @@ public:
 		}
 	}
 private:
-	FSlateDrawElement& AddCachedElement(EElementType InElementType);
+	/** Adds a cached element, generating a new cached list for the widget at the top of the cache if needed */
+	template<EElementType ElementType>
+	FSlateDrawElement& AddCachedElement()
+	{
+		FSlateCachedElementData* CurrentCachedElementData = GetCurrentCachedElementData();
+		check(CurrentCachedElementData);
+
+		FWidgetDrawElementState& CurrentWidgetState = WidgetDrawStack.Top();
+		check(!CurrentWidgetState.bIsVolatile);
+
+		if (!CurrentWidgetState.CacheHandle.IsValid())
+		{
+			CurrentWidgetState.CacheHandle = CurrentCachedElementData->AddCache(CurrentWidgetState.Widget);
+		}
+
+		return CurrentCachedElementData->AddCachedElement<ElementType>(CurrentWidgetState.CacheHandle, GetClippingManager(), CurrentWidgetState.Widget);
+	}
 
 	TArrayView<FSlateCachedElementData* const> GetCachedElementDataList() const { return MakeArrayView(CachedElementDataList.GetData(), CachedElementDataList.Num()); }
 
diff --git a/Engine/Source/Runtime/SlateCore/Public/Rendering/ElementBatcher.h b/Engine/Source/Runtime/SlateCore/Public/Rendering/ElementBatcher.h
--- a/Engine/Source/Runtime/SlateCore/Public/Rendering/ElementBatcher.h
+++ b/Engine/Source/Runtime/SlateCore/Public/Rendering/ElementBatcher.h
@@ -158,14 +158,10 @@ public:
 
 	SLATECORE_API void ResetData();
 
-	/**
-	* Returns a list of element batches for this window
-	*/
+	/** Returns a list of element batches for this window */
 	const TArray<FSlateRenderBatch>& GetRenderBatches() const { return RenderBatches; }
 
-	/**
-	*
-	*/
+	/** True if stencil buffer / clipping is needed. */
 	SLATECORE_API bool IsStencilClippingRequired() const;
 
 	int32 GetFirstRenderBatchIndex() const { return FirstRenderBatchIndex; }
@@ -178,22 +174,18 @@ public:
 	const FSlateIndexArray& GetFinalIndexData() const { return FinalIndexData; }
 
 	/**
-	* Fills batch data into the actual vertex and index buffer
-	*
-	* @param VertexBuffer	Pointer to the actual memory for the vertex buffer
-	* @param IndexBuffer	Pointer to the actual memory for an index buffer
-	* @param bAbsoluteIndices	Whether to write absolute indices (simplifies draw call setup on RHIs that do not support BaseVertex)
-	*/
+	 * Fills batch data into the actual vertex and index buffer
+	 *
+	 * @param VertexBuffer	Pointer to the actual memory for the vertex buffer
+	 * @param IndexBuffer	Pointer to the actual memory for an index buffer
+	 * @param bAbsoluteIndices	Whether to write absolute indices (simplifies draw call setup on RHIs that do not support BaseVertex)
+	 */
 	SLATECORE_API void FillVertexAndIndexBuffer(uint8* VertexBuffer, uint8* IndexBuffer, bool bAbsoluteIndices);
 
-	/**
-	* Creates rendering data from batched elements
-	*/
-	//SLATECORE_API void CreateRenderBatches(class FElementBatchMap& LayerToElementBatches);
-
+	/** Merges render batches across all elements where possible for final submit to GPU */
 	SLATECORE_API void MergeRenderBatches();
 
-	/** */
+	/** Adds a new render batch to list of batches */
 	FSlateRenderBatch& AddRenderBatch(
 		int32 InLayer,
 		const FShaderParams& InShaderParams,
@@ -204,6 +196,7 @@ public:
 		ESlateBatchDrawFlag InDrawFlags,
 		int8 SceneIndex);
 
+	/** Adds a cached batch, used in retained rendering */
 	void AddCachedBatches(const TSparseArray<FSlateRenderBatch>& InCachedBatches);
 private:
 	void FillBuffersFromNewBatch(FSlateRenderBatch& Batch, FSlateVertexArray& FinalVertices, FSlateIndexArray& FinalIndices);
@@ -279,6 +272,23 @@ private:
 	void AddElementsInternal(const FSlateDrawElementMap& DrawElements, FVector2f ViewportSize);
 	void AddCachedElements(FSlateCachedElementData& CachedElementData, FVector2f ViewportSize);
 
+	/**
+	 * Generates Vertices, Indices, Renderbatches, & associates each of these together correctly
+	 * Attempts to reuse renderbatches across elements if possible. 
+	 * 
+	 * Note: Future more efficient and less generic reuse is something we may consider.
+	 * 
+	 * @param DrawElements - Elements to iterate over
+	 * @param InElementAdder - Functor to add an individual slate draw element
+	 * @param InElementBatchParamCreator - Functor that generates batch params given a slate draw element, used during batch-reuse
+	 * @param InElementBatchReserver - Functor that reserves vertexes and indicies given an element range and list of elements
+	 */
+	template<typename ElementAdder, typename ElementBatchParamCreator, typename ElementBatchReserver>
+	FORCEINLINE void GenerateIndexedVertexBatches(const FSlateDrawElementArray& DrawElements
+		, ElementAdder&& InElementAdder
+		, ElementBatchParamCreator&& InElementBatchParamCreator
+		, ElementBatchReserver&& InElementBatchReserver);
+
 	/** 
 	 * Creates vertices necessary to draw a Quad element 
 	 */
@@ -286,10 +296,9 @@ private:
 	void AddDebugQuadElement( const FSlateDrawElement& DrawElement);
 
 	/** 
-	 * Creates vertices necessary to draw a 3x3 element
+	 * Creates vertices necessary to draw multiple 3x3 elements
 	 */
-	template<ESlateVertexRounding Rounding>
-	void AddBoxElement( const FSlateDrawElement& DrawElement );
+	void AddBoxElements( const FSlateDrawElementArray& DrawElement );
 
 	/** 
 	 * Creates vertices necessary to draw a string (one quad per character)
@@ -315,10 +324,9 @@ private:
 	void AddSplineElement( const FSlateDrawElement& DrawElement );
 
 	/** 
-	 * Creates vertices necessary to draw a series of attached line segments
+	 * Creates vertices necessary to draw a multiple attached line segments
 	 */
-	template<ESlateVertexRounding Rounding>
-	void AddLineElement( const FSlateDrawElement& DrawElement );
+	void AddLineElements(const FSlateDrawElementArray& DrawElements);
 	
 	/** 
 	 * Creates vertices necessary to draw a viewport (just a textured quad)
@@ -338,6 +346,20 @@ private:
 
 	void AddPostProcessPass(const FSlateDrawElement& DrawElement, FVector2f WindowSize);
 
+	FSlateRenderBatch& CreateRenderBatch(
+		int32 Layer,
+		const FShaderParams& ShaderParams,
+		const FSlateShaderResource* InResource,
+		ESlateDrawPrimitive PrimitiveType,
+		ESlateShader ShaderType,
+		ESlateDrawEffect DrawEffects,
+		ESlateBatchDrawFlag DrawFlags,
+		int8 SceneIndex,
+		const FSlateClippingState* ClippingState)
+	{
+		return CreateRenderBatch(BatchData, Layer, ShaderParams, InResource, PrimitiveType, ShaderType, DrawEffects, DrawFlags, SceneIndex, ClippingState);
+	}
+
 	FSlateRenderBatch& CreateRenderBatch(
 		int32 Layer,
 		const FShaderParams& ShaderParams,
@@ -351,6 +373,18 @@ private:
 		return CreateRenderBatch(BatchData, Layer, ShaderParams, InResource, PrimitiveType, ShaderType, DrawEffects, DrawFlags, DrawElement);
 	}
 
+	FSlateRenderBatch& CreateRenderBatch(
+		FSlateBatchData* SlateBatchData,
+		int32 Layer,
+		const FShaderParams& ShaderParams,
+		const FSlateShaderResource* InResource,
+		ESlateDrawPrimitive PrimitiveType,
+		ESlateShader ShaderType,
+		ESlateDrawEffect DrawEffects,
+		ESlateBatchDrawFlag DrawFlags,
+		int8 SceneIndex, 
+		const FSlateClippingState* ClippingState);
+
 	FSlateRenderBatch& CreateRenderBatch(
 		FSlateBatchData* SlateBatchData,
 		int32 Layer,
diff --git a/Engine/Source/Runtime/SlateCore/Public/Rendering/RenderingCommon.h b/Engine/Source/Runtime/SlateCore/Public/Rendering/RenderingCommon.h
--- a/Engine/Source/Runtime/SlateCore/Public/Rendering/RenderingCommon.h
+++ b/Engine/Source/Runtime/SlateCore/Public/Rendering/RenderingCommon.h
@@ -10,7 +10,10 @@
 #include "Input/Reply.h"
 #include "Input/NavigationReply.h"
 #include "Input/PopupMethodReply.h"
+#include "Rendering/DrawElementTypes.h"
 #include "SlateGlobals.h"
+#include <utility>
+
 #include "RenderingCommon.generated.h"
 
 class FSlateInstanceBufferUpdate;
@@ -44,7 +47,7 @@ enum class ESlateDrawPrimitive : uint8
 	LineList,
 	TriangleList,
 };
- 
+
 /**
  * Shader types. NOTE: mirrored in the shader file   
  * If you add a type here you must also implement the proper shader type (TSlateElementPS).  See SlateShaders.h
@@ -218,12 +221,13 @@ struct SLATECORE_API FSlateVertex
 	/** Local size of the element */
 	uint16 PixelSize[2];
 
-	FSlateVertex() {}
-	
+	FSlateVertex() 
+	{}
+
 public:
 
 	template<ESlateVertexRounding Rounding>
-	static FSlateVertex Make(const FSlateRenderTransform& RenderTransform, const FVector2f InLocalPosition, const FVector2f InTexCoord, const FVector2f InTexCoord2, const FColor InColor, const FColor SecondaryColor = FColor())
+	FORCEINLINE static FSlateVertex Make(const FSlateRenderTransform& RenderTransform, const FVector2f InLocalPosition, const FVector2f InTexCoord, const FVector2f InTexCoord2, const FColor InColor, const FColor SecondaryColor = FColor())
 	{
 		FSlateVertex Vertex;
 		Vertex.TexCoords[0] = InTexCoord.X;
@@ -236,7 +240,7 @@ public:
 	}
 
 	template<ESlateVertexRounding Rounding>
-	static FSlateVertex Make(const FSlateRenderTransform& RenderTransform, const FVector2f InLocalPosition, const FVector2f InTexCoord, const FColor& InColor, const FColor SecondaryColor = FColor())
+	FORCEINLINE static FSlateVertex Make(const FSlateRenderTransform& RenderTransform, const FVector2f InLocalPosition, const FVector2f InTexCoord, const FColor& InColor, const FColor SecondaryColor = FColor())
 	{
 		FSlateVertex Vertex;
 		Vertex.TexCoords[0] = InTexCoord.X;
@@ -249,7 +253,7 @@ public:
 	}
 
 	template<ESlateVertexRounding Rounding>
-	static FSlateVertex Make(const FSlateRenderTransform& RenderTransform, const FVector2f InLocalPosition, const FVector4f InTexCoords, const FVector2f InMaterialTexCoords, const FColor InColor, const FColor SecondaryColor = FColor())
+	FORCEINLINE static FSlateVertex Make(const FSlateRenderTransform& RenderTransform, const FVector2f InLocalPosition, const FVector4f InTexCoords, const FVector2f InMaterialTexCoords, const FColor InColor, const FColor SecondaryColor = FColor())
 	{
 		FSlateVertex Vertex;
 		Vertex.TexCoords[0] = InTexCoords.X;
@@ -263,7 +267,7 @@ public:
 	}
 
 	template<ESlateVertexRounding Rounding>
-	static FSlateVertex Make(const FSlateRenderTransform& RenderTransform, const FVector2f InLocalPosition, const FVector2f InLocalSize, float Scale, const FVector4f InTexCoords, const FColor InColor, const FColor SecondaryColor = FColor())
+	FORCEINLINE static FSlateVertex Make(const FSlateRenderTransform& RenderTransform, const FVector2f InLocalPosition, const FVector2f InLocalSize, float Scale, const FVector4f InTexCoords, const FColor InColor, const FColor SecondaryColor = FColor())
 	{
 		FSlateVertex Vertex;
 		Vertex.TexCoords[0] = InTexCoords.X;
@@ -284,12 +288,42 @@ public:
 		ensureMsgf((int32)Vertex.PixelSize[1] == PixelSizeY, TEXT("Conversion of PixelSizeY is bigger than 16. Cast:%d, int16:%d, int32:%d")
 			, (int32)Vertex.PixelSize[1], Vertex.PixelSize[1], PixelSizeY);
 #endif
-		Vertex.MaterialTexCoords = FVector2f(InLocalPosition.X / InLocalSize.X, InLocalPosition.Y / InLocalSize.Y);
 
 		return Vertex;
 	}
+	
+	FORCEINLINE static FSlateVertex Make(const FSlateRenderTransform& RenderTransform, const FVector2f InLocalPosition, const FVector2f InTexCoord, const FVector2f InTexCoord2, const FColor InColor, const FColor SecondaryColor = FColor(), const ESlateVertexRounding InRounding = ESlateVertexRounding::Disabled)
+	{
+		return InRounding == ESlateVertexRounding::Enabled
+			? FSlateVertex::Make<ESlateVertexRounding::Enabled>(RenderTransform, InLocalPosition, InTexCoord, InTexCoord2, InColor, SecondaryColor)
+			: FSlateVertex::Make<ESlateVertexRounding::Disabled>(RenderTransform, InLocalPosition, InTexCoord, InTexCoord2, InColor, SecondaryColor);
+	}
 
-	void SetTexCoords(const FVector4f InTexCoords)
+	
+	FORCEINLINE static FSlateVertex Make(const FSlateRenderTransform& RenderTransform, const FVector2f InLocalPosition, const FVector2f InTexCoord, const FColor& InColor, const FColor SecondaryColor = FColor(), const ESlateVertexRounding InRounding = ESlateVertexRounding::Disabled)
+	{
+		return InRounding == ESlateVertexRounding::Enabled
+			? FSlateVertex::Make<ESlateVertexRounding::Enabled>(RenderTransform, InLocalPosition, InTexCoord, InColor, SecondaryColor)
+			: FSlateVertex::Make<ESlateVertexRounding::Disabled>(RenderTransform, InLocalPosition, InTexCoord, InColor, SecondaryColor);
+	}
+
+	
+	FORCEINLINE static FSlateVertex Make(const FSlateRenderTransform& RenderTransform, const FVector2f InLocalPosition, const FVector4f InTexCoords, const FVector2f InMaterialTexCoords, const FColor InColor, const FColor SecondaryColor = FColor(), const ESlateVertexRounding InRounding = ESlateVertexRounding::Disabled)
+	{
+		return InRounding == ESlateVertexRounding::Enabled
+			? FSlateVertex::Make<ESlateVertexRounding::Enabled>(RenderTransform, InLocalPosition, InTexCoords, InMaterialTexCoords, InColor, SecondaryColor)
+			: FSlateVertex::Make<ESlateVertexRounding::Disabled>(RenderTransform, InLocalPosition, InTexCoords, InMaterialTexCoords, InColor, SecondaryColor);
+	}
+
+	
+	FORCEINLINE static FSlateVertex Make(const FSlateRenderTransform& RenderTransform, const FVector2f InLocalPosition, const FVector2f InLocalSize, float Scale, const FVector4f InTexCoords, const FColor InColor, const FColor SecondaryColor = FColor(), const ESlateVertexRounding InRounding = ESlateVertexRounding::Disabled)
+	{
+		return InRounding == ESlateVertexRounding::Enabled
+			? FSlateVertex::Make<ESlateVertexRounding::Enabled>(RenderTransform, InLocalPosition, InLocalSize, Scale, InTexCoords, InColor, SecondaryColor)
+			: FSlateVertex::Make<ESlateVertexRounding::Disabled>(RenderTransform, InLocalPosition, InLocalSize, Scale, InTexCoords, InColor, SecondaryColor);
+	}
+
+	FORCEINLINE void SetTexCoords(const FVector4f InTexCoords)
 	{
 		TexCoords[0] = InTexCoords.X;
 		TexCoords[1] = InTexCoords.Y;
@@ -297,7 +331,7 @@ public:
 		TexCoords[3] = InTexCoords.W;
 	}
 
-	void SetPosition(const FVector2f InPosition)
+	FORCEINLINE void SetPosition(const FVector2f InPosition)
 	{
 		Position = InPosition;
 	}
@@ -309,7 +343,7 @@ private:
 	{
 		Position = TransformPoint(RenderTransform, InLocalPosition);
 
-		if ( Rounding == ESlateVertexRounding::Enabled )
+		if constexpr ( Rounding == ESlateVertexRounding::Enabled )
 		{
 			Position.X = FMath::RoundToFloat(Position.X);
 			Position.Y = FMath::RoundToFloat(Position.Y);
@@ -322,6 +356,7 @@ private:
 
 template<> struct TIsPODType<FSlateVertex> { enum { Value = true }; };
 static_assert(TIsTriviallyDestructible<FSlateVertex>::Value == true, "FSlateVertex should be trivially destructible");
+static_assert(std::is_trivially_copyable_v<FSlateVertex> == true, "FSlateVertex should be trivially copyable");
 
 /** Stores an aligned rect as shorts. */
 struct FShortRect
@@ -382,6 +417,37 @@ struct FShortRect
 template<> struct TIsPODType<FShortRect> { enum { Value = true }; };
 static_assert(TIsTriviallyDestructible<FShortRect>::Value == true, "FShortRect should be trivially destructible");
 
+namespace UE::Slate
+{
+	template<std::size_t, typename T>
+	using TypeRepeater = T;
+
+	template<typename ElementType, std::size_t... Indices>
+	auto MakeTupleWithNumInner(std::index_sequence<Indices...>)
+	{
+		return TTuple<TypeRepeater<Indices, ElementType>... >();
+	};
+
+	template<typename ElementType, std::size_t Num, typename Indices = std::make_index_sequence<Num>>
+	auto MakeTupleWithNum()
+	{
+		return MakeTupleWithNumInner<ElementType>(Indices{});
+	}
+
+	template<typename IndexType, IndexType...Indices>
+	auto MakeTupleIndiciesInner(std::integer_sequence<IndexType, Indices...>)
+	{
+		return MakeTuple<IndexType>(Indices...);
+	};
+
+	template<typename IndexType, std::size_t Num, typename Indices = std::make_integer_sequence<IndexType, Num>>
+	auto MakeTupleIndicies()
+	{
+		return MakeTupleIndiciesInner(Indices{});
+	};
+};
+
+
 #if STATS
 
 struct FRenderingBufferStatTracker
@@ -482,8 +548,7 @@ struct FSlateDrawElementContainer
 	FSlateDrawElementArray Elements;
 };
 
-enum class EElementType : uint8;
-typedef TMap<EElementType, FSlateDrawElementContainer> FSlateDrawElementMap;
+using FSlateDrawElementMap = decltype(UE::Slate::MakeTupleWithNum<FSlateDrawElementContainer, (uint8)EElementType::ET_Count>());
 
 #else
 
@@ -496,8 +561,7 @@ struct FSlateDrawElementContainer
 	FSlateDrawElementArray Elements;
 };
 
-enum class EElementType : uint8;
-typedef TMap<EElementType, FSlateDrawElementContainer> FSlateDrawElementMap;
+using FSlateDrawElementMap = decltype(UE::Slate::MakeTupleWithNum<FSlateDrawElementContainer, (uint8)EElementType::ET_Count>());
 
 #endif
 
diff --git a/Engine/Source/Runtime/SlateCore/Public/Rendering/SlateRenderBatch.h b/Engine/Source/Runtime/SlateCore/Public/Rendering/SlateRenderBatch.h
--- a/Engine/Source/Runtime/SlateCore/Public/Rendering/SlateRenderBatch.h
+++ b/Engine/Source/Runtime/SlateCore/Public/Rendering/SlateRenderBatch.h
@@ -13,6 +13,55 @@ class FSlateRenderDataHandle;
 
 struct FSlateDataPayload;
 
+
+struct FSlateRenderBatchParams
+{
+	FSlateRenderBatchParams(
+		int32 InLayer,
+		const FShaderParams& InShaderParams,
+		const FSlateShaderResource* InResource,
+		ESlateDrawPrimitive InPrimitiveType,
+		ESlateShader InShaderType,
+		ESlateDrawEffect InDrawEffects,
+		ESlateBatchDrawFlag InDrawFlags,
+		int8 InSceneIndex,
+		const FSlateClippingState* InClippingState)
+		: Layer(InLayer)
+		, ShaderParams(InShaderParams)
+		, Resource(InResource)
+		, PrimitiveType(InPrimitiveType)
+		, ShaderType(InShaderType)
+		, DrawEffects(InDrawEffects)
+		, DrawFlags(InDrawFlags)
+		, SceneIndex(InSceneIndex)
+		, ClippingState(InClippingState)
+	{ 
+	}
+
+	int32 Layer;
+	const FShaderParams& ShaderParams;
+	const FSlateShaderResource* Resource;
+	ESlateDrawPrimitive PrimitiveType;
+	ESlateShader ShaderType;
+	ESlateDrawEffect DrawEffects;
+	ESlateBatchDrawFlag DrawFlags;
+	int8 SceneIndex;
+	const FSlateClippingState* ClippingState;
+
+	bool IsBatchableWith(const FSlateRenderBatchParams& Other) const
+	{
+		return Layer == Other.Layer 
+			&& ShaderParams == Other.ShaderParams 
+			&& Resource == Other.Resource 
+			&& PrimitiveType == Other.PrimitiveType 
+			&& ShaderType == Other.ShaderType 
+			&& DrawEffects == Other.DrawEffects 
+			&& DrawFlags == Other.DrawFlags 
+			&& SceneIndex == Other.SceneIndex
+			&& ClippingState == Other.ClippingState;
+	}
+};
+
 class FSlateRenderBatch
 {
 public:
@@ -53,6 +102,18 @@ public:
 		++NumIndices;
 	}
 
+	void EmplaceVertex(FSlateVertex&& Vertex)
+	{
+		SourceVertices->Emplace(Vertex);
+		++NumVertices;
+	}
+
+	void EmplaceIndex(SlateIndex Index)
+	{
+		SourceIndices->Emplace(Index);
+		++NumIndices;
+	}
+
 	void AddVertices(const TArray<FSlateVertex>& InVertices)
 	{
 		SourceVertices->Append(InVertices);
@@ -65,6 +126,18 @@ public:
 		NumIndices += InIndices.Num();
 	}
 
+	void AddVertices(TArray<FSlateVertex>&& InVertices)
+	{
+		SourceVertices->Append(InVertices);
+		NumVertices += InVertices.Num();
+	}
+
+	void AddIndices(TArray<SlateIndex>&& InIndices)
+	{
+		SourceIndices->Append(InIndices);
+		NumIndices += InIndices.Num();
+	}
+
 	uint32 GetNumVertices() const
 	{
 		return NumVertices;
diff --git a/Engine/Source/Runtime/SlateCore/Public/Rendering/SlateResourceHandle.h b/Engine/Source/Runtime/SlateCore/Public/Rendering/SlateResourceHandle.h
--- a/Engine/Source/Runtime/SlateCore/Public/Rendering/SlateResourceHandle.h
+++ b/Engine/Source/Runtime/SlateCore/Public/Rendering/SlateResourceHandle.h
@@ -18,7 +18,6 @@ class SLATECORE_API FSlateResourceHandle
 	friend class FSlateNullShaderResourceManager;
 public:
 	FSlateResourceHandle();
-	~FSlateResourceHandle();
 
 	/**
 	 * @return true if the handle still points to a valid rendering resource
diff --git a/Engine/Source/Runtime/UMG/Private/Slate/WidgetRenderer.cpp b/Engine/Source/Runtime/UMG/Private/Slate/WidgetRenderer.cpp
--- a/Engine/Source/Runtime/UMG/Private/Slate/WidgetRenderer.cpp
+++ b/Engine/Source/Runtime/UMG/Private/Slate/WidgetRenderer.cpp
@@ -304,6 +304,7 @@ bool FWidgetRenderer::DrawInvalidationRoot(TSharedRef<SVirtualWindow>& VirtualWi
 			ISlate3DRenderer::FScopedAcquireDrawBuffer ScopedDrawBuffer{ *Renderer, bDeferRenderTargetUpdate };
 			FSlateWindowElementList& WindowElementList = ScopedDrawBuffer.GetDrawBuffer().AddWindowElementList(VirtualWindow);
 
+			// Populates cached element data lists using the providied invalidation root
 			ContextCopy.WindowElementList = &WindowElementList;
 			FSlateInvalidationResult Result = Root.PaintInvalidationRoot(ContextCopy);
 
