Fix Asan regression slate element generation

diff --git a/Engine/Source/Runtime/SlateCore/Public/Rendering/SlateRenderBatch.h b/Engine/Source/Runtime/SlateCore/Public/Rendering/SlateRenderBatch.h
--- a/Engine/Source/Runtime/SlateCore/Public/Rendering/SlateRenderBatch.h
+++ b/Engine/Source/Runtime/SlateCore/Public/Rendering/SlateRenderBatch.h
@@ -39,7 +39,7 @@ struct FSlateRenderBatchParams
 	}
 
 	int32 Layer;
-	const FShaderParams& ShaderParams;
+	const FShaderParams ShaderParams;
 	const FSlateShaderResource* Resource;
 	ESlateDrawPrimitive PrimitiveType;
 	ESlateShader ShaderType;
