Avoid useless copy of FSlateRenderBatchParams in GenerateIndexedVertexBatches

#rnx

diff --git a/Engine/Source/Runtime/SlateCore/Private/Rendering/ElementBatcher.cpp b/Engine/Source/Runtime/SlateCore/Private/Rendering/ElementBatcher.cpp
--- a/Engine/Source/Runtime/SlateCore/Private/Rendering/ElementBatcher.cpp
+++ b/Engine/Source/Runtime/SlateCore/Private/Rendering/ElementBatcher.cpp
@@ -646,12 +646,15 @@ void FSlateElementBatcher::GenerateIndexedVertexBatches(
 		}
 
 		// Determine batch range
-		FSlateRenderBatchParams NewBatchParams = InElementBatchParamCreator(DrawElements[Index]);
+		FSlateRenderBatchParams NewBatchParams;
+		InElementBatchParamCreator(DrawElements[Index], NewBatchParams);
+
 		int32 BatchIndexEnd = Index;
 		while (DrawElements.IsValidIndex(++BatchIndexEnd))
 		{
 			const FSlateDrawElement& NextDrawElement = DrawElements[BatchIndexEnd];
-			FSlateRenderBatchParams NextBatchParams = InElementBatchParamCreator(NextDrawElement);
+			FSlateRenderBatchParams NextBatchParams;
+			InElementBatchParamCreator(NextDrawElement, NextBatchParams);
 
 			if (!NextBatchParams.IsBatchableWith(NewBatchParams))
 			{
@@ -1163,7 +1166,7 @@ void FSlateElementBatcher::AddBoxElements(const FSlateDrawElementArray& DrawElem
 		}
 	};
 	
-	auto GenerateBoxElementBatchParams = [&](const FSlateDrawElement& InDrawElement)
+	auto GenerateBoxElementBatchParams = [&](const FSlateDrawElement& InDrawElement, FSlateRenderBatchParams& OutBatchParameters)
 	{
 		const FSlateBoxPayload& DrawElementPayload = InDrawElement.GetDataPayload<FSlateBoxPayload>();
 		const FVector2f LocalSize = InDrawElement.GetLocalSize();
@@ -1171,14 +1174,13 @@ void FSlateElementBatcher::AddBoxElements(const FSlateDrawElementArray& DrawElem
 
 		// Shader type and params
 		ESlateShader ShaderType = ESlateShader::Default;
-		FShaderParams ShaderParams;
 		if (InDrawElement.GetElementType() == EElementType::ET_RoundedBox)
 		{
 			ShaderType = ESlateShader::RoundedBox;
 			const FSlateRoundedBoxPayload& RoundedPayload = InDrawElement.GetDataPayload<FSlateRoundedBoxPayload>();
 
-			ShaderParams.PixelParams = FVector4f(0, RoundedPayload.GetOutlineWeight(), LocalSize.X, LocalSize.Y);//RadiusWeight;
-			ShaderParams.PixelParams2 = RoundedPayload.GetRadius();
+			OutBatchParameters.ShaderParams.PixelParams = FVector4f(0, RoundedPayload.GetOutlineWeight(), LocalSize.X, LocalSize.Y);//RadiusWeight;
+			OutBatchParameters.ShaderParams.PixelParams2 = RoundedPayload.GetRadius();
 		}
 
 		// Shader Resource
@@ -1203,16 +1205,14 @@ void FSlateElementBatcher::AddBoxElements(const FSlateDrawElementArray& DrawElem
 		ESlateBatchDrawFlag DrawFlags = InDrawElement.GetBatchFlags();
 		DrawFlags |= ((bTileHorizontal ? ESlateBatchDrawFlag::TileU : ESlateBatchDrawFlag::None) | (bTileVertical ? ESlateBatchDrawFlag::TileV : ESlateBatchDrawFlag::None));
 
-		return FSlateRenderBatchParams(
-			InDrawElement.GetLayer()
-			, ShaderParams
-			, Resource
-			, ESlateDrawPrimitive::TriangleList
-			, ShaderType
-			, DrawEffects
-			, DrawFlags
-			, InDrawElement.GetSceneIndex()
-			, ResolveClippingState(InDrawElement));
+		OutBatchParameters.Layer = InDrawElement.GetLayer();
+		OutBatchParameters.Resource = Resource;
+		OutBatchParameters.PrimitiveType = ESlateDrawPrimitive::TriangleList;
+		OutBatchParameters.ShaderType = ShaderType;
+		OutBatchParameters.DrawEffects = DrawEffects;
+		OutBatchParameters.DrawFlags = DrawFlags;
+		OutBatchParameters.SceneIndex = InDrawElement.GetSceneIndex();
+		OutBatchParameters.ClippingState = ResolveClippingState(InDrawElement);
 	};
 
 	auto ReserveBoxElementBatch = [&](FSlateRenderBatch& RenderBatch, uint32 InBatchStart, uint32 InBatchEnd)
@@ -2585,14 +2585,13 @@ void FSlateElementBatcher::AddLineElements( const FSlateDrawElementArray& DrawEl
 		}
 	};
 
-	auto GenerateLineElementBatchParams = [&](const FSlateDrawElement& InDrawElement)
+	auto GenerateLineElementBatchParams = [&](const FSlateDrawElement& InDrawElement, FSlateRenderBatchParams& OutBatchParameters)
 	{
 		const FSlateLinePayload& DrawElementPayload = InDrawElement.GetDataPayload<FSlateLinePayload>();
 		const ESlateDrawEffect DrawEffects = InDrawElement.GetDrawEffects();
 
 		// Shader type and params
 		ESlateShader ShaderType = ESlateShader::Default;
-		FShaderParams ShaderParams = FShaderParams();
 		if (DrawElementPayload.IsAntialiased())
 		{
 			// Filter size to use for anti-aliasing.
@@ -2606,7 +2605,7 @@ void FSlateElementBatcher::AddLineElements( const FSlateDrawElementArray& DrawEl
 
 			const float InsideFilterU = (HalfThickness - FilterRadius) / FilteredHalfThickness;
 			const FVector4f PixelShaderParams = FVector4f(InsideFilterU, 0.0f, 0.0f, 0.0f);
-			ShaderParams = FShaderParams::MakePixelShaderParams(PixelShaderParams);
+			OutBatchParameters.ShaderParams = FShaderParams::MakePixelShaderParams(PixelShaderParams);
 			ShaderType = ESlateShader::LineSegment;
 		}
 
@@ -2629,16 +2628,14 @@ void FSlateElementBatcher::AddLineElements( const FSlateDrawElementArray& DrawEl
 			}
 		}
 
-		return FSlateRenderBatchParams(
-			InDrawElement.GetLayer()
-			, ShaderParams
-			, nullptr
-			, PrimitiveType
-			, ShaderType
-			, DrawEffects
-			, ESlateBatchDrawFlag::None
-			, InDrawElement.GetSceneIndex()
-			, ResolveClippingState(InDrawElement));
+		OutBatchParameters.Layer = InDrawElement.GetLayer();
+		OutBatchParameters.Resource = nullptr;
+		OutBatchParameters.PrimitiveType = PrimitiveType;
+		OutBatchParameters.ShaderType = ShaderType;
+		OutBatchParameters.DrawEffects = DrawEffects;
+		OutBatchParameters.DrawFlags = ESlateBatchDrawFlag::None;
+		OutBatchParameters.SceneIndex = InDrawElement.GetSceneIndex();
+		OutBatchParameters.ClippingState = ResolveClippingState(InDrawElement);
 	};
 
 	auto ReserveLineElementBatch = [&](FSlateRenderBatch& RenderBatch, uint32 InBatchStart, uint32 InBatchEnd)
diff --git a/Engine/Source/Runtime/SlateCore/Public/Rendering/SlateRenderBatch.h b/Engine/Source/Runtime/SlateCore/Public/Rendering/SlateRenderBatch.h
--- a/Engine/Source/Runtime/SlateCore/Public/Rendering/SlateRenderBatch.h
+++ b/Engine/Source/Runtime/SlateCore/Public/Rendering/SlateRenderBatch.h
@@ -16,37 +16,15 @@ struct FSlateDataPayload;
 
 struct FSlateRenderBatchParams
 {
-	FSlateRenderBatchParams(
-		int32 InLayer,
-		const FShaderParams& InShaderParams,
-		const FSlateShaderResource* InResource,
-		ESlateDrawPrimitive InPrimitiveType,
-		ESlateShader InShaderType,
-		ESlateDrawEffect InDrawEffects,
-		ESlateBatchDrawFlag InDrawFlags,
-		int8 InSceneIndex,
-		const FSlateClippingState* InClippingState)
-		: Layer(InLayer)
-		, ShaderParams(InShaderParams)
-		, Resource(InResource)
-		, PrimitiveType(InPrimitiveType)
-		, ShaderType(InShaderType)
-		, DrawEffects(InDrawEffects)
-		, DrawFlags(InDrawFlags)
-		, SceneIndex(InSceneIndex)
-		, ClippingState(InClippingState)
-	{ 
-	}
-
 	int32 Layer;
-	const FShaderParams ShaderParams;
-	const FSlateShaderResource* Resource;
+	FShaderParams ShaderParams;
+	const FSlateShaderResource* Resource = nullptr;
 	ESlateDrawPrimitive PrimitiveType;
 	ESlateShader ShaderType;
 	ESlateDrawEffect DrawEffects;
 	ESlateBatchDrawFlag DrawFlags;
 	int8 SceneIndex;
-	const FSlateClippingState* ClippingState;
+	const FSlateClippingState* ClippingState = nullptr;
 
 	bool IsBatchableWith(const FSlateRenderBatchParams& Other) const
 	{
