[SlateCore]
* Moved dll export from types to methods/staticvar

diff --git a/Engine/Source/Runtime/SlateCore/Public/Animation/CurveHandle.h b/Engine/Source/Runtime/SlateCore/Public/Animation/CurveHandle.h
--- a/Engine/Source/Runtime/SlateCore/Public/Animation/CurveHandle.h
+++ b/Engine/Source/Runtime/SlateCore/Public/Animation/CurveHandle.h
@@ -36,7 +36,7 @@ enum class ECurveEaseFunction : uint8
 /**
  * A handle to curve within a curve sequence.
  */
-struct SLATECORE_API FCurveHandle
+struct FCurveHandle
 {
 	/**
 	 * Creates and initializes a curve handle.
@@ -44,7 +44,7 @@ struct SLATECORE_API FCurveHandle
 	 * @param InOwnerSequence The curve sequence that owns this handle.
 	 * @param InCurveIndex The index of this handle.
 	 */
-	FCurveHandle( const struct FCurveSequence* InOwnerSequence = nullptr, int32 InCurveIndex = 0 );
+	SLATECORE_API FCurveHandle( const struct FCurveSequence* InOwnerSequence = nullptr, int32 InCurveIndex = 0 );
 
 public:
 
@@ -54,7 +54,7 @@ public:
 	 * @return Lerp value.
 	 * @see GetLerpLooping
 	 */
-	float GetLerp( ) const;
+	SLATECORE_API float GetLerp( ) const;
 
 	/**
 	 * Checks whether this handle is initialized.
@@ -70,7 +70,7 @@ public:
 public:
 
 	/** Applies animation easing to lerp value */
-	static float ApplyEasing( float Time, ECurveEaseFunction EaseType );
+	static SLATECORE_API float ApplyEasing( float Time, ECurveEaseFunction EaseType );
 
 private:
 
diff --git a/Engine/Source/Runtime/SlateCore/Public/Animation/CurveSequence.h b/Engine/Source/Runtime/SlateCore/Public/Animation/CurveSequence.h
--- a/Engine/Source/Runtime/SlateCore/Public/Animation/CurveSequence.h
+++ b/Engine/Source/Runtime/SlateCore/Public/Animation/CurveSequence.h
@@ -22,7 +22,7 @@ enum class EActiveTimerReturnType : uint8;
  *   FCurveHandle FadeCurve = Sequence.AddCurve( 0.15f, 0.1f );
  *	 Sequence.Play( this->AsShared() );
  */
-struct SLATECORE_API FCurveSequence : public TSharedFromThis<FCurveSequence>
+struct FCurveSequence : public TSharedFromThis<FCurveSequence>
 {
 public:
 
@@ -51,10 +51,10 @@ public:
 	};
 
 	/** Default constructor */
-	FCurveSequence( );
+	SLATECORE_API FCurveSequence( );
 
 	/** Makes sure the active timer is unregistered */
-	~FCurveSequence();
+	SLATECORE_API ~FCurveSequence();
 
 	/**
 	 * Construct by adding a single animation curve to this sequence.  Does not provide access to the curve though.
@@ -64,7 +64,7 @@ public:
 	 * @param InEaseFunction       Easing function to use for this curve.  Defaults to Linear.  Use this to smooth out your animation transitions.
 	 * @return A FCurveHandle that can be used to get the value of this curve after the animation starts playing.
 	 */
-	FCurveSequence( const float InStartTimeSeconds, const float InDurationSeconds, const ECurveEaseFunction InEaseFunction = ECurveEaseFunction::Linear  );
+	SLATECORE_API FCurveSequence( const float InStartTimeSeconds, const float InDurationSeconds, const ECurveEaseFunction InEaseFunction = ECurveEaseFunction::Linear  );
 
 	/**
 	 * Add a new curve at a given time and offset.
@@ -74,7 +74,7 @@ public:
 	 * @param InEaseFunction       Easing function to use for this curve.  Defaults to Linear.  Use this to smooth out your animation transitions.
 	 * @return A FCurveHandle that can be used to get the value of this curve after the animation starts playing.
 	 */
-	FCurveHandle AddCurve( const float InStartTimeSeconds, const float InDurationSeconds, const ECurveEaseFunction InEaseFunction = ECurveEaseFunction::Linear );
+	SLATECORE_API FCurveHandle AddCurve( const float InStartTimeSeconds, const float InDurationSeconds, const ECurveEaseFunction InEaseFunction = ECurveEaseFunction::Linear );
 
 	/**
 	 * Add a new curve relative to the current end of the sequence. Makes stacking easier.
@@ -89,7 +89,7 @@ public:
 	 * @param InDurationSecond     How long this curve lasts.
 	 * @param InEaseFunction       Easing function to use for this curve.  Defaults to Linear.  Use this to smooth out your animation transitions.
 	 */
-	FCurveHandle AddCurveRelative( const float InOffset, const float InDurationSecond, const ECurveEaseFunction InEaseFunction = ECurveEaseFunction::Linear );
+	SLATECORE_API FCurveHandle AddCurveRelative( const float InOffset, const float InDurationSecond, const ECurveEaseFunction InEaseFunction = ECurveEaseFunction::Linear );
 
 	/**
 	 * Start playing this curve sequence. Registers an active timer with the widget being animated.
@@ -99,12 +99,12 @@ public:
 	 * @param StartAtTime The relative time within the animation at which to begin playing (i.e. 0.0f is the beginning).
 	 * @param bRequiresActiveTimer	Whether or not we need to register an active timer on the widget to keep slate ticking while the animation is playing.  If that is not necessary in your use case, you can set it to false for a small performance boost
 	 */
-	void Play( const TSharedRef<SWidget>& InOwnerWidget, bool bPlayLooped = false, const float StartAtTime = 0.0f, bool bRequiresActiveTimer = true);
+	SLATECORE_API void Play( const TSharedRef<SWidget>& InOwnerWidget, bool bPlayLooped = false, const float StartAtTime = 0.0f, bool bRequiresActiveTimer = true);
 
 	/**
 	 * Plays the curve sequence for a ticker, rather than a widget.
 	 */
-	void Play(const FTickerDelegate& InDelegate, bool bPlayLooped = false, const float StartAtTime = 0.0f);
+	SLATECORE_API void Play(const FTickerDelegate& InDelegate, bool bPlayLooped = false, const float StartAtTime = 0.0f);
 
 	/**
 	 * Start playing this curve sequence in reverse. Registers an active timer for the widget using the sequence.
@@ -114,50 +114,50 @@ public:
 	 * @param StartAtTime The relative time within the animation at which to begin playing (i.e. 0.0f is the beginning).
 	 * @param bRequiresActiveTimer	Whether or not we need to register an active timer on the widget to keep slate ticking while the animation is playing.  If that is not necessary in your use case, you can set it to false for a small performance boost
 	 */
-	void PlayReverse( const TSharedRef<SWidget>& InOwnerWidget, bool bPlayLooped = false, const float StartAtTime = 0.0f, bool bRequiresActiveTimer = true);
+	SLATECORE_API void PlayReverse( const TSharedRef<SWidget>& InOwnerWidget, bool bPlayLooped = false, const float StartAtTime = 0.0f, bool bRequiresActiveTimer = true);
 
 	/** Reverse the direction of an in-progress animation */
-	void Reverse( );
+	SLATECORE_API void Reverse( );
 
 	/** Pause this curve sequence. */
-	void Pause();
+	SLATECORE_API void Pause();
 
 	/** Unpause this curve sequence to resume play. */
-	void Resume( );
+	SLATECORE_API void Resume( );
 
 	/** Plays forward if it can, otherwise holds it at the end of the animation, if we play in reverse it will reverse the animation if it can */
-	void PlayRelative(const TSharedRef<SWidget>& InOwnerWidget, bool bForward);
+	SLATECORE_API void PlayRelative(const TSharedRef<SWidget>& InOwnerWidget, bool bForward);
 
 	/**
 	 * Checks whether the sequence is currently playing.
 	 *
 	 * @return true if playing, false otherwise.
 	 */
-	bool IsPlaying( ) const;
+	SLATECORE_API bool IsPlaying( ) const;
 
 	/** @return the current time relative to the beginning of the sequence. */
-	float GetSequenceTime( ) const;
+	SLATECORE_API float GetSequenceTime( ) const;
 
 	/** @return true if the animation is in reverse */
-	bool IsInReverse( ) const;
+	SLATECORE_API bool IsInReverse( ) const;
 
 	/** @return true if the animation is in forward gear */
-	bool IsForward( ) const;
+	SLATECORE_API bool IsForward( ) const;
 
 	/** Jumps immediately to the beginning of the animation sequence */
-	void JumpToStart( );
+	SLATECORE_API void JumpToStart( );
 
 	/** Jumps immediately to the end of the animation sequence */
-	void JumpToEnd( );
+	SLATECORE_API void JumpToEnd( );
 
 	/** Is the sequence at the start? */
-	bool IsAtStart( ) const;
+	SLATECORE_API bool IsAtStart( ) const;
 
 	/** Is the sequence at the end? */
-	bool IsAtEnd( ) const;
+	SLATECORE_API bool IsAtEnd( ) const;
 
 	/** Is the sequence looping? */
-	bool IsLooping() const;
+	SLATECORE_API bool IsLooping() const;
 
 	/**
 	 * For single-curve animations, returns the interpolation alpha for the animation.  If you call this function
@@ -165,19 +165,19 @@ public:
 	 *
 	 * @return A linearly interpolated value between 0 and 1 for this curve.
 	 */
-	float GetLerp() const;
+	SLATECORE_API float GetLerp() const;
 
 	/**
 	 * @param CurveIndex  Index of a curve in the curves array.
 	 *
 	 * @return A curve given the index into the curves array
 	 */
-	const FCurveSequence::FSlateCurve& GetCurve( int32 CurveIndex ) const;
+	SLATECORE_API const FCurveSequence::FSlateCurve& GetCurve( int32 CurveIndex ) const;
 
 protected:
 
 	/** @param InStartTime  when this curve sequence started playing */
-	void SetStartTime( double InStartTime );
+	SLATECORE_API void SetStartTime( double InStartTime );
 
 private:
 	/** Helper to take care of registering the active timer */
diff --git a/Engine/Source/Runtime/SlateCore/Public/Application/SlateApplicationBase.h b/Engine/Source/Runtime/SlateCore/Public/Application/SlateApplicationBase.h
--- a/Engine/Source/Runtime/SlateCore/Public/Application/SlateApplicationBase.h
+++ b/Engine/Source/Runtime/SlateCore/Public/Application/SlateApplicationBase.h
@@ -97,12 +97,12 @@ namespace SlateApplicationDefs
  * This class currently serves a temporary workaround for solving SlateCore dependencies to FSlateApplication.
  * It should probably be removed once FSlateApplication has been refactored into SlateCore.
  */
-class SLATECORE_API FSlateApplicationBase
+class FSlateApplicationBase
 {
 	friend class SWidget;
 public:
 
-	FSlateApplicationBase();
+	SLATECORE_API FSlateApplicationBase();
 	virtual ~FSlateApplicationBase() { }
 
 	/**
@@ -247,16 +247,16 @@ public:
 	 *
 	 * @param OutDisplayMetrics Will contain the display metrics.
 	 */
-	void GetDisplayMetrics(FDisplayMetrics& OutDisplayMetrics);
+	SLATECORE_API void GetDisplayMetrics(FDisplayMetrics& OutDisplayMetrics);
 
 	/**
 	* Gets the application's cached display metrics.
 	*
 	* @param OutDisplayMetrics Will contain the display metrics.
 	*/
-	void GetCachedDisplayMetrics(FDisplayMetrics& OutDisplayMetrics) const;
+	SLATECORE_API void GetCachedDisplayMetrics(FDisplayMetrics& OutDisplayMetrics) const;
 
-	void GetSafeZoneSize(FMargin& SafeZone, const UE::Slate::FDeprecateVector2DParameter& OverrideSize);
+	SLATECORE_API void GetSafeZoneSize(FMargin& SafeZone, const UE::Slate::FDeprecateVector2DParameter& OverrideSize);
 
 	/**
 	 * Get the highest level of window transparency support currently enabled by this application
@@ -296,7 +296,7 @@ protected:
 	 */
 	virtual TSharedPtr< SWidget > GetMouseCaptorImpl() const = 0;
 
-	void GetSafeZoneRatio(FMargin& SafeZoneRatio);
+	SLATECORE_API void GetSafeZoneRatio(FMargin& SafeZoneRatio);
 
 public:
 	/**
@@ -356,7 +356,7 @@ public:
 
 	/** @return a hittesting object that can perform hittests agains widgets. Only certain classes can make use of FHitTesting */
 	friend class FHitTesting;
-	const FHitTesting& GetHitTesting() const;
+	SLATECORE_API const FHitTesting& GetHitTesting() const;
 
 	/** 
 	 * Given the screen-space coordinate of the mouse cursor, searches for a string of widgets that are under the mouse.
@@ -414,7 +414,7 @@ public:
 	 * @return The new title bar widget.
 	 */
 	UE_DEPRECATED(4.26, "This version of MakeWindowTitleBar has been deprecated. Use the version that takes in an FWindowTitleBarArgs parameter instead.")
-	virtual TSharedRef<SWidget> MakeWindowTitleBar(const TSharedRef<SWindow>& Window, const TSharedPtr<SWidget>& CenterContent, EHorizontalAlignment CenterContentAlignment, TSharedPtr<IWindowTitleBar>& OutTitleBar) const;
+	SLATECORE_API virtual TSharedRef<SWidget> MakeWindowTitleBar(const TSharedRef<SWindow>& Window, const TSharedPtr<SWidget>& CenterContent, EHorizontalAlignment CenterContentAlignment, TSharedPtr<IWindowTitleBar>& OutTitleBar) const;
 
 	/**
 	 * Creates a title bar for the specified window.
@@ -488,23 +488,23 @@ public:
 	DECLARE_EVENT_OneParam(FSlateApplicationBase, FOnGlobalInvalidationToggled, bool);
 	FOnGlobalInvalidationToggled& OnGlobalInvalidationToggled() { return OnGlobalInvalidationToggledEvent; }
 
-	void ToggleGlobalInvalidation(bool bIsGlobalInvalidationEnabled);
+	SLATECORE_API void ToggleGlobalInvalidation(bool bIsGlobalInvalidationEnabled);
 
 	/**
 	 * Notifies all invalidation panels that they should invalidate their contents
 	 * Note: this is a very expensive call and should only be done in non-performance critical situations
 	 */
-	void InvalidateAllWidgets(bool bClearResourcesImmediately) const;
+	SLATECORE_API void InvalidateAllWidgets(bool bClearResourcesImmediately) const;
 private:
 	/**
 	 * Implementation for active timer registration. See SWidget::RegisterActiveTimer.
 	 */
-	void RegisterActiveTimer( const TSharedRef<FActiveTimerHandle>& ActiveTimerHandle );
+	SLATECORE_API void RegisterActiveTimer( const TSharedRef<FActiveTimerHandle>& ActiveTimerHandle );
 
 	/**
 	 * Implementation for active timer registration. See SWidget::UnRegisterActiveTimer.
 	 */
-	void UnRegisterActiveTimer( const TSharedRef<FActiveTimerHandle>& ActiveTimerHandle );
+	SLATECORE_API void UnRegisterActiveTimer( const TSharedRef<FActiveTimerHandle>& ActiveTimerHandle );
 
 	/** The list of active timer handles. */
 	TArray<TWeakPtr<FActiveTimerHandle>> ActiveTimerHandles;
@@ -521,14 +521,14 @@ protected:
 	 * Used to determine if any active timer handles are ready to fire.
 	 * Means we need to tick slate even if no user interaction.
 	 */
-	bool AnyActiveTimersArePending();
+	SLATECORE_API bool AnyActiveTimersArePending();
 
 public:
-	const static uint32 CursorPointerIndex;
-	const static uint32 CursorUserIndex;
+	SLATECORE_API const static uint32 CursorPointerIndex;
+	SLATECORE_API const static uint32 CursorUserIndex;
 
 	/** The primary platform user for this application. They will own the mouse and cursor by default. Default ID is 0 */
-	const static FPlatformUserId SlateAppPrimaryPlatformUser;
+	SLATECORE_API const static FPlatformUserId SlateAppPrimaryPlatformUser;
 
 	/**
 	 * Returns the current instance of the application. The application should have been initialized before
@@ -600,9 +600,9 @@ protected:
 	/** Given a window, locate a widget under the cursor in it; returns an invalid path if cursor is not over this window. */
 	virtual FWidgetPath LocateWidgetInWindow(UE::Slate::FDeprecateVector2DParameter ScreenspaceMouseCoordinate, const TSharedRef<SWindow>& Window, bool bIgnoreEnabledStatus, int32 UserIndex) const = 0;
 	
-	void UpdateCustomSafeZone(const FMargin& NewSafeZoneRatio, bool bShouldRecacheMetrics);
+	SLATECORE_API void UpdateCustomSafeZone(const FMargin& NewSafeZoneRatio, bool bShouldRecacheMetrics);
 #if WITH_EDITOR
-	void SwapSafeZoneTypes();
+	SLATECORE_API void SwapSafeZoneTypes();
 #endif
 
 protected:
@@ -616,10 +616,10 @@ protected:
 protected:
 
 	// Holds a pointer to the current application.
-	static TSharedPtr<FSlateApplicationBase> CurrentBaseApplication;
+	static SLATECORE_API TSharedPtr<FSlateApplicationBase> CurrentBaseApplication;
 
 	// Holds a pointer to the platform application.
-	static TSharedPtr<class GenericApplication> PlatformApplication;
+	static SLATECORE_API TSharedPtr<class GenericApplication> PlatformApplication;
 
 	// Caches the application's display metrics
 	FDisplayMetrics CachedDisplayMetrics;
@@ -633,7 +633,7 @@ public:
 	 *
 	 * @return True if Slate is sleeping.
 	 */
-	bool IsSlateAsleep();
+	SLATECORE_API bool IsSlateAsleep();
 
 	TSharedPtr<ICursor> GetPlatformCursor()
 	{
@@ -645,9 +645,9 @@ public:
 		return PlatformApplication;
 	}
 
-	void ResetCustomSafeZone();
-	bool IsCustomSafeZoneSet() const;
-	void SetCustomSafeZone(const FMargin& InSafeZone);
+	SLATECORE_API void ResetCustomSafeZone();
+	SLATECORE_API bool IsCustomSafeZoneSet() const;
+	SLATECORE_API void SetCustomSafeZone(const FMargin& InSafeZone);
 	const FMargin& GetCustomSafeZone() const { return CustomSafeZoneRatio; }
 
 #if WITH_EDITORONLY_DATA
diff --git a/Engine/Source/Runtime/SlateCore/Public/Application/SlateWindowHelper.h b/Engine/Source/Runtime/SlateCore/Public/Application/SlateWindowHelper.h
--- a/Engine/Source/Runtime/SlateCore/Public/Application/SlateWindowHelper.h
+++ b/Engine/Source/Runtime/SlateCore/Public/Application/SlateWindowHelper.h
@@ -13,7 +13,7 @@ class FWidgetPath;
 /**
  * Implements a manager for Slate windows.
  */
-class SLATECORE_API FSlateWindowHelper
+class FSlateWindowHelper
 {
 public:
 
@@ -23,7 +23,7 @@ public:
 	 * @param Windows The collection of windows to reorder.
 	 * @param WindowToBringToFront The window to bring to the front.
 	 */
-	static void ArrangeWindowToFront( TArray<TSharedRef<SWindow>>& Windows, const TSharedRef<SWindow>& WindowToBringToFront );
+	static SLATECORE_API void ArrangeWindowToFront( TArray<TSharedRef<SWindow>>& Windows, const TSharedRef<SWindow>& WindowToBringToFront );
 
 	/**
 	 * Put 'BringMeToFront' at the font of the list of 'WindowsToReorder'.
@@ -33,7 +33,7 @@ public:
 	 * @param Windows An ordered list of windows.
 	 * @param BingMeToFront The window to bring to front.
 	 */
-	static void BringWindowToFront( TArray<TSharedRef<SWindow>>& Windows, const TSharedRef<SWindow>& BringMeToFront );
+	static SLATECORE_API void BringWindowToFront( TArray<TSharedRef<SWindow>>& Windows, const TSharedRef<SWindow>& BringMeToFront );
 
 	/**
 	 * Checks whether any of the given windows overlap the specified work area.
@@ -42,7 +42,7 @@ public:
 	 * @param WorkAreaRect The work area to check.
 	 * @return true if at least one window overlaps with the work area, false otherwise.
 	 */
-	static bool CheckWorkAreaForWindows( const TArray< TSharedRef<SWindow> >& WindowsToSearch, const FSlateRect& WorkAreaRect );
+	static SLATECORE_API bool CheckWorkAreaForWindows( const TArray< TSharedRef<SWindow> >& WindowsToSearch, const FSlateRect& WorkAreaRect );
 
 	/**
 	 * Checks whether the given collection of windows contains the specified window.
@@ -51,7 +51,7 @@ public:
 	 * @param WindowToFind The window to find.
 	 * @return true if the collection contains the window, false otherwise.
 	 */
-	static bool ContainsWindow( const TArray<TSharedRef<SWindow>>& WindowsToSearch, const TSharedRef<SWindow>& WindowToFind );
+	static SLATECORE_API bool ContainsWindow( const TArray<TSharedRef<SWindow>>& WindowsToSearch, const TSharedRef<SWindow>& WindowToFind );
 
 	/**
 	 * Searches for the specified widget and generates a full path to it.
@@ -63,7 +63,7 @@ public:
 	 * @param VisibilityFilter Widgets must have this type of visibility to be included the path.
 	 * @return True if the widget path was found.
 	 */
-	static bool FindPathToWidget( const TArray<TSharedRef<SWindow>>& WindowsToSearch, TSharedRef<const SWidget> InWidget, FWidgetPath& OutWidgetPath, EVisibility VisibilityFilter = EVisibility::Visible );
+	static SLATECORE_API bool FindPathToWidget( const TArray<TSharedRef<SWindow>>& WindowsToSearch, TSharedRef<const SWidget> InWidget, FWidgetPath& OutWidgetPath, EVisibility VisibilityFilter = EVisibility::Visible );
 
 	/**
 	 * Searches the given collection of windows to find the Slate window that corresponds to the specified platform window.
@@ -72,7 +72,7 @@ public:
 	 * @param PlatformWindow The platform window to find the Slate window for.
 	 * @return The corresponding Slate window, or nullptr if it was not found.
 	 */
-	static TSharedPtr<SWindow> FindWindowByPlatformWindow( const TArray<TSharedRef<SWindow>>& WindowsToSearch, const TSharedRef<FGenericWindow>& PlatformWindow );
+	static SLATECORE_API TSharedPtr<SWindow> FindWindowByPlatformWindow( const TArray<TSharedRef<SWindow>>& WindowsToSearch, const TSharedRef<FGenericWindow>& PlatformWindow );
 
 	/**
 	 * Removes the specified window from the given collection of windows.
@@ -80,7 +80,7 @@ public:
 	 * @param Windows The collection of windows.
 	 * @param WindowToRemove The window to remove.
 	 */
-	static void RemoveWindowFromList( TArray<TSharedRef<SWindow>>& Windows, const TSharedRef<SWindow>& WindowToRemove );
+	static SLATECORE_API void RemoveWindowFromList( TArray<TSharedRef<SWindow>>& Windows, const TSharedRef<SWindow>& WindowToRemove );
 
 protected:
 
@@ -90,5 +90,5 @@ protected:
 	 *
 	 * @return The top-most window whose children were re-arranged
 	 */
-	static TSharedRef<SWindow> BringToFrontInParent( const TSharedRef<SWindow>& WindowToBringToFront );
+	static SLATECORE_API TSharedRef<SWindow> BringToFrontInParent( const TSharedRef<SWindow>& WindowToBringToFront );
 };
diff --git a/Engine/Source/Runtime/SlateCore/Public/Application/ThrottleManager.h b/Engine/Source/Runtime/SlateCore/Public/Application/ThrottleManager.h
--- a/Engine/Source/Runtime/SlateCore/Public/Application/ThrottleManager.h
+++ b/Engine/Source/Runtime/SlateCore/Public/Application/ThrottleManager.h
@@ -42,14 +42,14 @@ private:
 /**
  * A class which manages requests to throttle parts of the engine to ensure Slate UI performance                   
  */
-class SLATECORE_API FSlateThrottleManager
+class FSlateThrottleManager
 {
 public:
 
 	/**
 	 * Constructor                   
 	 */
-	FSlateThrottleManager();
+	SLATECORE_API FSlateThrottleManager();
 
 public:
 
@@ -58,7 +58,7 @@ public:
 	 *
 	 * @return A handle to the request to enter responsive mode.  Can only be ended with this request
 	 */
-	FThrottleRequest EnterResponsiveMode();
+	SLATECORE_API FThrottleRequest EnterResponsiveMode();
 
 	/**
 	 * Request to leave responsive mode.  
@@ -66,14 +66,14 @@ public:
 	 *
 	 * @param InHandle	The handle that was created with EnterResponsiveMode
 	 */
-	void LeaveResponsiveMode( FThrottleRequest& InHandle );
+	SLATECORE_API void LeaveResponsiveMode( FThrottleRequest& InHandle );
 
 	/**
 	 * Whether or not we allow expensive tasks which could hurt performance to occur
 	 *
 	 * @return true if we allow expensive tasks, false otherwise
 	 */
-	bool IsAllowingExpensiveTasks() const;
+	SLATECORE_API bool IsAllowingExpensiveTasks() const;
 
 	/**
 	 * Explicitly disable Slate throttling. This is intended to be used for code-driven exemptions
@@ -82,14 +82,14 @@ public:
 	 *
 	 * @param bState true if throttling should be disabled.
 	 */
-	void DisableThrottle(bool bState);
+	SLATECORE_API void DisableThrottle(bool bState);
 
 public:
 
 	/**
 	 * Gets the instance of this manager                   
 	 */
-	static FSlateThrottleManager& Get( );
+	static SLATECORE_API FSlateThrottleManager& Get( );
 
 private:
 	/** CVar variable to check if we allow throttling (int32 for compatibility) */
diff --git a/Engine/Source/Runtime/SlateCore/Public/Brushes/SlateBorderBrush.h b/Engine/Source/Runtime/SlateCore/Public/Brushes/SlateBorderBrush.h
--- a/Engine/Source/Runtime/SlateCore/Public/Brushes/SlateBorderBrush.h
+++ b/Engine/Source/Runtime/SlateCore/Public/Brushes/SlateBorderBrush.h
@@ -10,7 +10,7 @@
  * Similar to FSlateBoxBrush but has no middle and the sides tile instead of stretching.
  * The margin is applied exactly as in FSlateBoxBrush.
  */
-struct SLATECORE_API FSlateBorderBrush
+struct FSlateBorderBrush
 	: public FSlateBrush
 {
 	/**
@@ -118,11 +118,11 @@ struct SLATECORE_API FSlateBorderBrush
 
 };
 
-struct SLATECORE_API FSlateVectorBorderBrush
+struct FSlateVectorBorderBrush
 	: public FSlateBorderBrush
 {
 	FSlateVectorBorderBrush(const FString& InImageName, const FMargin& InMargin, const FLinearColor& InColorAndOpacity = FLinearColor(1, 1, 1, 1))
 		: FSlateBorderBrush(InImageName, InMargin, InColorAndOpacity, ESlateBrushImageType::Vector)
 	{
 	}
-};
\ No newline at end of file
+};
diff --git a/Engine/Source/Runtime/SlateCore/Public/Brushes/SlateBoxBrush.h b/Engine/Source/Runtime/SlateCore/Public/Brushes/SlateBoxBrush.h
--- a/Engine/Source/Runtime/SlateCore/Public/Brushes/SlateBoxBrush.h
+++ b/Engine/Source/Runtime/SlateCore/Public/Brushes/SlateBoxBrush.h
@@ -21,7 +21,7 @@
  *                |__|_______________|__|
  *
  */
-struct SLATECORE_API FSlateBoxBrush
+struct FSlateBoxBrush
 	: public FSlateBrush
 {
 	/**
@@ -162,7 +162,7 @@ struct SLATECORE_API FSlateBoxBrush
 };
 
 
-struct SLATECORE_API FSlateVectorBoxBrush
+struct FSlateVectorBoxBrush
 	: public FSlateBoxBrush
 {
 	FSlateVectorBoxBrush(const FString& InImageName, const UE::Slate::FDeprecateVector2DParameter& ImageSize, const FMargin& InMargin, const FLinearColor& InColorAndOpacity = FLinearColor(1, 1, 1, 1))
diff --git a/Engine/Source/Runtime/SlateCore/Public/Brushes/SlateColorBrush.h b/Engine/Source/Runtime/SlateCore/Public/Brushes/SlateColorBrush.h
--- a/Engine/Source/Runtime/SlateCore/Public/Brushes/SlateColorBrush.h
+++ b/Engine/Source/Runtime/SlateCore/Public/Brushes/SlateColorBrush.h
@@ -9,7 +9,7 @@
 /**
  * A color brush.  Draws a box with no margins that only has a color applied to it                   
  */
-struct SLATECORE_API FSlateColorBrush
+struct FSlateColorBrush
 	: public FSlateBrush
 {
 	/** 
diff --git a/Engine/Source/Runtime/SlateCore/Public/Brushes/SlateDynamicImageBrush.h b/Engine/Source/Runtime/SlateCore/Public/Brushes/SlateDynamicImageBrush.h
--- a/Engine/Source/Runtime/SlateCore/Public/Brushes/SlateDynamicImageBrush.h
+++ b/Engine/Source/Runtime/SlateCore/Public/Brushes/SlateDynamicImageBrush.h
@@ -10,7 +10,7 @@
 /**
  * Ignores the Margin. Just renders the image. Can tile the image instead of stretching.
  */
-struct SLATECORE_API FSlateDynamicImageBrush
+struct FSlateDynamicImageBrush
 	: public FSlateBrush, public TSharedFromThis<FSlateDynamicImageBrush>
 {
 	/**
@@ -87,7 +87,7 @@ struct SLATECORE_API FSlateDynamicImageBrush
 	* @param InTiling			How do we tile if at all?
 	* @param InImageType		The type of image this this is
 	*/
-	static TSharedPtr<FSlateDynamicImageBrush> CreateWithImageData(
+	static SLATECORE_API TSharedPtr<FSlateDynamicImageBrush> CreateWithImageData(
 		const FName InTextureName,
 		const UE::Slate::FDeprecateVector2DParameter& InImageSize,
 		const TArray<uint8>& InImageData,
@@ -98,13 +98,13 @@ struct SLATECORE_API FSlateDynamicImageBrush
 	/**
 	 * Releases the resource when it is safe to do so
 	 */
-	void ReleaseResource();
+	SLATECORE_API void ReleaseResource();
 
 	/** Destructor. */
-	virtual ~FSlateDynamicImageBrush();
+	SLATECORE_API virtual ~FSlateDynamicImageBrush();
 
 private:
-	void ReleaseResourceInternal();
+	SLATECORE_API void ReleaseResourceInternal();
 
 	void InitFromTextureObject(FName InTextureName)
 	{
diff --git a/Engine/Source/Runtime/SlateCore/Public/Brushes/SlateImageBrush.h b/Engine/Source/Runtime/SlateCore/Public/Brushes/SlateImageBrush.h
--- a/Engine/Source/Runtime/SlateCore/Public/Brushes/SlateImageBrush.h
+++ b/Engine/Source/Runtime/SlateCore/Public/Brushes/SlateImageBrush.h
@@ -10,7 +10,7 @@
 /**
  * Ignores the Margin. Just renders the image. Can tile the image instead of stretching.
  */
-struct SLATECORE_API FSlateImageBrush
+struct FSlateImageBrush
 	: public FSlateBrush
 {
 	/**
@@ -82,7 +82,7 @@ struct SLATECORE_API FSlateImageBrush
 
 };
 
-struct SLATECORE_API FSlateVectorImageBrush
+struct FSlateVectorImageBrush
 	: public FSlateImageBrush
 {
 	FSlateVectorImageBrush(const FString& InImageName, const UE::Slate::FDeprecateVector2DParameter& InImageSize, const FSlateColor& InTint, ESlateBrushTileType::Type InTiling = ESlateBrushTileType::NoTile)
@@ -92,4 +92,4 @@ struct SLATECORE_API FSlateVectorImageBrush
 	FSlateVectorImageBrush(const FString& InImageName, const UE::Slate::FDeprecateVector2DParameter& InImageSize, const FLinearColor& InTint = FLinearColor(1, 1, 1, 1), ESlateBrushTileType::Type InTiling = ESlateBrushTileType::NoTile)
 		: FSlateImageBrush(InImageName, InImageSize, InTint, InTiling, ESlateBrushImageType::Vector)
 	{ }
-};
\ No newline at end of file
+};
diff --git a/Engine/Source/Runtime/SlateCore/Public/Brushes/SlateNoResource.h b/Engine/Source/Runtime/SlateCore/Public/Brushes/SlateNoResource.h
--- a/Engine/Source/Runtime/SlateCore/Public/Brushes/SlateNoResource.h
+++ b/Engine/Source/Runtime/SlateCore/Public/Brushes/SlateNoResource.h
@@ -10,7 +10,7 @@
 /**
  * A resource that has no appearance
  */
-struct SLATECORE_API FSlateNoResource
+struct FSlateNoResource
 	: public FSlateBrush
 {
 	/**
@@ -27,7 +27,7 @@ struct SLATECORE_API FSlateNoResource
  * This represents an undefined brush. FSlateNoResource is a valid brush which means don't render anything.  FSlateOptionalBrush means let the widget decide what to do if the brush is unset. 
  * Widgets that don't check for the brush being unspecified will draw nothing
  */
-struct SLATECORE_API FSlateOptionalBrush
+struct FSlateOptionalBrush
 	: public FSlateBrush
 {
 	FSlateOptionalBrush(const UE::Slate::FDeprecateVector2DParameter& InImageSize = FVector2f::ZeroVector)
@@ -35,4 +35,4 @@ struct SLATECORE_API FSlateOptionalBrush
 	{
 		bIsSet = false;
 	}
-};
\ No newline at end of file
+};
diff --git a/Engine/Source/Runtime/SlateCore/Public/Brushes/SlateRoundedBoxBrush.h b/Engine/Source/Runtime/SlateCore/Public/Brushes/SlateRoundedBoxBrush.h
--- a/Engine/Source/Runtime/SlateCore/Public/Brushes/SlateRoundedBoxBrush.h
+++ b/Engine/Source/Runtime/SlateCore/Public/Brushes/SlateRoundedBoxBrush.h
@@ -6,7 +6,7 @@
 #include "Styling/SlateBrush.h"
 #include "Types/SlateVector2.h"
 
-struct SLATECORE_API FSlateRoundedBoxBrush
+struct FSlateRoundedBoxBrush
 	: public FSlateBrush
 {
 
diff --git a/Engine/Source/Runtime/SlateCore/Public/Debugging/SlateCrashReporterHandler.h b/Engine/Source/Runtime/SlateCore/Public/Debugging/SlateCrashReporterHandler.h
--- a/Engine/Source/Runtime/SlateCore/Public/Debugging/SlateCrashReporterHandler.h
+++ b/Engine/Source/Runtime/SlateCore/Public/Debugging/SlateCrashReporterHandler.h
@@ -17,14 +17,14 @@ namespace UE::Slate
 #if UE_WITH_SLATE_CRASHREPORTER
 
 /** Helper object allowing easy tracking of Slate code in crash reporter. */
-class SLATECORE_API FSlateCrashReporterScope
+class FSlateCrashReporterScope
 {
 public:
-	static FName NAME_Paint;
-	static FName NAME_Prepass;
+	static SLATECORE_API FName NAME_Paint;
+	static SLATECORE_API FName NAME_Prepass;
 
-	explicit FSlateCrashReporterScope(const SWidget& Widget, FName Context);
-	~FSlateCrashReporterScope();
+	SLATECORE_API explicit FSlateCrashReporterScope(const SWidget& Widget, FName Context);
+	SLATECORE_API ~FSlateCrashReporterScope();
 
 private:
 	uint8 ThreadId = 0;
@@ -49,4 +49,4 @@ private:
 
 #define UE_SLATE_CRASH_REPORTER_CONTEXT_SCOPE(Widget, Context)
 
-#endif //UE_WITH_SLATE_CRASHREPORTER
\ No newline at end of file
+#endif //UE_WITH_SLATE_CRASHREPORTER
diff --git a/Engine/Source/Runtime/SlateCore/Public/Debugging/SlateDebugging.h b/Engine/Source/Runtime/SlateCore/Public/Debugging/SlateDebugging.h
--- a/Engine/Source/Runtime/SlateCore/Public/Debugging/SlateDebugging.h
+++ b/Engine/Source/Runtime/SlateCore/Public/Debugging/SlateDebugging.h
@@ -88,10 +88,10 @@ enum class ESlateDebuggingNavigationMethod : uint8
 	HitTestGrid
 };
 
-struct SLATECORE_API FSlateDebuggingInputEventArgs
+struct FSlateDebuggingInputEventArgs
 {
 public:
-	FSlateDebuggingInputEventArgs(ESlateDebuggingInputEvent InInputEventType, const FInputEvent* InInputEvent, const FReply& InReply, const TSharedPtr<SWidget>& InHandlerWidget, const FString& InAdditionalContent);
+	SLATECORE_API FSlateDebuggingInputEventArgs(ESlateDebuggingInputEvent InInputEventType, const FInputEvent* InInputEvent, const FReply& InReply, const TSharedPtr<SWidget>& InHandlerWidget, const FString& InAdditionalContent);
 
 	const ESlateDebuggingInputEvent InputEventType;
 	const FInputEvent* InputEvent;
@@ -99,24 +99,24 @@ public:
 	const TSharedPtr<SWidget>& HandlerWidget;
 	const FString& AdditionalContent;
 
-	FText ToText() const;
+	SLATECORE_API FText ToText() const;
 };
 
-struct SLATECORE_API FSlateDebuggingCursorQueryEventArgs
+struct FSlateDebuggingCursorQueryEventArgs
 {
 public:
-	FSlateDebuggingCursorQueryEventArgs(const TSharedPtr<const SWidget>& InWidgetOverridingCursor, const FCursorReply& InReply);
+	SLATECORE_API FSlateDebuggingCursorQueryEventArgs(const TSharedPtr<const SWidget>& InWidgetOverridingCursor, const FCursorReply& InReply);
 
 	const TSharedPtr<const SWidget>& WidgetOverridingCursor;
 	const FCursorReply& Reply;
 
-	FText ToText() const;
+	SLATECORE_API FText ToText() const;
 };
 
-struct SLATECORE_API FSlateDebuggingElementTypeAddedEventArgs
+struct FSlateDebuggingElementTypeAddedEventArgs
 {
 public:
-	FSlateDebuggingElementTypeAddedEventArgs(
+	SLATECORE_API FSlateDebuggingElementTypeAddedEventArgs(
 		const FSlateWindowElementList& InElementList,
 		int32 InElementIndex,
 		EElementType InElementType
@@ -145,10 +145,10 @@ enum class ESlateDebuggingFocusEvent : uint8
 #if WITH_SLATE_DEBUGGING
 
 
-struct SLATECORE_API FSlateDebuggingFocusEventArgs
+struct FSlateDebuggingFocusEventArgs
 {
 public:
-	FSlateDebuggingFocusEventArgs(
+	SLATECORE_API FSlateDebuggingFocusEventArgs(
 		ESlateDebuggingFocusEvent InFocusEventType,
 		const FFocusEvent& InFocusEvent,
 		const FWeakWidgetPath& InOldFocusedWidgetPath,
@@ -164,13 +164,13 @@ public:
 	const FWidgetPath& NewFocusedWidgetPath;
 	const TSharedPtr<SWidget>& NewFocusedWidget;
 
-	FText ToText() const;
+	SLATECORE_API FText ToText() const;
 };
 
-struct SLATECORE_API FSlateDebuggingNavigationEventArgs
+struct FSlateDebuggingNavigationEventArgs
 {
 public:
-	FSlateDebuggingNavigationEventArgs(
+	SLATECORE_API FSlateDebuggingNavigationEventArgs(
 		const FNavigationEvent& InNavigationEvent,
 		const FNavigationReply& InNavigationReply,
 		const FWidgetPath& InNavigationSource,
@@ -184,18 +184,18 @@ public:
 	const TSharedPtr<SWidget>& DestinationWidget;
 	const ESlateDebuggingNavigationMethod NavigationMethod;
 
-	FText ToText() const;
+	SLATECORE_API FText ToText() const;
 };
 
-struct SLATECORE_API FSlateDebuggingExecuteNavigationEventArgs
+struct FSlateDebuggingExecuteNavigationEventArgs
 {
 public:
 };
 
-struct SLATECORE_API FSlateDebuggingWarningEventArgs
+struct FSlateDebuggingWarningEventArgs
 {
 public:
-	FSlateDebuggingWarningEventArgs(
+	SLATECORE_API FSlateDebuggingWarningEventArgs(
 		const FText& InWarning,
 		const TSharedPtr<SWidget>& InOptionalContextWidget
 	);
@@ -203,13 +203,13 @@ public:
 	const FText& Warning;
 	const TSharedPtr<SWidget>& OptionalContextWidget;
 
-	FText ToText() const;
+	SLATECORE_API FText ToText() const;
 };
 
-struct SLATECORE_API FSlateDebuggingMouseCaptureEventArgs
+struct FSlateDebuggingMouseCaptureEventArgs
 {
 public:
-	FSlateDebuggingMouseCaptureEventArgs(
+	SLATECORE_API FSlateDebuggingMouseCaptureEventArgs(
 		bool InCaptured,
 		uint32 InUserIndex,
 		uint32 InPointerIndex,
@@ -221,7 +221,7 @@ public:
 	uint32 PointerIndex;
 	const TSharedPtr<const SWidget>& CaptureWidget;
 
-	FText ToText() const;
+	SLATECORE_API FText ToText() const;
 };
 
 enum class ESlateDebuggingInvalidateRootReason
@@ -238,14 +238,14 @@ SLATECORE_API FString LexToString(ESlateDebuggingInvalidateRootReason Reason);
 SLATECORE_API bool LexTryParseString(ESlateDebuggingInvalidateRootReason& OutMode, const TCHAR* InBuffer);
 SLATECORE_API void LexFromString(ESlateDebuggingInvalidateRootReason& OutMode, const TCHAR* InBuffer);
 
-struct SLATECORE_API FSlateDebuggingInvalidateArgs
+struct FSlateDebuggingInvalidateArgs
 {
-	FSlateDebuggingInvalidateArgs(
+	SLATECORE_API FSlateDebuggingInvalidateArgs(
 		const SWidget* WidgetInvalidated,
 		const SWidget* WidgetInvalidateInvestigator,
 		EInvalidateWidgetReason InvalidateReason);
 
-	FSlateDebuggingInvalidateArgs(
+	SLATECORE_API FSlateDebuggingInvalidateArgs(
 		const SWidget* WidgetInvalidated,
 		const SWidget* WidgetInvalidateInvestigator,
 		ESlateDebuggingInvalidateRootReason InvalidateReason);
@@ -256,10 +256,10 @@ struct SLATECORE_API FSlateDebuggingInvalidateArgs
 	ESlateDebuggingInvalidateRootReason InvalidateInvalidationRootReason;
 };
 
-struct SLATECORE_API FSlateDebuggingWidgetUpdatedEventArgs
+struct FSlateDebuggingWidgetUpdatedEventArgs
 {
 public:
-	FSlateDebuggingWidgetUpdatedEventArgs(
+	SLATECORE_API FSlateDebuggingWidgetUpdatedEventArgs(
 		const SWidget* Widget,
 		EWidgetUpdateFlags UpdateFlags,
 		bool bFromPaint
@@ -271,34 +271,34 @@ public:
 	/** The widget got painted as a side effect of another widget that got painted */
 	bool bFromPaint;
 
-	FText ToText() const;
+	SLATECORE_API FText ToText() const;
 };
 
 /**
  * 
  */
-class SLATECORE_API FSlateDebugging
+class FSlateDebugging
 {
 public:
 	/** Called when a widget begins painting. */
 	DECLARE_MULTICAST_DELEGATE_OneParam(FBeginWindow, const FSlateWindowElementList& /*ElementList*/);
-	static FBeginWindow BeginWindow;
+	static SLATECORE_API FBeginWindow BeginWindow;
 
 	/** Called when a window finishes painting. */
 	DECLARE_MULTICAST_DELEGATE_OneParam(FEndWindow, const FSlateWindowElementList& /*ElementList*/);
-	static FEndWindow EndWindow;
+	static SLATECORE_API FEndWindow EndWindow;
 
 	/** Called just before a widget paints. */
 	DECLARE_MULTICAST_DELEGATE_SixParams(FBeginWidgetPaint, const SWidget* /*Widget*/, const FPaintArgs& /*Args*/, const FGeometry& /*AllottedGeometry*/, const FSlateRect& /*MyCullingRect*/, const FSlateWindowElementList& /*OutDrawElements*/, int32 /*LayerId*/);
-	static FBeginWidgetPaint BeginWidgetPaint;
+	static SLATECORE_API FBeginWidgetPaint BeginWidgetPaint;
 
 	/** Called after a widget finishes painting. */
 	DECLARE_MULTICAST_DELEGATE_ThreeParams(FEndWidgetPaint, const SWidget* /*Widget*/, const FSlateWindowElementList& /*OutDrawElements*/, int32 /*LayerId*/);
-	static FEndWidgetPaint EndWidgetPaint;
+	static SLATECORE_API FEndWidgetPaint EndWidgetPaint;
 
 	/** Called once the window is drawn (including PaintDeferred). It can be used to draw additional debug info on top of the window. */
 	DECLARE_MULTICAST_DELEGATE_FourParams(FPaintDebugElements, const FPaintArgs& /*InArgs*/, const FGeometry& /*InAllottedGeometry*/, FSlateWindowElementList& /*InOutDrawElements*/, int32& /*InOutLayerId*/);
-	static FPaintDebugElements PaintDebugElements;
+	static SLATECORE_API FPaintDebugElements PaintDebugElements;
 
 	/**
 	 * Called as soon as the element is added to the element list.
@@ -306,7 +306,7 @@ public:
 	 */
 	UE_DEPRECATED(5.2, "FSlateDebugging::ElementAdded is deprecated, use FSlateDebugging::ElementTypeAdded instead")
 	DECLARE_MULTICAST_DELEGATE_TwoParams(FDrawElement, const FSlateWindowElementList& /*ElementList*/, int32 /*ElementIndex*/);
-PRAGMA_DISABLE_DEPRECATION_WARNINGS
+SLATECORE_API PRAGMA_DISABLE_DEPRECATION_WARNINGS
 	static FDrawElement ElementAdded;
 PRAGMA_ENABLE_DEPRECATION_WARNINGS
 
@@ -315,46 +315,46 @@ PRAGMA_ENABLE_DEPRECATION_WARNINGS
 	 * NOTE: These elements are not valid until the widget finishes painting, or you can resolve them all after the window finishes painting.
 	 */
 	DECLARE_MULTICAST_DELEGATE_OneParam(FDrawElementType, const FSlateDebuggingElementTypeAddedEventArgs& /*ElementTypeAddedArgs*/);
-	static FDrawElementType ElementTypeAdded;
+	static SLATECORE_API FDrawElementType ElementTypeAdded;
 
 public:
 	/**  */
 	DECLARE_MULTICAST_DELEGATE_OneParam(FWidgetWarningEvent, const FSlateDebuggingWarningEventArgs& /*EventArgs*/);
-	static FWidgetWarningEvent Warning;
-	static void BroadcastWarning(const FText& WarningText, const TSharedPtr<SWidget>& OptionalContextWidget);
+	static SLATECORE_API FWidgetWarningEvent Warning;
+	static SLATECORE_API void BroadcastWarning(const FText& WarningText, const TSharedPtr<SWidget>& OptionalContextWidget);
 
 public:
 	/**  */
 	DECLARE_MULTICAST_DELEGATE_OneParam(FWidgetInputEvent, const FSlateDebuggingInputEventArgs& /*EventArgs*/);
-	static FWidgetInputEvent InputEvent;
+	static SLATECORE_API FWidgetInputEvent InputEvent;
 
-	static void BroadcastInputEvent(ESlateDebuggingInputEvent InputEventType, const FInputEvent* InInputEvent, const FReply& InReply);
-	static void BroadcastInputEvent(ESlateDebuggingInputEvent InputEventType, const FInputEvent* InInputEvent, const TSharedPtr<SWidget>& HandlerWidget);
-	static void BroadcastInputEvent(ESlateDebuggingInputEvent InputEventType, const FInputEvent* InInputEvent, const FReply& InReply, const TSharedPtr<SWidget>& HandlerWidget);
-	static void BroadcastInputEvent(ESlateDebuggingInputEvent InputEventType, const FInputEvent* InInputEvent, const FReply& InReply, const TSharedPtr<SWidget>& HandlerWidget, const FString& AdditionalContent);
-	static void BroadcastInputEvent(ESlateDebuggingInputEvent InputEventType, const FInputEvent* InInputEvent, const FReply& InReply, const TSharedPtr<SWidget>& HandlerWidget, const FName& AdditionalContent);
-	static void BroadcastInputEvent(ESlateDebuggingInputEvent InputEventType, const FInputEvent* InInputEvent, const FReply& InReply, const TSharedPtr<SWidget>& HandlerWidget, const TCHAR AdditionalContent);
-	static void BroadcastNoReplyInputEvent(ESlateDebuggingInputEvent InputEventType, const FInputEvent* InInputEvent, const TSharedPtr<SWidget>& HandlerWidget);
+	static SLATECORE_API void BroadcastInputEvent(ESlateDebuggingInputEvent InputEventType, const FInputEvent* InInputEvent, const FReply& InReply);
+	static SLATECORE_API void BroadcastInputEvent(ESlateDebuggingInputEvent InputEventType, const FInputEvent* InInputEvent, const TSharedPtr<SWidget>& HandlerWidget);
+	static SLATECORE_API void BroadcastInputEvent(ESlateDebuggingInputEvent InputEventType, const FInputEvent* InInputEvent, const FReply& InReply, const TSharedPtr<SWidget>& HandlerWidget);
+	static SLATECORE_API void BroadcastInputEvent(ESlateDebuggingInputEvent InputEventType, const FInputEvent* InInputEvent, const FReply& InReply, const TSharedPtr<SWidget>& HandlerWidget, const FString& AdditionalContent);
+	static SLATECORE_API void BroadcastInputEvent(ESlateDebuggingInputEvent InputEventType, const FInputEvent* InInputEvent, const FReply& InReply, const TSharedPtr<SWidget>& HandlerWidget, const FName& AdditionalContent);
+	static SLATECORE_API void BroadcastInputEvent(ESlateDebuggingInputEvent InputEventType, const FInputEvent* InInputEvent, const FReply& InReply, const TSharedPtr<SWidget>& HandlerWidget, const TCHAR AdditionalContent);
+	static SLATECORE_API void BroadcastNoReplyInputEvent(ESlateDebuggingInputEvent InputEventType, const FInputEvent* InInputEvent, const TSharedPtr<SWidget>& HandlerWidget);
 
 public:
 	/** */
-	struct SLATECORE_API FScopeProcessInputEvent
+	struct FScopeProcessInputEvent
 	{
 		ESlateDebuggingInputEvent InputEvent;
-		FScopeProcessInputEvent(ESlateDebuggingInputEvent InputEvent, const FInputEvent& Event);
-		~FScopeProcessInputEvent();
+		SLATECORE_API FScopeProcessInputEvent(ESlateDebuggingInputEvent InputEvent, const FInputEvent& Event);
+		SLATECORE_API ~FScopeProcessInputEvent();
 	};
 
 	/** */
-	struct SLATECORE_API FScopeRouteInputEvent
+	struct FScopeRouteInputEvent
 	{
 		ESlateDebuggingInputEvent InputEvent;
-		FScopeRouteInputEvent(ESlateDebuggingInputEvent InputEvent, const FName& RoutingType);
-		~FScopeRouteInputEvent();
+		SLATECORE_API FScopeRouteInputEvent(ESlateDebuggingInputEvent InputEvent, const FName& RoutingType);
+		SLATECORE_API ~FScopeRouteInputEvent();
 	};
 
 	/** */
-	static void BroadcastPreProcessInputEvent(ESlateDebuggingInputEvent InputEventType, const TCHAR* InputPrecessorName, bool bHandled);
+	static SLATECORE_API void BroadcastPreProcessInputEvent(ESlateDebuggingInputEvent InputEventType, const TCHAR* InputPrecessorName, bool bHandled);
 
 	/** */
 	struct IWidgetInputRoutingEvent
@@ -367,66 +367,66 @@ public:
 		virtual void OnInputProcessed(ESlateDebuggingInputEvent InputEventType) = 0;
 	};
 
-	static void RegisterWidgetInputRoutingEvent(IWidgetInputRoutingEvent* Event);
-	static void UnregisterWidgetInputRoutingEvent(IWidgetInputRoutingEvent* Event);
+	static SLATECORE_API void RegisterWidgetInputRoutingEvent(IWidgetInputRoutingEvent* Event);
+	static SLATECORE_API void UnregisterWidgetInputRoutingEvent(IWidgetInputRoutingEvent* Event);
 
 public:
 	/**  */
 	DECLARE_MULTICAST_DELEGATE_OneParam(FWidgetFocusEvent, const FSlateDebuggingFocusEventArgs& /*EventArgs*/);
-	static FWidgetFocusEvent FocusEvent;
+	static SLATECORE_API FWidgetFocusEvent FocusEvent;
 
-	static void BroadcastFocusChanging(const FFocusEvent& InFocusEvent, const FWeakWidgetPath& InOldFocusedWidgetPath, const TSharedPtr<SWidget>& InOldFocusedWidget, const FWidgetPath& InNewFocusedWidgetPath, const TSharedPtr<SWidget>& InNewFocusedWidget);
-	static void BroadcastFocusLost(const FFocusEvent& InFocusEvent, const FWeakWidgetPath& InOldFocusedWidgetPath, const TSharedPtr<SWidget>& InOldFocusedWidget, const FWidgetPath& InNewFocusedWidgetPath, const TSharedPtr<SWidget>& InNewFocusedWidget);
-	static void BroadcastFocusReceived(const FFocusEvent& InFocusEvent, const FWeakWidgetPath& InOldFocusedWidgetPath, const TSharedPtr<SWidget>& InOldFocusedWidget, const FWidgetPath& InNewFocusedWidgetPath, const TSharedPtr<SWidget>& InNewFocusedWidget);
+	static SLATECORE_API void BroadcastFocusChanging(const FFocusEvent& InFocusEvent, const FWeakWidgetPath& InOldFocusedWidgetPath, const TSharedPtr<SWidget>& InOldFocusedWidget, const FWidgetPath& InNewFocusedWidgetPath, const TSharedPtr<SWidget>& InNewFocusedWidget);
+	static SLATECORE_API void BroadcastFocusLost(const FFocusEvent& InFocusEvent, const FWeakWidgetPath& InOldFocusedWidgetPath, const TSharedPtr<SWidget>& InOldFocusedWidget, const FWidgetPath& InNewFocusedWidgetPath, const TSharedPtr<SWidget>& InNewFocusedWidget);
+	static SLATECORE_API void BroadcastFocusReceived(const FFocusEvent& InFocusEvent, const FWeakWidgetPath& InOldFocusedWidgetPath, const TSharedPtr<SWidget>& InOldFocusedWidget, const FWidgetPath& InNewFocusedWidgetPath, const TSharedPtr<SWidget>& InNewFocusedWidget);
 
 public:
 	/**  */
 	DECLARE_MULTICAST_DELEGATE_OneParam(FWidgetAttemptNavigationEvent, const FSlateDebuggingNavigationEventArgs& /*EventArgs*/);
-	static FWidgetAttemptNavigationEvent AttemptNavigationEvent;
+	static SLATECORE_API FWidgetAttemptNavigationEvent AttemptNavigationEvent;
 
-	static void BroadcastAttemptNavigation(const FNavigationEvent& InNavigationEvent, const FNavigationReply& InNavigationReply, const FWidgetPath& InNavigationSource, const TSharedPtr<SWidget>& InDestinationWidget, ESlateDebuggingNavigationMethod InNavigationMethod);
+	static SLATECORE_API void BroadcastAttemptNavigation(const FNavigationEvent& InNavigationEvent, const FNavigationReply& InNavigationReply, const FWidgetPath& InNavigationSource, const TSharedPtr<SWidget>& InDestinationWidget, ESlateDebuggingNavigationMethod InNavigationMethod);
 
 	/**  */
 	DECLARE_MULTICAST_DELEGATE_OneParam(FWidgetExecuteNavigationEvent, const FSlateDebuggingExecuteNavigationEventArgs& /*EventArgs*/);
-	static FWidgetExecuteNavigationEvent ExecuteNavigationEvent;
+	static SLATECORE_API FWidgetExecuteNavigationEvent ExecuteNavigationEvent;
 
-	static void BroadcastExecuteNavigation();
+	static SLATECORE_API void BroadcastExecuteNavigation();
 
 public:
 	/**  */
 	DECLARE_MULTICAST_DELEGATE_OneParam(FWidgetMouseCaptureEvent, const FSlateDebuggingMouseCaptureEventArgs& /*EventArgs*/);
-	static FWidgetMouseCaptureEvent MouseCaptureEvent;
+	static SLATECORE_API FWidgetMouseCaptureEvent MouseCaptureEvent;
 
-	static void BroadcastMouseCapture(uint32 UserIndex, uint32 PointerIndex, TSharedPtr<const SWidget> InCapturingWidget);
-	static void BroadcastMouseCaptureLost(uint32 UserIndex, uint32 PointerIndex, TSharedPtr<const SWidget> InWidgetLostCapture);
+	static SLATECORE_API void BroadcastMouseCapture(uint32 UserIndex, uint32 PointerIndex, TSharedPtr<const SWidget> InCapturingWidget);
+	static SLATECORE_API void BroadcastMouseCaptureLost(uint32 UserIndex, uint32 PointerIndex, TSharedPtr<const SWidget> InWidgetLostCapture);
 
 public:
 	DECLARE_MULTICAST_DELEGATE_OneParam(FWidgetCursorQuery, const FSlateDebuggingCursorQueryEventArgs& /*EventArgs*/);
-	static FWidgetCursorQuery CursorChangedEvent;
+	static SLATECORE_API FWidgetCursorQuery CursorChangedEvent;
 
-	static void BroadcastCursorQuery(TSharedPtr<const SWidget> InWidgetOverridingCursor, const FCursorReply& InReply);
+	static SLATECORE_API void BroadcastCursorQuery(TSharedPtr<const SWidget> InWidgetOverridingCursor, const FCursorReply& InReply);
 
 public:
 	DECLARE_MULTICAST_DELEGATE_OneParam(FWidgetInvalidate, const FSlateDebuggingInvalidateArgs& /*EventArgs*/);
-	static FWidgetInvalidate WidgetInvalidateEvent;
+	static SLATECORE_API FWidgetInvalidate WidgetInvalidateEvent;
 
-	static void BroadcastWidgetInvalidate(const SWidget* WidgetInvalidated, const SWidget* WidgetInvalidateInvestigator, EInvalidateWidgetReason InvalidateReason);
-	static void BroadcastInvalidationRootInvalidate(const SWidget* WidgetInvalidated, const SWidget* WidgetInvalidateInvestigator, ESlateDebuggingInvalidateRootReason InvalidateReason);
+	static SLATECORE_API void BroadcastWidgetInvalidate(const SWidget* WidgetInvalidated, const SWidget* WidgetInvalidateInvestigator, EInvalidateWidgetReason InvalidateReason);
+	static SLATECORE_API void BroadcastInvalidationRootInvalidate(const SWidget* WidgetInvalidated, const SWidget* WidgetInvalidateInvestigator, ESlateDebuggingInvalidateRootReason InvalidateReason);
 
 public:
 	DECLARE_MULTICAST_DELEGATE_OneParam(FWidgetUpdatedEvent, const FSlateDebuggingWidgetUpdatedEventArgs& /*Args*/);
-	static FWidgetUpdatedEvent WidgetUpdatedEvent;
+	static SLATECORE_API FWidgetUpdatedEvent WidgetUpdatedEvent;
 
-	static void BroadcastWidgetUpdated(const SWidget* Invalidated, EWidgetUpdateFlags UpdateFlags);
-	static void BroadcastWidgetUpdatedByPaint(const SWidget* Invalidated, EWidgetUpdateFlags UpdateFlags);
+	static SLATECORE_API void BroadcastWidgetUpdated(const SWidget* Invalidated, EWidgetUpdateFlags UpdateFlags);
+	static SLATECORE_API void BroadcastWidgetUpdatedByPaint(const SWidget* Invalidated, EWidgetUpdateFlags UpdateFlags);
 
 public:
 	/**  */
 	DECLARE_MULTICAST_DELEGATE_TwoParams(FUICommandRun, const FName& /*CommandName*/, const FText& /*CommandLabel*/);
-	static FUICommandRun CommandRun;
+	static SLATECORE_API FUICommandRun CommandRun;
 
-	static const TArray<const SWidget*>& GetAllWidgets();
-	static void ExportWidgetList(FStringView Filename);
+	static SLATECORE_API const TArray<const SWidget*>& GetAllWidgets();
+	static SLATECORE_API void ExportWidgetList(FStringView Filename);
 
 private:
 
@@ -444,4 +444,4 @@ private:
 	static TArray<IWidgetInputRoutingEvent*> RoutingEvents;
 };
 
-#endif
\ No newline at end of file
+#endif
diff --git a/Engine/Source/Runtime/SlateCore/Public/Fonts/CompositeFont.h b/Engine/Source/Runtime/SlateCore/Public/Fonts/CompositeFont.h
--- a/Engine/Source/Runtime/SlateCore/Public/Fonts/CompositeFont.h
+++ b/Engine/Source/Runtime/SlateCore/Public/Fonts/CompositeFont.h
@@ -55,7 +55,7 @@ typedef TSharedRef<const FFontFaceData, ESPMode::ThreadSafe> FFontFaceDataConstR
  * Raw font data for a Font Face asset.
  * @note Exists as a struct so it can be shared between its owner asset and the font cache without worrying about UObject lifetimes. 
  */
-struct SLATECORE_API FFontFaceData
+struct FFontFaceData
 {
 public:
 	/** Default constructor */
@@ -124,8 +124,8 @@ public:
 
 private:
 	/** Memory stat tracking */
-	void TrackMemoryUsage() const;
-	void UntrackMemoryUsage() const;
+	SLATECORE_API void TrackMemoryUsage() const;
+	SLATECORE_API void UntrackMemoryUsage() const;
 
 	/** Internal data */
 	TArray<uint8> Data;
@@ -133,74 +133,74 @@ private:
 
 /** Payload data describing an individual font in a typeface. Keep this lean as it's also used as a key! */
 USTRUCT()
-struct SLATECORE_API FFontData
+struct FFontData
 {
 	GENERATED_USTRUCT_BODY()
 
 	/** Default constructor */
-	FFontData();
+	SLATECORE_API FFontData();
 
 	/** Construct the raw data from a font face asset */
-	explicit FFontData(const UObject* const InFontFaceAsset, const int32 InSubFaceIndex = 0);
+	SLATECORE_API explicit FFontData(const UObject* const InFontFaceAsset, const int32 InSubFaceIndex = 0);
 
 	/** Construct the raw data from a filename and the font data attributes */
-	FFontData(FString InFontFilename, const EFontHinting InHinting, const EFontLoadingPolicy InLoadingPolicy, const int32 InSubFaceIndex = 0);
+	SLATECORE_API FFontData(FString InFontFilename, const EFontHinting InHinting, const EFontLoadingPolicy InLoadingPolicy, const int32 InSubFaceIndex = 0);
 
 	/** Is this font data set to a font? (either by filename or by inline data) */
-	bool HasFont() const;
+	SLATECORE_API bool HasFont() const;
 
 	/** Get the filename of the font to use. This may not actually exist on disk in editor builds and we should load the face buffer instead. */
-	const FString& GetFontFilename() const;
+	SLATECORE_API const FString& GetFontFilename() const;
 
 	/** Get the hinting algorithm to use with the font. */
-	EFontHinting GetHinting() const;
+	SLATECORE_API EFontHinting GetHinting() const;
 
 	/** Get the enum controlling how this font should be loaded at runtime. */
-	EFontLoadingPolicy GetLoadingPolicy() const;
+	SLATECORE_API EFontLoadingPolicy GetLoadingPolicy() const;
 
 	/** Get the index of the sub-face that should be used. */
-	int32 GetSubFaceIndex() const;
+	SLATECORE_API int32 GetSubFaceIndex() const;
 
 	/** Set the index of the sub-face that should be used. */
-	void SetSubFaceIndex(const int32 InSubFaceIndex);
+	SLATECORE_API void SetSubFaceIndex(const int32 InSubFaceIndex);
 
 	/** Get the method to use when laying out the font? */
-	EFontLayoutMethod GetLayoutMethod() const;
+	SLATECORE_API EFontLayoutMethod GetLayoutMethod() const;
 
 	/** Returns true if the ascend is overridden. */
-	bool IsAscendOverridden() const;
+	SLATECORE_API bool IsAscendOverridden() const;
 
 	/** Returns the overridden value of the ascend. This value will be used only if IsAscendOverridden returns true. */
-	int32 GetAscendOverriddenValue() const;
+	SLATECORE_API int32 GetAscendOverriddenValue() const;
 
 	/** Returns true if the descend is overridden. */
-	bool IsDescendOverridden() const;
+	SLATECORE_API bool IsDescendOverridden() const;
 
 	/** Returns the overridden value of the descend. This value will be used only if IsDescendOverridden returns true. */
-	int32 GetDescendOverriddenValue() const;
+	SLATECORE_API int32 GetDescendOverriddenValue() const;
 
 	/** Get the data buffer containing the data for the current font face. */
-	FFontFaceDataConstPtr GetFontFaceData() const;
+	SLATECORE_API FFontFaceDataConstPtr GetFontFaceData() const;
 
 	/** Get the font face asset used by this data (if any). */
-	const UObject* GetFontFaceAsset() const;
+	SLATECORE_API const UObject* GetFontFaceAsset() const;
 
 #if WITH_EDITORONLY_DATA
 	/** True if this object contains any legacy data that needs to be upgraded PostLoad by calling the functions below (in order). */
-	bool HasLegacyData() const;
+	SLATECORE_API bool HasLegacyData() const;
 
 	/** Upgrade v1 font data to v2 bulk data. */
-	void ConditionalUpgradeFontDataToBulkData(UObject* InOuter);
+	SLATECORE_API void ConditionalUpgradeFontDataToBulkData(UObject* InOuter);
 
 	/** Upgrade v2 bulk data to v3 font face. */
-	void ConditionalUpgradeBulkDataToFontFace(UObject* InOuter, UClass* InFontFaceClass, const FName InFontFaceName);
+	SLATECORE_API void ConditionalUpgradeBulkDataToFontFace(UObject* InOuter, UClass* InFontFaceClass, const FName InFontFaceName);
 #endif // WITH_EDITORONLY_DATA
 
 	/** Check to see whether this font data is equal to the other font data */
-	bool operator==(const FFontData& Other) const;
+	SLATECORE_API bool operator==(const FFontData& Other) const;
 
 	/** Check to see whether this font data is not equal to the other font data */
-	bool operator!=(const FFontData& Other) const;
+	SLATECORE_API bool operator!=(const FFontData& Other) const;
 
 	/** Get the type hash for this font data */
 	friend inline uint32 GetTypeHash(const FFontData& Key)
@@ -223,7 +223,7 @@ struct SLATECORE_API FFontData
 	}
 
 	/** Handle serialization for this struct. */
-	bool Serialize(FArchive& Ar);
+	SLATECORE_API bool Serialize(FArchive& Ar);
 	friend FArchive& operator<<(FArchive& Ar, FFontData& InFontData)
 	{
 		InFontData.Serialize(Ar);
@@ -231,7 +231,7 @@ struct SLATECORE_API FFontData
 	}
 
 	/** Called by FStandaloneCompositeFont to prevent our objects from being GCd */
-	void AddReferencedObjects(FReferenceCollector& Collector);
+	SLATECORE_API void AddReferencedObjects(FReferenceCollector& Collector);
 
 private:
 	/**
@@ -303,7 +303,7 @@ struct TStructOpsTypeTraits<FFontData> : public TStructOpsTypeTraitsBase2<FFontD
 
 /** A single entry in a typeface */
 USTRUCT()
-struct SLATECORE_API FTypefaceEntry
+struct FTypefaceEntry
 {
 	GENERATED_USTRUCT_BODY()
 
@@ -336,7 +336,7 @@ struct SLATECORE_API FTypefaceEntry
 
 /** Definition for a typeface (a family of fonts) */
 USTRUCT()
-struct SLATECORE_API FTypeface
+struct FTypeface
 {
 	GENERATED_USTRUCT_BODY()
 
@@ -364,7 +364,7 @@ struct SLATECORE_API FTypeface
 };
 
 USTRUCT()
-struct SLATECORE_API FCompositeFallbackFont
+struct FCompositeFallbackFont
 {
 	GENERATED_USTRUCT_BODY()
 
@@ -385,7 +385,7 @@ struct SLATECORE_API FCompositeFallbackFont
 };
 
 USTRUCT()
-struct SLATECORE_API FCompositeSubFont : public FCompositeFallbackFont
+struct FCompositeSubFont : public FCompositeFallbackFont
 {
 	GENERATED_USTRUCT_BODY()
 
@@ -410,7 +410,7 @@ struct SLATECORE_API FCompositeSubFont : public FCompositeFallbackFont
 };
 
 USTRUCT()
-struct SLATECORE_API FCompositeFont
+struct FCompositeFont
 {
 	GENERATED_USTRUCT_BODY()
 
@@ -467,7 +467,7 @@ struct SLATECORE_API FCompositeFont
  * A version of FCompositeFont that should be used when it's not being embedded within another UObject
  * This derives from FGCObject to ensure that the bulk data objects are referenced correctly 
  */
-struct SLATECORE_API FStandaloneCompositeFont : public FCompositeFont, public FGCObject
+struct FStandaloneCompositeFont : public FCompositeFont, public FGCObject
 {
 	/** Default constructor */
 	FStandaloneCompositeFont()
@@ -481,6 +481,6 @@ struct SLATECORE_API FStandaloneCompositeFont : public FCompositeFont, public FG
 	}
 
 	// FGCObject interface
-	virtual void AddReferencedObjects(FReferenceCollector& Collector) override;
-	virtual FString GetReferencerName() const override;
+	SLATECORE_API virtual void AddReferencedObjects(FReferenceCollector& Collector) override;
+	SLATECORE_API virtual FString GetReferencerName() const override;
 };
diff --git a/Engine/Source/Runtime/SlateCore/Public/Fonts/FontBulkData.h b/Engine/Source/Runtime/SlateCore/Public/Fonts/FontBulkData.h
--- a/Engine/Source/Runtime/SlateCore/Public/Fonts/FontBulkData.h
+++ b/Engine/Source/Runtime/SlateCore/Public/Fonts/FontBulkData.h
@@ -8,32 +8,32 @@
 #include "Serialization/BulkData.h"
 #include "FontBulkData.generated.h"
 
-UCLASS()
-class SLATECORE_API UFontBulkData : public UObject
+UCLASS(MinimalAPI)
+class UFontBulkData : public UObject
 {
 	GENERATED_BODY()
 
 public:
 	/** Default constructor */
-	UFontBulkData();
+	SLATECORE_API UFontBulkData();
 
 	/** Construct the bulk font data from the given file */
-	void Initialize(const FString& InFontFilename);
+	SLATECORE_API void Initialize(const FString& InFontFilename);
 
 	/** Construct the bulk font data from the given data */
-	void Initialize(const void* const InFontData, const int64 InFontDataSizeBytes);
+	SLATECORE_API void Initialize(const void* const InFontData, const int64 InFontDataSizeBytes);
 
 	/** Locks the bulk font data and returns a read-only pointer to it */
-	const void* Lock(int64& OutFontDataSizeBytes) const;
+	SLATECORE_API const void* Lock(int64& OutFontDataSizeBytes) const;
 
 	/** Unlock the bulk font data, after which point the pointer returned by Lock no longer is valid */
-	void Unlock();
+	SLATECORE_API void Unlock();
 
 	/** Returns the size of the bulk data in bytes */
-	int64 GetBulkDataSize() const;
+	SLATECORE_API int64 GetBulkDataSize() const;
 
 	// UObject interface
-	virtual void Serialize(FArchive& Ar) override;
+	SLATECORE_API virtual void Serialize(FArchive& Ar) override;
 
 private:
 	/** Internal bulk data */
diff --git a/Engine/Source/Runtime/SlateCore/Public/Fonts/FontCache.h b/Engine/Source/Runtime/SlateCore/Public/Fonts/FontCache.h
--- a/Engine/Source/Runtime/SlateCore/Public/Fonts/FontCache.h
+++ b/Engine/Source/Runtime/SlateCore/Public/Fonts/FontCache.h
@@ -198,7 +198,7 @@ private:
 };
 
 /** Information for rendering a shaped text sequence */
-class SLATECORE_API FShapedGlyphSequence
+class FShapedGlyphSequence
 {
 public:
 	struct FSourceTextRange
@@ -224,11 +224,11 @@ public:
 		, SourceIndicesToGlyphData(FSourceTextRange(0, 0))
 	{ }
 
-	FShapedGlyphSequence(TArray<FShapedGlyphEntry> InGlyphsToRender, const int16 InTextBaseline, const uint16 InMaxTextHeight, const UObject* InFontMaterial, const FFontOutlineSettings& InOutlineSettings, const FSourceTextRange& InSourceTextRange);
-	~FShapedGlyphSequence();
+	SLATECORE_API FShapedGlyphSequence(TArray<FShapedGlyphEntry> InGlyphsToRender, const int16 InTextBaseline, const uint16 InMaxTextHeight, const UObject* InFontMaterial, const FFontOutlineSettings& InOutlineSettings, const FSourceTextRange& InSourceTextRange);
+	SLATECORE_API ~FShapedGlyphSequence();
 
 	/** Get the amount of memory allocated to this sequence */
-	SIZE_T GetAllocatedSize() const;
+	SLATECORE_API SIZE_T GetAllocatedSize() const;
 
 	/** Get the array of glyphs in this sequence. This data will be ordered so that you can iterate and draw left-to-right, which means it will be backwards for right-to-left languages */
 	const TArray<FShapedGlyphEntry>& GetGlyphsToRender() const
@@ -261,20 +261,20 @@ public:
 	}
 
 	/** Check to see whether this glyph sequence is dirty (ie, contains glyphs with invalid font pointers) */
-	bool IsDirty() const;
+	SLATECORE_API bool IsDirty() const;
 
 	/**
 	 * Get the measured width of the entire shaped text
 	 * @return The measured width
 	 */
-	int32 GetMeasuredWidth() const;
+	SLATECORE_API int32 GetMeasuredWidth() const;
 
 	/**
 	 * Get the measured width of the specified range of this shaped text
 	 * @note The indices used here are relative to the start of the text we were shaped from, even if we were only shaped from a sub-section of that text
 	 * @return The measured width, or an unset value if the text couldn't be measured (eg, because you started or ended on a merged ligature, or because the range is out-of-bounds)
 	 */
-	TOptional<int32> GetMeasuredWidth(const int32 InStartIndex, const int32 InEndIndex, const bool InIncludeKerningWithPrecedingGlyph = true) const;
+	SLATECORE_API TOptional<int32> GetMeasuredWidth(const int32 InStartIndex, const int32 InEndIndex, const bool InIncludeKerningWithPrecedingGlyph = true) const;
 
 	/** Return data used by GetGlyphAtOffset */
 	struct FGlyphOffsetResult
@@ -312,38 +312,38 @@ public:
 	 * Get the information for the glyph at the specified position in pixels along the string horizontally
 	 * @return The result data (see FGlyphOffsetResult)
 	 */
-	FGlyphOffsetResult GetGlyphAtOffset(FSlateFontCache& InFontCache, const int32 InHorizontalOffset, const int32 InStartOffset = 0) const;
+	SLATECORE_API FGlyphOffsetResult GetGlyphAtOffset(FSlateFontCache& InFontCache, const int32 InHorizontalOffset, const int32 InStartOffset = 0) const;
 
 	/**
 	 * Get the information for the glyph at the specified position in pixels along the string horizontally
 	 * @note The indices used here are relative to the start of the text we were shaped from, even if we were only shaped from a sub-section of that text
 	 * @return The result data (see FGlyphOffsetResult), or an unset value if we couldn't find the character (eg, because you started or ended on a merged ligature, or because the range is out-of-bounds)
 	 */
-	TOptional<FGlyphOffsetResult> GetGlyphAtOffset(FSlateFontCache& InFontCache, const int32 InStartIndex, const int32 InEndIndex, const int32 InHorizontalOffset, const int32 InStartOffset = 0, const bool InIncludeKerningWithPrecedingGlyph = true) const;
+	SLATECORE_API TOptional<FGlyphOffsetResult> GetGlyphAtOffset(FSlateFontCache& InFontCache, const int32 InStartIndex, const int32 InEndIndex, const int32 InHorizontalOffset, const int32 InStartOffset = 0, const bool InIncludeKerningWithPrecedingGlyph = true) const;
 
 	/**
 	 * Get the kerning value between the given entry and the next entry in the sequence
 	 * @note The index used here is relative to the start of the text we were shaped from, even if we were only shaped from a sub-section of that text
 	 * @return The kerning, or an unset value if we couldn't get the kerning (eg, because you specified a merged ligature, or because the index is out-of-bounds)
 	 */
-	TOptional<int8> GetKerning(const int32 InIndex) const;
+	SLATECORE_API TOptional<int8> GetKerning(const int32 InIndex) const;
 
 	/**
 	 * Get a sub-sequence of the specified range
 	 * @note The indices used here are relative to the start of the text we were shaped from, even if we were only shaped from a sub-section of that text
 	 * @return The sub-sequence, or an null if the sub-sequence couldn't be created (eg, because you started or ended on a merged ligature, or because the range is out-of-bounds)
 	 */
-	FShapedGlyphSequencePtr GetSubSequence(const int32 InStartIndex, const int32 InEndIndex) const;
+	SLATECORE_API FShapedGlyphSequencePtr GetSubSequence(const int32 InStartIndex, const int32 InEndIndex) const;
 
-	void AddReferencedObjects(FReferenceCollector& Collector);
+	SLATECORE_API void AddReferencedObjects(FReferenceCollector& Collector);
 
 private:
 	/** Non-copyable */
-	FShapedGlyphSequence(const FShapedGlyphSequence&);
-	FShapedGlyphSequence& operator=(const FShapedGlyphSequence&);
+	SLATECORE_API FShapedGlyphSequence(const FShapedGlyphSequence&);
+	SLATECORE_API FShapedGlyphSequence& operator=(const FShapedGlyphSequence&);
 
 	/** Helper function to share some common logic between the bound and unbound GetGlyphAtOffset functions */
-	bool HasFoundGlyphAtOffset(FSlateFontCache& InFontCache, const int32 InHorizontalOffset, const FShapedGlyphEntry& InCurrentGlyph, const int32 InCurrentGlyphIndex, int32& InOutCurrentOffset, const FShapedGlyphEntry*& OutMatchedGlyph) const;
+	SLATECORE_API bool HasFoundGlyphAtOffset(FSlateFontCache& InFontCache, const int32 InHorizontalOffset, const FShapedGlyphEntry& InCurrentGlyph, const int32 InCurrentGlyphIndex, int32& InOutCurrentOffset, const FShapedGlyphEntry*& OutMatchedGlyph) const;
 
 	/**
 	 * Enumerate all of the glyphs within the given source index range (enumerates either visually or logically)
@@ -352,8 +352,8 @@ private:
 	 */
 	enum class EEnumerateGlyphsResult : uint8 { EnumerationFailed, EnumerationAborted, EnumerationComplete };
 	typedef TFunctionRef<bool(const FShapedGlyphEntry&, int32)> FForEachShapedGlyphEntryCallback;
-	EEnumerateGlyphsResult EnumerateLogicalGlyphsInSourceRange(const int32 InStartIndex, const int32 InEndIndex, const FForEachShapedGlyphEntryCallback& InGlyphCallback) const;
-	EEnumerateGlyphsResult EnumerateVisualGlyphsInSourceRange(const int32 InStartIndex, const int32 InEndIndex, const FForEachShapedGlyphEntryCallback& InGlyphCallback) const;
+	SLATECORE_API EEnumerateGlyphsResult EnumerateLogicalGlyphsInSourceRange(const int32 InStartIndex, const int32 InEndIndex, const FForEachShapedGlyphEntryCallback& InGlyphCallback) const;
+	SLATECORE_API EEnumerateGlyphsResult EnumerateVisualGlyphsInSourceRange(const int32 InStartIndex, const int32 InEndIndex, const FForEachShapedGlyphEntryCallback& InGlyphCallback) const;
 
 	/** Contains the information needed when performing a reverse look-up from a source index to the corresponding shaped glyph */
 	struct FSourceIndexToGlyphData
@@ -458,7 +458,7 @@ private:
 };
 
 /** Information for rendering one non-shaped character */
-struct SLATECORE_API FCharacterEntry
+struct FCharacterEntry
 {
 	/** The character this entry is for */
 	TCHAR Character = 0;
@@ -506,10 +506,10 @@ struct SLATECORE_API FCharacterEntry
  * Every character indexed by TCHAR could potentially cost a lot of memory of a lot of empty entries are created
  * because characters being used are far apart
  */
-class SLATECORE_API FCharacterList
+class FCharacterList
 {
 public:
-	FCharacterList( const FSlateFontKey& InFontKey, FSlateFontCache& InFontCache );
+	SLATECORE_API FCharacterList( const FSlateFontKey& InFontKey, FSlateFontCache& InFontCache );
 
 	/**
 	 * Gets data about how to render and measure a character.
@@ -520,11 +520,11 @@ public:
 	 * @param MaxFontFallback	The maximum fallback level that can be used when resolving glyphs
 	 * @return				Data about the character
 	 */
-	const FCharacterEntry& GetCharacter(TCHAR Character, const EFontFallback MaxFontFallback);
+	SLATECORE_API const FCharacterEntry& GetCharacter(TCHAR Character, const EFontFallback MaxFontFallback);
 
 #if WITH_EDITORONLY_DATA
 	/** Check to see if our cached data is potentially stale for our font */
-	bool IsStale() const;
+	SLATECORE_API bool IsStale() const;
 #endif	// WITH_EDITORONLY_DATA
 
 	/**
@@ -535,7 +535,7 @@ public:
 	 * @param MaxFontFallback	The maximum fallback level that can be used when resolving glyphs
 	 * @return The kerning value
 	 */
-	int8 GetKerning(TCHAR FirstChar, TCHAR SecondChar, const EFontFallback MaxFontFallback);
+	SLATECORE_API int8 GetKerning(TCHAR FirstChar, TCHAR SecondChar, const EFontFallback MaxFontFallback);
 
 	/**
 	 * Gets a kerning value for a pair of character entries
@@ -544,19 +544,19 @@ public:
 	 * @param SecondCharacterEntry	The second character entry in the pair
 	 * @return The kerning value
 	 */
-	int8 GetKerning( const FCharacterEntry& FirstCharacterEntry, const FCharacterEntry& SecondCharacterEntry );
+	SLATECORE_API int8 GetKerning( const FCharacterEntry& FirstCharacterEntry, const FCharacterEntry& SecondCharacterEntry );
 
 	/**
 	 * @return The global max height for any character in this font
 	 */
-	uint16 GetMaxHeight() const;
+	SLATECORE_API uint16 GetMaxHeight() const;
 
 	/** 
 	 * Returns the baseline for the font used by this character 
 	 *
 	 * @return The offset from the bottom of the max character height to the baseline. Be aware that the value will be negative.
 	 */
-	int16 GetBaseline() const;
+	SLATECORE_API int16 GetBaseline() const;
 
 private:
 	/**
@@ -596,7 +596,7 @@ private:
  * Font caching implementation
  * Caches characters into textures as needed
  */
-class SLATECORE_API FSlateFontCache : public ISlateAtlasProvider, public FSlateFlushableAtlasCache
+class FSlateFontCache : public ISlateAtlasProvider, public FSlateFlushableAtlasCache
 {
 	friend FCharacterList;
 
@@ -607,13 +607,13 @@ public:
 	 * @param InTextureSize The size of the atlas texture
 	 * @param InFontAlas	Platform specific font atlas resource
 	 */
-	FSlateFontCache( TSharedRef<ISlateFontAtlasFactory> InFontAtlasFactory, ESlateTextureAtlasThreadId InOwningThread);
-	virtual ~FSlateFontCache();
+	SLATECORE_API FSlateFontCache( TSharedRef<ISlateFontAtlasFactory> InFontAtlasFactory, ESlateTextureAtlasThreadId InOwningThread);
+	SLATECORE_API virtual ~FSlateFontCache();
 
 	/** ISlateAtlasProvider */
-	virtual int32 GetNumAtlasPages() const override;
-	virtual FSlateShaderResource* GetAtlasPageResource(const int32 InIndex) const override;
-	virtual bool IsAtlasPageResourceAlphaOnly(const int32 InIndex) const override;
+	SLATECORE_API virtual int32 GetNumAtlasPages() const override;
+	SLATECORE_API virtual FSlateShaderResource* GetAtlasPageResource(const int32 InIndex) const override;
+	SLATECORE_API virtual bool IsAtlasPageResourceAlphaOnly(const int32 InIndex) const override;
 #if WITH_ATLAS_DEBUGGING
 	virtual FAtlasSlotInfo GetAtlasSlotInfoAtPosition(FIntPoint InPosition, int32 AtlasIndex) const override { return FAtlasSlotInfo(); }
 #endif
@@ -630,8 +630,8 @@ public:
 	 * @param InBaseDirection		The overall reading direction of the text (see TextBiDi::ComputeBaseDirection). This will affect where some characters (such as brackets and quotes) are placed within the resultant shaped text
 	 * @param InTextShapingMethod	The text shaping method to use
 	 */
-	FShapedGlyphSequenceRef ShapeBidirectionalText( const FString& InText, const FSlateFontInfo &InFontInfo, const float InFontScale, const TextBiDi::ETextDirection InBaseDirection, const ETextShapingMethod InTextShapingMethod ) const;
-	FShapedGlyphSequenceRef ShapeBidirectionalText( const TCHAR* InText, const int32 InTextStart, const int32 InTextLen, const FSlateFontInfo &InFontInfo, const float InFontScale, const TextBiDi::ETextDirection InBaseDirection, const ETextShapingMethod InTextShapingMethod ) const;
+	SLATECORE_API FShapedGlyphSequenceRef ShapeBidirectionalText( const FString& InText, const FSlateFontInfo &InFontInfo, const float InFontScale, const TextBiDi::ETextDirection InBaseDirection, const ETextShapingMethod InTextShapingMethod ) const;
+	SLATECORE_API FShapedGlyphSequenceRef ShapeBidirectionalText( const TCHAR* InText, const int32 InTextStart, const int32 InTextLen, const FSlateFontInfo &InFontInfo, const float InFontScale, const TextBiDi::ETextDirection InBaseDirection, const ETextShapingMethod InTextShapingMethod ) const;
 
 	/** 
 	 * Performs text shaping on the given range of the string using the given font info. Returns you the shaped text sequence to use for text rendering via FSlateDrawElement::MakeShapedText.
@@ -646,13 +646,13 @@ public:
 	 * @param InTextDirection		The reading direction of the text to shape (valid values are LeftToRight or RightToLeft)
 	 * @param InTextShapingMethod	The text shaping method to use
 	 */
-	FShapedGlyphSequenceRef ShapeUnidirectionalText( const FString& InText, const FSlateFontInfo &InFontInfo, const float InFontScale, const TextBiDi::ETextDirection InTextDirection, const ETextShapingMethod InTextShapingMethod ) const;
-	FShapedGlyphSequenceRef ShapeUnidirectionalText( const TCHAR* InText, const int32 InTextStart, const int32 InTextLen, const FSlateFontInfo &InFontInfo, const float InFontScale, const TextBiDi::ETextDirection InTextDirection, const ETextShapingMethod InTextShapingMethod ) const;
+	SLATECORE_API FShapedGlyphSequenceRef ShapeUnidirectionalText( const FString& InText, const FSlateFontInfo &InFontInfo, const float InFontScale, const TextBiDi::ETextDirection InTextDirection, const ETextShapingMethod InTextShapingMethod ) const;
+	SLATECORE_API FShapedGlyphSequenceRef ShapeUnidirectionalText( const TCHAR* InText, const int32 InTextStart, const int32 InTextLen, const FSlateFontInfo &InFontInfo, const float InFontScale, const TextBiDi::ETextDirection InTextDirection, const ETextShapingMethod InTextShapingMethod ) const;
 
 	/**
 	 * Performs text shaping on the overflow glyph sequence for a given font. The overflow sequence is used to replace characters that are clipped
 	 */
-	FShapedGlyphSequenceRef ShapeOverflowEllipsisText(const FSlateFontInfo& InFontInfo, const float InFontScale);
+	SLATECORE_API FShapedGlyphSequenceRef ShapeOverflowEllipsisText(const FSlateFontInfo& InFontInfo, const float InFontScale);
 
 	/** 
 	 * Gets information for how to draw all non-shaped characters in the specified string. Caches characters as they are found
@@ -661,44 +661,44 @@ public:
 	 * @param FontScale			The scale to apply to the font
 	 * @param OutCharacterEntries	Populated array of character entries. Indices of characters in Text match indices in this array
 	 */
-	class FCharacterList& GetCharacterList( const FSlateFontInfo &InFontInfo, float FontScale, const FFontOutlineSettings& InOutlineSettings = FFontOutlineSettings::NoOutline);
+	SLATECORE_API class FCharacterList& GetCharacterList( const FSlateFontInfo &InFontInfo, float FontScale, const FFontOutlineSettings& InOutlineSettings = FFontOutlineSettings::NoOutline);
 
 	/**
 	 * Get the atlas information for the given shaped glyph. This information will be cached if required 
 	 */
-	FShapedGlyphFontAtlasData GetShapedGlyphFontAtlasData( const FShapedGlyphEntry& InShapedGlyph, const FFontOutlineSettings& InOutlineSettings);
+	SLATECORE_API FShapedGlyphFontAtlasData GetShapedGlyphFontAtlasData( const FShapedGlyphEntry& InShapedGlyph, const FFontOutlineSettings& InOutlineSettings);
 
 	/**
 	 * Gets the overflow glyph sequence for a given font. The overflow sequence is used to replace characters that are clipped
 	 */
 	UE_DEPRECATED(5.1, "GetOverflowEllipsisText is known to create dangling pointer. Use FShapedTextCache::FindOrAddOverflowEllipsisText.")
-	FShapedGlyphSequenceRef GetOverflowEllipsisText(const FSlateFontInfo& InFontInfo, const float InFontScale);
+	SLATECORE_API FShapedGlyphSequenceRef GetOverflowEllipsisText(const FSlateFontInfo& InFontInfo, const float InFontScale);
 
 public:
 	/**
 	 * Flush the given object out of the cache
 	 */
-	void FlushObject( const UObject* const InObject );
+	SLATECORE_API void FlushObject( const UObject* const InObject );
 
 	/**
 	 * Flush the given composite font out of the cache
 	 */
-	void FlushCompositeFont(const FCompositeFont& InCompositeFont);
+	SLATECORE_API void FlushCompositeFont(const FCompositeFont& InCompositeFont);
 
 	/** 
 	 * Flush the cache if needed
 	 */
-	bool ConditionalFlushCache();
+	SLATECORE_API bool ConditionalFlushCache();
 
 	/**
 	 * Updates the texture used for rendering
 	 */
-	void UpdateCache();
+	SLATECORE_API void UpdateCache();
 
 	/**
 	 * Releases rendering resources
 	 */
-	void ReleaseResources();
+	SLATECORE_API void ReleaseResources();
 
 	/**
 	 * Event called after releasing the rendering resources in ReleaseResources
@@ -720,7 +720,7 @@ public:
 	 * 
 	 * @return The raw font data
 	 */
-	const FFontData& GetDefaultFontData( const FSlateFontInfo& InFontInfo ) const;
+	SLATECORE_API const FFontData& GetDefaultFontData( const FSlateFontInfo& InFontInfo ) const;
 
 	/**
 	 * Returns the font to use from the typeface associated with the given codepoint
@@ -731,7 +731,7 @@ public:
 	 * 
 	 * @return The raw font data
 	 */
-	const FFontData& GetFontDataForCodepoint( const FSlateFontInfo& InFontInfo, const UTF32CHAR InCodepoint, float& OutScalingFactor ) const;
+	SLATECORE_API const FFontData& GetFontDataForCodepoint( const FSlateFontInfo& InFontInfo, const UTF32CHAR InCodepoint, float& OutScalingFactor ) const;
 
 	/**
 	 * Returns the height of the largest character in the font. 
@@ -741,7 +741,7 @@ public:
 	 * 
 	 * @return The largest character height
 	 */
-	uint16 GetMaxCharacterHeight( const FSlateFontInfo& InFontInfo, float FontScale ) const;
+	SLATECORE_API uint16 GetMaxCharacterHeight( const FSlateFontInfo& InFontInfo, float FontScale ) const;
 
 	/**
 	 * Returns the baseline for the specified font.
@@ -751,7 +751,7 @@ public:
 	 * 
 	 * @return The offset from the bottom of the max character height to the baseline.
 	 */
-	int16 GetBaseline( const FSlateFontInfo& InFontInfo, float FontScale ) const;
+	SLATECORE_API int16 GetBaseline( const FSlateFontInfo& InFontInfo, float FontScale ) const;
 
 	/**
 	 * Get the underline metrics for the specified font.
@@ -761,7 +761,7 @@ public:
 	 * @param OutUnderlinePos		The offset from the baseline to the center of the underline bar
 	 * @param OutUnderlineThickness	The thickness of the underline bar
 	 */
-	void GetUnderlineMetrics( const FSlateFontInfo& InFontInfo, const float FontScale, int16& OutUnderlinePos, int16& OutUnderlineThickness ) const;
+	SLATECORE_API void GetUnderlineMetrics( const FSlateFontInfo& InFontInfo, const float FontScale, int16& OutUnderlinePos, int16& OutUnderlineThickness ) const;
 
 	/**
 	 * Get the strike metrics for the specified font.
@@ -771,7 +771,7 @@ public:
 	 * @param OutStrikeLinePos		The offset from the baseline to the center of the strike bar
 	 * @param OutStrikeLineThickness The thickness of the strike bar
 	 */
-	void GetStrikeMetrics( const FSlateFontInfo& InFontInfo, const float FontScale, int16& OutStrikeLinePos, int16& OutStrikeLineThickness ) const;
+	SLATECORE_API void GetStrikeMetrics( const FSlateFontInfo& InFontInfo, const float FontScale, int16& OutStrikeLinePos, int16& OutStrikeLineThickness ) const;
 
 	/**
 	 * Calculates the kerning amount for a pair of characters
@@ -782,12 +782,12 @@ public:
 	 * @param Second		The second character in the pair
 	 * @return The kerning amount, 0 if no kerning
 	 */
-	int8 GetKerning( const FFontData& InFontData, const int32 InSize, TCHAR First, TCHAR Second, float Scale ) const;
+	SLATECORE_API int8 GetKerning( const FFontData& InFontData, const int32 InSize, TCHAR First, TCHAR Second, float Scale ) const;
 
 	/**
 	 * @return Whether or not the font used has kerning information
 	 */
-	bool HasKerning( const FFontData& InFontData ) const;
+	SLATECORE_API bool HasKerning( const FFontData& InFontData ) const;
 
 	/**
 	 * Whether or not the specified character, within the specified font, can be loaded with the specified maximum font fallback level
@@ -797,7 +797,7 @@ public:
 	 * @param MaxFallbackLevel	The maximum fallback level to try for the font
 	 * @return					Whether or not the character can be loaded
 	 */
-	bool CanLoadCodepoint(const FFontData& InFontData, const UTF32CHAR InCodepoint, EFontFallback MaxFallbackLevel = EFontFallback::FF_NoFallback) const;
+	SLATECORE_API bool CanLoadCodepoint(const FFontData& InFontData, const UTF32CHAR InCodepoint, EFontFallback MaxFallbackLevel = EFontFallback::FF_NoFallback) const;
 
 	/**
 	 * Returns the font attributes for the specified font.
@@ -806,48 +806,48 @@ public:
 	 * 
 	 * @return The font attributes for the specified font.
 	 */
-	const TSet<FName>& GetFontAttributes( const FFontData& InFontData ) const;
+	SLATECORE_API const TSet<FName>& GetFontAttributes( const FFontData& InFontData ) const;
 
 	/**
 	 * Get the available sub-face data from the given font.
 	 * Typically there will only be one face unless this is a TTC/OTC font.
 	 * The index of the returned entry can be passed as InFaceIndex to the FFreeTypeFace constructor.
 	 */
-	TArray<FString> GetAvailableFontSubFaces(FFontFaceDataConstRef InMemory) const;
-	TArray<FString> GetAvailableFontSubFaces(const FString& InFilename) const;
+	SLATECORE_API TArray<FString> GetAvailableFontSubFaces(FFontFaceDataConstRef InMemory) const;
+	SLATECORE_API TArray<FString> GetAvailableFontSubFaces(const FString& InFilename) const;
 
 	/**
 	 * Issues a request to clear all cached data from the cache
 	 */
-	void RequestFlushCache(const FString& FlushReason);
+	SLATECORE_API void RequestFlushCache(const FString& FlushReason);
 
 	/**
 	 * Clears just the cached font data, but leaves the atlases alone
 	 */
-	void FlushData();
+	SLATECORE_API void FlushData();
 
 	/**
 	 * Gets the allocated font face data for a font data asset
 	 */
-	SIZE_T GetFontDataAssetResidentMemory(const UObject* FontDataAsset) const;
+	SLATECORE_API SIZE_T GetFontDataAssetResidentMemory(const UObject* FontDataAsset) const;
 
 private:
 	// Non-copyable
-	FSlateFontCache(const FSlateFontCache&);
-	FSlateFontCache& operator=(const FSlateFontCache&);
+	SLATECORE_API FSlateFontCache(const FSlateFontCache&);
+	SLATECORE_API FSlateFontCache& operator=(const FSlateFontCache&);
 
 	/**
 	 * Clears all cached data from the cache
 	 */
-	bool FlushCache();
+	SLATECORE_API bool FlushCache();
 
 	/**
 	 * Clears out any pending UFont objects that were requested to be flushed
 	 */
-	void FlushFontObjects();
+	SLATECORE_API void FlushFontObjects();
 
 	/** Called after the active culture has changed */
-	void HandleCultureChanged();
+	SLATECORE_API void HandleCultureChanged();
 
 	/**
 	 * Add a new entries into a cache atlas
@@ -857,9 +857,9 @@ private:
 	 * @param FontScale		The font scale to use
 	 * @return true if the characters could be cached. false if the cache is full
 	 */
-	bool AddNewEntry(const FShapedGlyphEntry& InShapedGlyph, const FFontOutlineSettings& InOutlineSettings, FShapedGlyphFontAtlasData& OutAtlasData);
+	SLATECORE_API bool AddNewEntry(const FShapedGlyphEntry& InShapedGlyph, const FFontOutlineSettings& InOutlineSettings, FShapedGlyphFontAtlasData& OutAtlasData);
 
-	bool AddNewEntry(const FCharacterRenderData InRenderData, uint8& OutTextureIndex, uint16& OutGlyphX, uint16& OutGlyphY, uint16& OutGlyphWidth, uint16& OutGlyphHeight);
+	SLATECORE_API bool AddNewEntry(const FCharacterRenderData InRenderData, uint8& OutTextureIndex, uint16& OutGlyphX, uint16& OutGlyphY, uint16& OutGlyphWidth, uint16& OutGlyphHeight);
 private:
 
 	/** FreeType library instance (owned by this font cache) */
diff --git a/Engine/Source/Runtime/SlateCore/Public/Fonts/FontMeasure.h b/Engine/Source/Runtime/SlateCore/Public/Fonts/FontMeasure.h
--- a/Engine/Source/Runtime/SlateCore/Public/Fonts/FontMeasure.h
+++ b/Engine/Source/Runtime/SlateCore/Public/Fonts/FontMeasure.h
@@ -10,11 +10,11 @@ struct FSlateFontMeasureCache;
 typedef class FLRUStringCache FMeasureCache;
 struct FSlateFontMeasureCache;
 
-class SLATECORE_API FSlateFontMeasure final
+class FSlateFontMeasure final
 {
 public:
 
-	static TSharedRef< FSlateFontMeasure > Create( const TSharedRef<class FSlateFontCache>& FontCache );
+	static SLATECORE_API TSharedRef< FSlateFontMeasure > Create( const TSharedRef<class FSlateFontCache>& FontCache );
 
 public:
 
@@ -26,7 +26,7 @@ public:
 	 * @param FontScale	The scale to apply to the font
 	 * @return The width and height of the string.
 	 */
-	UE::Slate::FDeprecateVector2DResult Measure( const FString& Text, const FSlateFontInfo &InFontInfo, float FontScale = 1.0f ) const;
+	SLATECORE_API UE::Slate::FDeprecateVector2DResult Measure( const FString& Text, const FSlateFontInfo &InFontInfo, float FontScale = 1.0f ) const;
 
 	/** 
 	 * Measures the width and height of a passed in text.  The height is the maximum height of the largest character in the font/size pair
@@ -36,7 +36,7 @@ public:
 	 * @param FontScale	The scale to apply to the font
 	 * @return The width and height of the text.
 	 */
-	UE::Slate::FDeprecateVector2DResult Measure( const FText& Text, const FSlateFontInfo &InFontInfo, float FontScale = 1.0f ) const;
+	SLATECORE_API UE::Slate::FDeprecateVector2DResult Measure( const FText& Text, const FSlateFontInfo &InFontInfo, float FontScale = 1.0f ) const;
 	
 	/** 
 	 * Measures the width and height of a passed in string.  The height is the maximum height of the largest character in the font/size pair
@@ -48,7 +48,7 @@ public:
 	 * @param FontScale	The scale to apply to the font
 	 * @return The width and height of the string.
 	 */
-	UE::Slate::FDeprecateVector2DResult Measure( const FString& Text, int32 StartIndex, int32 EndIndex, const FSlateFontInfo &InFontInfo, bool IncludeKerningWithPrecedingChar = true, float FontScale = 1.0f ) const;
+	SLATECORE_API UE::Slate::FDeprecateVector2DResult Measure( const FString& Text, int32 StartIndex, int32 EndIndex, const FSlateFontInfo &InFontInfo, bool IncludeKerningWithPrecedingChar = true, float FontScale = 1.0f ) const;
 
 	/** 
 	 * Finds the last whole character index before the specified position in pixels along the string horizontally
@@ -59,7 +59,7 @@ public:
 	 *
 	 * @return The index of the last whole character before the specified horizontal offset
 	 */
-	int32 FindLastWholeCharacterIndexBeforeOffset( const FString& Text, const FSlateFontInfo& InFontInfo, const int32 HorizontalOffset, float FontScale = 1.0f ) const;
+	SLATECORE_API int32 FindLastWholeCharacterIndexBeforeOffset( const FString& Text, const FSlateFontInfo& InFontInfo, const int32 HorizontalOffset, float FontScale = 1.0f ) const;
 
 	/** 
 	 * Finds the last whole character index before the specified position in pixels along the string horizontally
@@ -70,7 +70,7 @@ public:
 	 *
 	 * @return The index of the last whole character before the specified horizontal offset
 	 */
-	int32 FindLastWholeCharacterIndexBeforeOffset( const FText& Text, const FSlateFontInfo& InFontInfo, const int32 HorizontalOffset, float FontScale = 1.0f ) const;
+	SLATECORE_API int32 FindLastWholeCharacterIndexBeforeOffset( const FText& Text, const FSlateFontInfo& InFontInfo, const int32 HorizontalOffset, float FontScale = 1.0f ) const;
 
 	/** 
 	 * Finds the last whole character index before the specified position in pixels along the text horizontally
@@ -83,7 +83,7 @@ public:
 	 *
 	 * @return The index of the last whole character before the specified horizontal offset
 	 */
-	int32 FindLastWholeCharacterIndexBeforeOffset( const FString& Text, int32 StartIndex, int32 EndIndex, const FSlateFontInfo& InFontInfo, const int32 HorizontalOffset, bool IncludeKerningWithPrecedingChar = true, float FontScale = 1.0f ) const;
+	SLATECORE_API int32 FindLastWholeCharacterIndexBeforeOffset( const FString& Text, int32 StartIndex, int32 EndIndex, const FSlateFontInfo& InFontInfo, const int32 HorizontalOffset, bool IncludeKerningWithPrecedingChar = true, float FontScale = 1.0f ) const;
 
 	/** 
 	 * Finds the first whole character index after the specified position in pixels along the string horizontally
@@ -94,7 +94,7 @@ public:
 	 *
 	 * @return The index of the first whole character after the specified horizontal offset
 	 */
-	int32 FindFirstWholeCharacterIndexAfterOffset( const FString& Text, const FSlateFontInfo& InFontInfo, const int32 HorizontalOffset, float FontScale = 1.0f ) const;
+	SLATECORE_API int32 FindFirstWholeCharacterIndexAfterOffset( const FString& Text, const FSlateFontInfo& InFontInfo, const int32 HorizontalOffset, float FontScale = 1.0f ) const;
 
 	/** 
 	 * Finds the first whole character index after the specified position in pixels along the string horizontally
@@ -105,7 +105,7 @@ public:
 	 *
 	 * @return The index of the first whole character after the specified horizontal offset
 	 */
-	int32 FindFirstWholeCharacterIndexAfterOffset( const FText& Text, const FSlateFontInfo& InFontInfo, const int32 HorizontalOffset, float FontScale = 1.0f ) const;
+	SLATECORE_API int32 FindFirstWholeCharacterIndexAfterOffset( const FText& Text, const FSlateFontInfo& InFontInfo, const int32 HorizontalOffset, float FontScale = 1.0f ) const;
 
 	/** 
 	 * Finds the first whole character index after the specified position in pixels along the string horizontally
@@ -118,7 +118,7 @@ public:
 	 *
 	 * @return The index of the first whole character after the specified horizontal offset
 	 */
-	int32 FindFirstWholeCharacterIndexAfterOffset( const FString& Text, int32 StartIndex, int32 EndIndex, const FSlateFontInfo& InFontInfo, const int32 HorizontalOffset, bool IncludeKerningWithPrecedingChar = true, float FontScale = 1.0f ) const;
+	SLATECORE_API int32 FindFirstWholeCharacterIndexAfterOffset( const FString& Text, int32 StartIndex, int32 EndIndex, const FSlateFontInfo& InFontInfo, const int32 HorizontalOffset, bool IncludeKerningWithPrecedingChar = true, float FontScale = 1.0f ) const;
 
 	/** 
 	 * Finds the character index at the specified position in pixels along the string horizontally
@@ -131,7 +131,7 @@ public:
 	 *
 	 * @return The index of the character closest to the specified horizontal offset
 	 */
-	int32 FindCharacterIndexAtOffset( const FString& Text, const FSlateFontInfo &InFontInfo, const int32 HorizontalOffset, float FontScale = 1.0f ) const;
+	SLATECORE_API int32 FindCharacterIndexAtOffset( const FString& Text, const FSlateFontInfo &InFontInfo, const int32 HorizontalOffset, float FontScale = 1.0f ) const;
 
 	/** 
 	 * Finds the character index at the specified position in pixels along the text horizontally
@@ -144,7 +144,7 @@ public:
 	 *
 	 * @return The index of the character closest to the specified horizontal offset
 	 */
-	int32 FindCharacterIndexAtOffset( const FText& Text, const FSlateFontInfo &InFontInfo, const int32 HorizontalOffset, float FontScale = 1.0f ) const;
+	SLATECORE_API int32 FindCharacterIndexAtOffset( const FText& Text, const FSlateFontInfo &InFontInfo, const int32 HorizontalOffset, float FontScale = 1.0f ) const;
 
 	/** 
 	 * Finds the character index at the specified position in pixels along the string horizontally
@@ -159,7 +159,7 @@ public:
 	 *
 	 * @return The index of the character closest to the specified horizontal offset
 	 */
-	int32 FindCharacterIndexAtOffset( const FString& Text, int32 StartIndex, int32 EndIndex, const FSlateFontInfo &InFontInfo, const int32 HorizontalOffset, bool IncludeKerningWithPrecedingChar = true, float FontScale = 1.0f ) const;
+	SLATECORE_API int32 FindCharacterIndexAtOffset( const FString& Text, int32 StartIndex, int32 EndIndex, const FSlateFontInfo &InFontInfo, const int32 HorizontalOffset, bool IncludeKerningWithPrecedingChar = true, float FontScale = 1.0f ) const;
 
 	/**
 	 * Returns the height of the largest character in the font. 
@@ -169,7 +169,7 @@ public:
 	 * 
 	 * @return The largest character height
 	 */
-	uint16 GetMaxCharacterHeight( const FSlateFontInfo& InFontInfo, float FontScale = 1.0f ) const;
+	SLATECORE_API uint16 GetMaxCharacterHeight( const FSlateFontInfo& InFontInfo, float FontScale = 1.0f ) const;
 
 	/**
 	 * Returns the kerning value for the specified pair of characters.
@@ -181,7 +181,7 @@ public:
 	 *
 	 * @return the kerning value used between the two specified characters
 	 */
-	int8 GetKerning(const FSlateFontInfo& InFontInfo, float FontScale, TCHAR PreviousCharacter, TCHAR CurrentCharacter) const;
+	SLATECORE_API int8 GetKerning(const FSlateFontInfo& InFontInfo, float FontScale, TCHAR PreviousCharacter, TCHAR CurrentCharacter) const;
 
 	/**
 	 * Returns the baseline for the specified font.
@@ -191,9 +191,9 @@ public:
 	 * 
 	 * @return The offset from the bottom of the max character height to the baseline.
 	 */
-	int16 GetBaseline( const FSlateFontInfo& InFontInfo, float FontScale = 1.0f ) const;
+	SLATECORE_API int16 GetBaseline( const FSlateFontInfo& InFontInfo, float FontScale = 1.0f ) const;
 
-	void FlushCache();
+	SLATECORE_API void FlushCache();
 
 
 private:
diff --git a/Engine/Source/Runtime/SlateCore/Public/Fonts/FontTypes.h b/Engine/Source/Runtime/SlateCore/Public/Fonts/FontTypes.h
--- a/Engine/Source/Runtime/SlateCore/Public/Fonts/FontTypes.h
+++ b/Engine/Source/Runtime/SlateCore/Public/Fonts/FontTypes.h
@@ -18,7 +18,7 @@ class FSlateShaderResource;
 DECLARE_MULTICAST_DELEGATE_OneParam(FOnReleaseFontResources, const class FSlateFontCache&);
 
 
-struct SLATECORE_API FSlateFontKey
+struct FSlateFontKey
 {
 public:
 	FSlateFontKey( const FSlateFontInfo& InInfo, const FFontOutlineSettings& InFontOutlineSettings, const float InScale )
@@ -145,28 +145,28 @@ public:
 /** 
  * Representation of a texture for fonts in which characters are packed tightly based on their bounding rectangle 
  */
-class SLATECORE_API FSlateFontAtlas : public ISlateFontTexture, public FSlateTextureAtlas
+class FSlateFontAtlas : public ISlateFontTexture, public FSlateTextureAtlas
 {
 public:
-	FSlateFontAtlas(uint32 InWidth, uint32 InHeight, const bool InIsGrayscale);
-	virtual ~FSlateFontAtlas();
+	SLATECORE_API FSlateFontAtlas(uint32 InWidth, uint32 InHeight, const bool InIsGrayscale);
+	SLATECORE_API virtual ~FSlateFontAtlas();
 
 	//~ ISlateFontTexture interface
-	virtual bool IsGrayscale() const override final;
+	SLATECORE_API virtual bool IsGrayscale() const override final;
 	virtual FSlateShaderResource* GetAtlasTexture() const override { return GetSlateTexture(); }
 	virtual void ReleaseRenderingResources() { ReleaseResources(); }
 
 	/**
 	 * Flushes all cached data.
 	 */
-	void Flush();
+	SLATECORE_API void Flush();
 
 	/** 
 	 * Adds a character to the texture.
 	 *
 	 * @param CharInfo	Information about the size of the character
 	 */
-	const struct FAtlasedTextureSlot* AddCharacter( const FCharacterRenderData& CharInfo );
+	SLATECORE_API const struct FAtlasedTextureSlot* AddCharacter( const FCharacterRenderData& CharInfo );
 };
 
 class ISlateFontAtlasFactory
diff --git a/Engine/Source/Runtime/SlateCore/Public/Fonts/UnicodeBlockRange.h b/Engine/Source/Runtime/SlateCore/Public/Fonts/UnicodeBlockRange.h
--- a/Engine/Source/Runtime/SlateCore/Public/Fonts/UnicodeBlockRange.h
+++ b/Engine/Source/Runtime/SlateCore/Public/Fonts/UnicodeBlockRange.h
@@ -17,7 +17,7 @@ enum class EUnicodeBlockRange : uint16
 };
 
 /** Pre-defined Unicode block ranges that can be used with the character ranges in sub-fonts */
-struct SLATECORE_API FUnicodeBlockRange
+struct FUnicodeBlockRange
 {
 	/** Index enum of this block */
 	EUnicodeBlockRange Index;
@@ -29,8 +29,8 @@ struct SLATECORE_API FUnicodeBlockRange
 	FInt32Range Range;
 
 	/** Returns an array containing all of the pre-defined block ranges */
-	static TArrayView<const FUnicodeBlockRange> GetUnicodeBlockRanges();
+	static SLATECORE_API TArrayView<const FUnicodeBlockRange> GetUnicodeBlockRanges();
 
 	/** Returns the block corresponding to the given enum */
-	static FUnicodeBlockRange GetUnicodeBlockRange(const EUnicodeBlockRange InBlockIndex);
+	static SLATECORE_API FUnicodeBlockRange GetUnicodeBlockRange(const EUnicodeBlockRange InBlockIndex);
 };
