[SlateCore]

diff --git a/Engine/Source/Runtime/SlateCore/Public/Input/DragAndDrop.h b/Engine/Source/Runtime/SlateCore/Public/Input/DragAndDrop.h
--- a/Engine/Source/Runtime/SlateCore/Public/Input/DragAndDrop.h
+++ b/Engine/Source/Runtime/SlateCore/Public/Input/DragAndDrop.h
@@ -16,7 +16,7 @@ class SWindow;
  * a drag drop operation which keeps a decorator window (optionally)
  * Implement derived types with DRAG_DROP_OPERATOR_TYPE (see below)
  */
-class SLATECORE_API FDragDropOperation
+class FDragDropOperation
 	: public TSharedFromThis<FDragDropOperation>
 {
 public:
@@ -24,12 +24,12 @@ public:
 	/**
 	 * Default constructor.
 	 */
-	FDragDropOperation();
+	SLATECORE_API FDragDropOperation();
 
 	/**
 	 * Destructor.
 	 */
-	virtual ~FDragDropOperation( );
+	SLATECORE_API virtual ~FDragDropOperation( );
 
 	/** Check if this drag and drop operation can cast safely to the specified template type */
 	template<class TType> bool IsOfType() const 
@@ -51,17 +51,17 @@ public:
 	 * @param bDropWasHandled   true when the drop was handled by some widget; false otherwise
 	 * @param MouseEvent        The mouse event which caused the on drop to be called.
 	 */
-	virtual void OnDrop( bool bDropWasHandled, const FPointerEvent& MouseEvent );
+	SLATECORE_API virtual void OnDrop( bool bDropWasHandled, const FPointerEvent& MouseEvent );
 
 	/** 
 	 * Called when the mouse was moved during a drag and drop operation
 	 *
 	 * @param DragDropEvent    The event that describes this drag drop operation.
 	 */
-	virtual void OnDragged( const class FDragDropEvent& DragDropEvent );
+	SLATECORE_API virtual void OnDragged( const class FDragDropEvent& DragDropEvent );
 	
 	/** Allows drag/drop operations to override the current cursor */
-	virtual FCursorReply OnCursorQuery();
+	SLATECORE_API virtual FCursorReply OnCursorQuery();
 
 	/**
 	 * Gets the widget that will serve as the decorator unless overridden. 
@@ -73,7 +73,7 @@ public:
 	virtual FVector2D GetDecoratorPosition() const { return FVector2D(0, 0); }
 	
 	/** Alters the visibility of the window */
-	virtual void SetDecoratorVisibility(bool bVisible);
+	SLATECORE_API virtual void SetDecoratorVisibility(bool bVisible);
 
 	/** Is this drag Drop operation going to interact with applications outside of Slate */
 	virtual bool IsExternalOperation() const { return false; }
@@ -86,7 +86,7 @@ public:
 	 * control can give temporary feedback, for example - a slashed circle telling 
 	 * the user 'this can't be dropped here'.
 	 */
-	void SetCursorOverride( TOptional<EMouseCursor::Type> CursorType );
+	SLATECORE_API void SetCursorOverride( TOptional<EMouseCursor::Type> CursorType );
 
 	/**
 	 * Checks whether this drag and drop operation can cast safely to the specified type.
@@ -100,17 +100,17 @@ protected:
 	/**
 	 * Constructs the window and widget if applicable.
 	 */
-	virtual void Construct();
+	SLATECORE_API virtual void Construct();
 
 	/**
 	 * Creates a window for the cursor decorator.
 	 */
-	void CreateCursorDecoratorWindow();
+	SLATECORE_API void CreateCursorDecoratorWindow();
 
 	/**
 	 * Destroys the cursor decorator window.
 	 */
-	void DestroyCursorDecoratorWindow();
+	SLATECORE_API void DestroyCursorDecoratorWindow();
 
 	virtual TSharedPtr<FDragDropOperation> ConvertTo(const FString& TypeId)
 	{
@@ -213,7 +213,7 @@ DECLARE_DELEGATE_OneParam( FOnDragDropUpdate,
  * An external drag and drop operation that originates outside of slate.
  * E.g. an OLE drag and drop.
  */
-class SLATECORE_API FExternalDragOperation : public FDragDropOperation
+class FExternalDragOperation : public FDragDropOperation
 {
 private:
 	/** A private constructor to ensure that the appropriate "New" factory method is used below */
@@ -225,11 +225,11 @@ public:
 	DRAG_DROP_OPERATOR_TYPE(FExternalDragOperation, FDragDropOperation)
 
 	/** Creates a new external text drag operation */
-	static TSharedRef<FExternalDragOperation> NewText( FString InText );
+	static SLATECORE_API TSharedRef<FExternalDragOperation> NewText( FString InText );
 	/** Creates a new external file drag operation */
-	static TSharedRef<FExternalDragOperation> NewFiles( TArray<FString> InFileNames );
+	static SLATECORE_API TSharedRef<FExternalDragOperation> NewFiles( TArray<FString> InFileNames );
 	/** Creates a new external combined drag operation */
-	static TSharedRef<FExternalDragOperation> NewOperation( FString InText, TArray<FString> InFileNames );
+	static SLATECORE_API TSharedRef<FExternalDragOperation> NewOperation( FString InText, TArray<FString> InFileNames );
 
 	/** @return true if this is a text drag operation */
 	bool HasText() const {return !!(DragType & DragText);}
@@ -254,7 +254,7 @@ private:
 };
 
 
-class SLATECORE_API FGameDragDropOperation : public FDragDropOperation
+class FGameDragDropOperation : public FDragDropOperation
 {
 public:
 	DRAG_DROP_OPERATOR_TYPE(FGameDragDropOperation, FDragDropOperation)
@@ -262,9 +262,9 @@ public:
 	/**
 	* Default constructor.
 	*/
-	FGameDragDropOperation();
+	SLATECORE_API FGameDragDropOperation();
 
-	virtual FVector2D GetDecoratorPosition() const override;
+	SLATECORE_API virtual FVector2D GetDecoratorPosition() const override;
 
 protected:
 	/** The absolute position of the decorator. */
diff --git a/Engine/Source/Runtime/SlateCore/Public/Input/HittestGrid.h b/Engine/Source/Runtime/SlateCore/Public/Input/HittestGrid.h
--- a/Engine/Source/Runtime/SlateCore/Public/Input/HittestGrid.h
+++ b/Engine/Source/Runtime/SlateCore/Public/Input/HittestGrid.h
@@ -24,17 +24,17 @@ public:
 	virtual TOptional<FVirtualPointerPosition> TranslateMouseCoordinateForCustomHitTestChild(const SWidget& ChildWidget, const FGeometry& MyGeometry, const FVector2D ScreenSpaceMouseCoordinate, const FVector2D LastScreenSpaceMouseCoordinate) const = 0;
 };
 
-class SLATECORE_API FHittestGrid : public FNoncopyable
+class FHittestGrid : public FNoncopyable
 {
 public:
 
-	FHittestGrid();
+	SLATECORE_API FHittestGrid();
 
 	/**
 	 * Given a Slate Units coordinate in virtual desktop space, perform a hittest
 	 * and return the path along which the corresponding event would be bubbled.
 	 */
-	TArray<FWidgetAndPointer> GetBubblePath(UE::Slate::FDeprecateVector2DParameter DesktopSpaceCoordinate, float CursorRadius, bool bIgnoreEnabledStatus, int32 UserIndex = INDEX_NONE);
+	SLATECORE_API TArray<FWidgetAndPointer> GetBubblePath(UE::Slate::FDeprecateVector2DParameter DesktopSpaceCoordinate, float CursorRadius, bool bIgnoreEnabledStatus, int32 UserIndex = INDEX_NONE);
 
 	/**
 	 * Set the position and size of the hittest area in desktop coordinates
@@ -44,14 +44,14 @@ public:
 	 *
 	 * @return      Returns true if a clear of the hittest grid was required. 
 	 */
-	bool SetHittestArea(const UE::Slate::FDeprecateVector2DParameter& HittestPositionInDesktop, const UE::Slate::FDeprecateVector2DParameter& HittestDimensions, const UE::Slate::FDeprecateVector2DParameter& HitestOffsetInWindow = FVector2f::ZeroVector);
+	SLATECORE_API bool SetHittestArea(const UE::Slate::FDeprecateVector2DParameter& HittestPositionInDesktop, const UE::Slate::FDeprecateVector2DParameter& HittestDimensions, const UE::Slate::FDeprecateVector2DParameter& HitestOffsetInWindow = FVector2f::ZeroVector);
 
 	/** Insert custom hit test data for a widget already in the grid */
 	UE_DEPRECATED(5.0, "Deprecated. Use the InsertCustomHitTestPath with a pointer.")
-	void InsertCustomHitTestPath(const TSharedRef<SWidget> InWidget, TSharedRef<ICustomHitTestPath> CustomHitTestPath);
+	SLATECORE_API void InsertCustomHitTestPath(const TSharedRef<SWidget> InWidget, TSharedRef<ICustomHitTestPath> CustomHitTestPath);
 
 	/** Insert custom hit test data for a widget already in the grid */
-	void InsertCustomHitTestPath(const SWidget* InWidget, const TSharedRef<ICustomHitTestPath>& CustomHitTestPath);
+	SLATECORE_API void InsertCustomHitTestPath(const SWidget* InWidget, const TSharedRef<ICustomHitTestPath>& CustomHitTestPath);
 
 	/** Sets the current slate user index that should be associated with any added widgets */
 	void SetUserIndex(int32 UserIndex) { CurrentUserIndex = UserIndex; }
@@ -73,42 +73,42 @@ public:
 	 * @param NavigationReply The Navigation Reply to specify a boundary rule for the search.
 	 * @param RuleWidget      The Widget that is applying the boundary rule, used to get the bounds of the Rule.
 	 */
-	TSharedPtr<SWidget> FindNextFocusableWidget(const FArrangedWidget& StartingWidget, const EUINavigation Direction, const FNavigationReply& NavigationReply, const FArrangedWidget& RuleWidget, int32 UserIndex);
+	SLATECORE_API TSharedPtr<SWidget> FindNextFocusableWidget(const FArrangedWidget& StartingWidget, const EUINavigation Direction, const FNavigationReply& NavigationReply, const FArrangedWidget& RuleWidget, int32 UserIndex);
 
 	UE::Slate::FDeprecateVector2DResult GetGridSize() const { return GridSize; }
 	UE::Slate::FDeprecateVector2DResult GetGridOrigin() const { return GridOrigin; }
 	UE::Slate::FDeprecateVector2DResult GetGridWindowOrigin() const { return GridWindowOrigin; }
 
 	/** Clear the grid */
-	void Clear();
+	SLATECORE_API void Clear();
 
 	/** Add SWidget from the HitTest Grid */
 	UE_DEPRECATED(4.27, "Deprecated. Use the AddWidget with the FSlateInvalidationWidgetSortOrder type parameters instead. Passing FSlateInvalidationWidgetSortOrder()")
-	void AddWidget(const TSharedRef<SWidget>& InWidget, int32 InBatchPriorityGroup, int32 InLayerId, int32 InSecondarySort);
+	SLATECORE_API void AddWidget(const TSharedRef<SWidget>& InWidget, int32 InBatchPriorityGroup, int32 InLayerId, int32 InSecondarySort);
 
 	/** Add SWidget from the HitTest Grid */
 	UE_DEPRECATED(5.0, "Deprecated. Use the AddWidget with a pointer.")
-	void AddWidget(const TSharedRef<SWidget>& InWidget, int32 InBatchPriorityGroup, int32 InLayerId, FSlateInvalidationWidgetSortOrder InSecondarySort);
+	SLATECORE_API void AddWidget(const TSharedRef<SWidget>& InWidget, int32 InBatchPriorityGroup, int32 InLayerId, FSlateInvalidationWidgetSortOrder InSecondarySort);
 
 	/** Add SWidget from the HitTest Grid */
-	void AddWidget(const SWidget* InWidget, int32 InBatchPriorityGroup, int32 InLayerId, FSlateInvalidationWidgetSortOrder InSecondarySort);
+	SLATECORE_API void AddWidget(const SWidget* InWidget, int32 InBatchPriorityGroup, int32 InLayerId, FSlateInvalidationWidgetSortOrder InSecondarySort);
 
 	/** Remove SWidget from the HitTest Grid */
 	UE_DEPRECATED(5.0, "Deprecated. Use the RemoveWidget with a pointer.")
-	void RemoveWidget(const TSharedRef<SWidget>& InWidget);
+	SLATECORE_API void RemoveWidget(const TSharedRef<SWidget>& InWidget);
 
 	/** Remove SWidget from the HitTest Grid */
-	void RemoveWidget(const SWidget* InWidget);
+	SLATECORE_API void RemoveWidget(const SWidget* InWidget);
 
 	/** Update the widget SecondarySort without removing it and readding it again. */
 	UE_DEPRECATED(5.0, "Deprecated. Use the UpdateWidget with a pointer.")
-	void UpdateWidget(const TSharedRef<SWidget>& InWidget, FSlateInvalidationWidgetSortOrder InSecondarySort);
+	SLATECORE_API void UpdateWidget(const TSharedRef<SWidget>& InWidget, FSlateInvalidationWidgetSortOrder InSecondarySort);
 	
 	/** Update the widget SecondarySort without removing it and readding it again. */
-	void UpdateWidget(const SWidget* InWidget, FSlateInvalidationWidgetSortOrder InSecondarySort);
+	SLATECORE_API void UpdateWidget(const SWidget* InWidget, FSlateInvalidationWidgetSortOrder InSecondarySort);
 
 	/** Check if SWidget is contained within the HitTest Grid */
-	bool ContainsWidget(const SWidget* InWidget) const;
+	SLATECORE_API bool ContainsWidget(const SWidget* InWidget) const;
 
 	/** Append an already existing grid that occupy the same space. */
 	UE_DEPRECATED(4.26, "Deprecated. Use the FHittestGrid::AddGrid method instead")
@@ -118,13 +118,13 @@ public:
 	 * Add an already existing grid that occupy the same space.
 	 * The grid needs to have an owner, not be this grid and occupy the same space as this grid.
 	 */
-	void AddGrid(const TSharedRef<const FHittestGrid>& OtherGrid);
+	SLATECORE_API void AddGrid(const TSharedRef<const FHittestGrid>& OtherGrid);
 
 	/** Remove a grid that was appended. */
-	void RemoveGrid(const TSharedRef<const FHittestGrid>& OtherGrid);
+	SLATECORE_API void RemoveGrid(const TSharedRef<const FHittestGrid>& OtherGrid);
 
 	/** Remove a grid that was appended. */
-	void RemoveGrid(const SWidget* OtherGridOwner);
+	SLATECORE_API void RemoveGrid(const SWidget* OtherGridOwner);
 
 	struct FDebuggingFindNextFocusableWidgetArgs
 	{
@@ -146,9 +146,9 @@ public:
 
 #if WITH_SLATE_DEBUGGING
 	DECLARE_MULTICAST_DELEGATE_TwoParams(FDebuggingFindNextFocusableWidget, const FHittestGrid* /*HittestGrid*/, const FDebuggingFindNextFocusableWidgetArgs& /*Info*/);
-	static FDebuggingFindNextFocusableWidget OnFindNextFocusableWidgetExecuted;
+	static SLATECORE_API FDebuggingFindNextFocusableWidget OnFindNextFocusableWidgetExecuted;
 
-	void LogGrid() const;
+	SLATECORE_API void LogGrid() const;
 
 	enum class EDisplayGridFlags
 	{
@@ -157,7 +157,7 @@ public:
 		HideUnsupportedKeyboardFocusWidgets = 1 << 1,	// Hide hit box for widgets that have SupportsKeyboardFocus false
 		UseFocusBrush = 1 << 2,
 	};
-	void DisplayGrid(int32 InLayer, const FGeometry& AllottedGeometry, FSlateWindowElementList& WindowElementList, EDisplayGridFlags DisplayFlags = EDisplayGridFlags::UseFocusBrush) const;
+	SLATECORE_API void DisplayGrid(int32 InLayer, const FGeometry& AllottedGeometry, FSlateWindowElementList& WindowElementList, EDisplayGridFlags DisplayFlags = EDisplayGridFlags::UseFocusBrush) const;
 
 	struct FWidgetSortData
 	{
@@ -165,7 +165,7 @@ public:
 		int64 PrimarySort;
 		FSlateInvalidationWidgetSortOrder SecondarySort;
 	};
-	TArray<FWidgetSortData> GetAllWidgetSortDatas() const;
+	SLATECORE_API TArray<FWidgetSortData> GetAllWidgetSortDatas() const;
 #endif
 
 private:
@@ -260,22 +260,22 @@ private:
 	};
 
 	//~ Helper functions
-	bool IsValidCellCoord(const FIntPoint& CellCoord) const;
-	bool IsValidCellCoord(const int32 XCoord, const int32 YCoord) const;
-	void ClearInternal(int32 TotalCells);
+	SLATECORE_API bool IsValidCellCoord(const FIntPoint& CellCoord) const;
+	SLATECORE_API bool IsValidCellCoord(const int32 XCoord, const int32 YCoord) const;
+	SLATECORE_API void ClearInternal(int32 TotalCells);
 
 	/** Return the Index and distance to a hit given the testing params */
-	FIndexAndDistance GetHitIndexFromCellIndex(const FGridTestingParams& Params) const;
+	SLATECORE_API FIndexAndDistance GetHitIndexFromCellIndex(const FGridTestingParams& Params) const;
 
 	/** @returns true if the child is a paint descendant of the provided Parent. */
-	bool IsDescendantOf(const SWidget* Parent, const FWidgetData& ChildData) const;
+	SLATECORE_API bool IsDescendantOf(const SWidget* Parent, const FWidgetData& ChildData) const;
 
 	/** Utility function for searching for the next focusable widget. */
 	template<typename TCompareFunc, typename TSourceSideFunc, typename TDestSideFunc>
 	TSharedPtr<SWidget> FindFocusableWidget(const FSlateRect WidgetRect, const FSlateRect SweptRect, int32 AxisIndex, int32 Increment, const EUINavigation Direction, const FNavigationReply& NavigationReply, TCompareFunc CompareFunc, TSourceSideFunc SourceSideFunc, TDestSideFunc DestSideFunc, int32 UserIndex, TArray<FDebuggingFindNextFocusableWidgetArgs::FWidgetResult>* IntermediatedResultPtr, TSet<TSharedPtr<SWidget>>* DisabledDestinations) const;
 
 	/** Constrains a float position into the grid coordinate. */
-	FIntPoint GetCellCoordinate(UE::Slate::FDeprecateVector2DParameter Position) const;
+	SLATECORE_API FIntPoint GetCellCoordinate(UE::Slate::FDeprecateVector2DParameter Position) const;
 
 	/** Access a cell at coordinates X, Y. Coordinates are row and column indexes. */
 	FORCEINLINE_DEBUGGABLE FCell& CellAt(const int32 X, const int32 Y)
@@ -292,21 +292,21 @@ private:
 	}
 
 	/** Is the other grid compatible with this grid. */
-	bool CanBeAppended(const FHittestGrid* OtherGrid) const;
+	SLATECORE_API bool CanBeAppended(const FHittestGrid* OtherGrid) const;
 
 	/** Are both grid of the same size. */
-	bool SameSize(const FHittestGrid* OtherGrid) const;
+	SLATECORE_API bool SameSize(const FHittestGrid* OtherGrid) const;
 
 	using FCollapsedHittestGridArray = TArray<const FHittestGrid*, TInlineAllocator<16>>;
 	/** Get all the hittest grid appended to this grid. */
-	void GetCollapsedHittestGrid(FCollapsedHittestGridArray& OutResult) const;
+	SLATECORE_API void GetCollapsedHittestGrid(FCollapsedHittestGridArray& OutResult) const;
 
 	using FCollapsedWidgetsArray = TArray<FWidgetIndex, TInlineAllocator<100>>;
 	/** Return the list of all the widget in that cell. */
-	void GetCollapsedWidgets(FCollapsedWidgetsArray& Out, const int32 X, const int32 Y) const;
+	SLATECORE_API void GetCollapsedWidgets(FCollapsedWidgetsArray& Out, const int32 X, const int32 Y) const;
 
 	/** Remove appended hittest grid that are not valid anymore. */
-	void RemoveStaleAppendedHittestGrid();
+	SLATECORE_API void RemoveStaleAppendedHittestGrid();
 
 private:
 	/** Map of all the widgets currently in the hit test grid to their stable index. */
diff --git a/Engine/Source/Runtime/SlateCore/Public/Input/Reply.h b/Engine/Source/Runtime/SlateCore/Public/Input/Reply.h
--- a/Engine/Source/Runtime/SlateCore/Public/Input/Reply.h
+++ b/Engine/Source/Runtime/SlateCore/Public/Input/Reply.h
@@ -20,7 +20,7 @@ class FDragDropOperation;
  * For example, a widget may handle an OnMouseDown event by asking the system to give mouse capture to a specific Widget.
  * To do this, return FReply::CaptureMouse( NewMouseCapture ).
  */
-class SLATECORE_API FReply : public TReplyBase<FReply>
+class FReply : public TReplyBase<FReply>
 {
 public:
 		
@@ -45,10 +45,10 @@ public:
 	/**
 	 * An event should return FReply::Handled().SetMousePos to ask Slate to move the mouse cursor to a different location
 	 */
-	FReply& SetMousePos( const FIntPoint& NewMousePos );
+	SLATECORE_API FReply& SetMousePos( const FIntPoint& NewMousePos );
 
 	/** An event should return FReply::Handled().SetUserFocus( SomeWidget ) as a means of asking the system to set users focus to the provided widget*/
-	FReply& SetUserFocus(TSharedRef<SWidget> GiveMeFocus, EFocusCause ReasonFocusIsChanging = EFocusCause::SetDirectly, bool bInAllUsers = false);
+	SLATECORE_API FReply& SetUserFocus(TSharedRef<SWidget> GiveMeFocus, EFocusCause ReasonFocusIsChanging = EFocusCause::SetDirectly, bool bInAllUsers = false);
 
 	/** An event should return a FReply::Handled().ClearUserFocus() to ask the system to clear user focus*/
 	FReply& ClearUserFocus(bool bInAllUsers = false)
@@ -57,14 +57,14 @@ public:
 	}
 
 	/** An event should return a FReply::Handled().ClearUserFocus() to ask the system to clear user focus*/
-	FReply& ClearUserFocus(EFocusCause ReasonFocusIsChanging, bool bInAllUsers = false);
+	SLATECORE_API FReply& ClearUserFocus(EFocusCause ReasonFocusIsChanging, bool bInAllUsers = false);
 
 	/** 
 	 *	An event should rarely invoke FReply::CancelFocusRequest. This will change the reply
 	 *  so that it no longer contains a request for the system to set or clear focus. 
 	 *  This is most useful when a reply has been cached for use in async operations.
 	 */
-	FReply& CancelFocusRequest();
+	SLATECORE_API FReply& CancelFocusRequest();
 
 	/** An event should return FReply::Handled().SetNavigation( NavigationType ) as a means of asking the system to attempt a navigation*/
 	FReply& SetNavigation(EUINavigation InNavigationType, const ENavigationGenesis InNavigationGenesis, const ENavigationSource InNavigationSource = ENavigationSource::FocusedWidget)
@@ -223,7 +223,7 @@ public:
 	const TOptional<FIntPoint>& GetRequestedMousePos() const { return RequestedMousePos; }
 
 	/** Converts the reply into an string representation. */
-	FString ToString();
+	SLATECORE_API FString ToString();
 		
 public:
 
diff --git a/Engine/Source/Runtime/SlateCore/Public/Layout/ArrangedChildren.h b/Engine/Source/Runtime/SlateCore/Public/Layout/ArrangedChildren.h
--- a/Engine/Source/Runtime/SlateCore/Public/Layout/ArrangedChildren.h
+++ b/Engine/Source/Runtime/SlateCore/Public/Layout/ArrangedChildren.h
@@ -11,7 +11,7 @@
  * FArrangedChildren supports a filter that is useful for excluding widgets with unwanted
  * visibilities.
  */
-class SLATECORE_API FArrangedChildren
+class FArrangedChildren
 {
 	private:
 	
@@ -83,9 +83,9 @@ class SLATECORE_API FArrangedChildren
 	 * Add an arranged widget (i.e. widget and its resulting geometry) to the list of Arranged children
 	 * based on the the visibility filter and the arranged widget's visibility
 	 */
-	void AddWidget(const FArrangedWidget& InWidgetGeometry);
+	SLATECORE_API void AddWidget(const FArrangedWidget& InWidgetGeometry);
 
-	void InsertWidget(const FArrangedWidget& InWidgetGeometry, int32 Index);
+	SLATECORE_API void InsertWidget(const FArrangedWidget& InWidgetGeometry, int32 Index);
 
 	FORCEINLINE_DEBUGGABLE void SetFilter(EVisibility InVisibility)
 	{
diff --git a/Engine/Source/Runtime/SlateCore/Public/Layout/Children.h b/Engine/Source/Runtime/SlateCore/Public/Layout/Children.h
--- a/Engine/Source/Runtime/SlateCore/Public/Layout/Children.h
+++ b/Engine/Source/Runtime/SlateCore/Public/Layout/Children.h
@@ -145,14 +145,14 @@ protected:
  * Widgets with no Children can return an instance of FNoChildren.
  * For convenience a shared instance FNoChildren::NoChildrenInstance can be used.
  */
-class SLATECORE_API FNoChildren : public FChildren
+class FNoChildren : public FChildren
 {
 public:
-	static FNoChildren NoChildrenInstance;
+	static SLATECORE_API FNoChildren NoChildrenInstance;
 
 public:
 	UE_DEPRECATED(5.0, "FNoChildren take a valid reference to a SWidget")
-	FNoChildren();
+	SLATECORE_API FNoChildren();
 
 	FNoChildren(SWidget* InOwner)
 		: FChildren(InOwner)
@@ -1172,4 +1172,4 @@ private:
 
 	TPanelChildren<SlotType>* AllChildren;
 	const TAttribute<int32>* WidgetIndex;
-};
\ No newline at end of file
+};
diff --git a/Engine/Source/Runtime/SlateCore/Public/Layout/Clipping.h b/Engine/Source/Runtime/SlateCore/Public/Layout/Clipping.h
--- a/Engine/Source/Runtime/SlateCore/Public/Layout/Clipping.h
+++ b/Engine/Source/Runtime/SlateCore/Public/Layout/Clipping.h
@@ -57,7 +57,7 @@ enum class EWidgetClipping : uint8
 /**
  * The Clipping Zone represents some arbitrary plane segment that can be used to clip the geometry in Slate.
  */
-class SLATECORE_API FSlateClippingZone
+class FSlateClippingZone
 {
 public:
 
@@ -66,11 +66,11 @@ public:
 	FVector2f BottomLeft;
 	FVector2f BottomRight;
 
-	explicit FSlateClippingZone(const FShortRect& AxisAlignedRect);
-	explicit FSlateClippingZone(const FSlateRect& AxisAlignedRect);
-	explicit FSlateClippingZone(const FGeometry& BoundingGeometry);
-	explicit FSlateClippingZone(const FPaintGeometry& PaintingGeometry);
-	FSlateClippingZone(const UE::Slate::FDeprecateVector2DParameter& InTopLeft, const UE::Slate::FDeprecateVector2DParameter& InTopRight, const UE::Slate::FDeprecateVector2DParameter& InBottomLeft, const UE::Slate::FDeprecateVector2DParameter& InBottomRight);
+	SLATECORE_API explicit FSlateClippingZone(const FShortRect& AxisAlignedRect);
+	SLATECORE_API explicit FSlateClippingZone(const FSlateRect& AxisAlignedRect);
+	SLATECORE_API explicit FSlateClippingZone(const FGeometry& BoundingGeometry);
+	SLATECORE_API explicit FSlateClippingZone(const FPaintGeometry& PaintingGeometry);
+	SLATECORE_API FSlateClippingZone(const UE::Slate::FDeprecateVector2DParameter& InTopLeft, const UE::Slate::FDeprecateVector2DParameter& InTopRight, const UE::Slate::FDeprecateVector2DParameter& InBottomLeft, const UE::Slate::FDeprecateVector2DParameter& InBottomRight);
 	FSlateClippingZone() {}
 
 	/**  */
@@ -119,18 +119,18 @@ public:
 	}
 
 	/** Is a point inside the clipping zone? */
-	bool IsPointInside(const UE::Slate::FDeprecateVector2DParameter& Point) const;
+	SLATECORE_API bool IsPointInside(const UE::Slate::FDeprecateVector2DParameter& Point) const;
 
 	/**
 	 * Intersects two clipping zones and returns the new clipping zone that would need to be used.
 	 * This can only be called between two axis aligned clipping zones.
 	 */
-	FSlateClippingZone Intersect(const FSlateClippingZone& Other) const;
+	SLATECORE_API FSlateClippingZone Intersect(const FSlateClippingZone& Other) const;
 
 	/**
 	 * Gets the bounding box of the points making up this clipping zone.
 	 */
-	FSlateRect GetBoundingBox() const;
+	SLATECORE_API FSlateRect GetBoundingBox() const;
 
 	bool operator==(const FSlateClippingZone& Other) const
 	{
@@ -161,7 +161,7 @@ public:
 		return Absolute;
 	}
 private:
-	void InitializeFromArbitraryPoints(const UE::Slate::FDeprecateVector2DParameter& InTopLeft, const UE::Slate::FDeprecateVector2DParameter& InTopRight, const UE::Slate::FDeprecateVector2DParameter& InBottomLeft, const UE::Slate::FDeprecateVector2DParameter& InBottomRight);
+	SLATECORE_API void InitializeFromArbitraryPoints(const UE::Slate::FDeprecateVector2DParameter& InTopLeft, const UE::Slate::FDeprecateVector2DParameter& InTopRight, const UE::Slate::FDeprecateVector2DParameter& InBottomLeft, const UE::Slate::FDeprecateVector2DParameter& InBottomRight);
 
 private:
 	/** Is the clipping zone axis aligned?  Axis aligned clipping zones are much cheaper. */
@@ -203,13 +203,13 @@ ENUM_CLASS_FLAGS(EClippingFlags)
 /**
  * Captures everything about a single draw calls clipping state.
  */
-class SLATECORE_API FSlateClippingState
+class FSlateClippingState
 {
 public:
-	FSlateClippingState(EClippingFlags InFlags = EClippingFlags::None);
+	SLATECORE_API FSlateClippingState(EClippingFlags InFlags = EClippingFlags::None);
 	
 	/** Is a point inside the clipping state? */
-	bool IsPointInside(const UE::Slate::FDeprecateVector2DParameter& Point) const;
+	SLATECORE_API bool IsPointInside(const UE::Slate::FDeprecateVector2DParameter& Point) const;
 
 #if WITH_SLATE_DEBUGGING
 	/** Set the state index that this clipping state originated from.  We just do this for debugging purposes. */
@@ -339,28 +339,28 @@ public:
 /**
  * The clipping manager maintain the running clip state.  This is used for both maintain and for hit testing.
  */
-class SLATECORE_API FSlateClippingManager
+class FSlateClippingManager
 {
 public:
-	FSlateClippingManager();
+	SLATECORE_API FSlateClippingManager();
 
-	int32 PushClip(const FSlateClippingZone& InClippingZone);
-	int32 PushClippingState(const FSlateClippingState& InClipState);
-	int32 GetClippingIndex() const;
-	TOptional<FSlateClippingState> GetActiveClippingState() const;
+	SLATECORE_API int32 PushClip(const FSlateClippingZone& InClippingZone);
+	SLATECORE_API int32 PushClippingState(const FSlateClippingState& InClipState);
+	SLATECORE_API int32 GetClippingIndex() const;
+	SLATECORE_API TOptional<FSlateClippingState> GetActiveClippingState() const;
 	const TArray<int32>& GetClippingStack() const { return ClippingStack; }
-	const TArray<FSlateClippingState>& GetClippingStates() const;
-	void PopClip();
-	void PopToStackIndex(int32 Index);
+	SLATECORE_API const TArray<FSlateClippingState>& GetClippingStates() const;
+	SLATECORE_API void PopClip();
+	SLATECORE_API void PopToStackIndex(int32 Index);
 	int32 GetClippingIndexAtStackIndex(int32 StackIndex) const { return ClippingStack.IsValidIndex(StackIndex) ? ClippingStack[StackIndex] : INDEX_NONE; }
 	int32 GetStackDepth() const { return ClippingStack.Num(); }
-	const FSlateClippingState* GetPreviousClippingState(bool bWillIntersectWithParent) const;
+	SLATECORE_API const FSlateClippingState* GetPreviousClippingState(bool bWillIntersectWithParent) const;
 
-	void ResetClippingState();
+	SLATECORE_API void ResetClippingState();
 
 private:
 
-	FSlateClippingState CreateClippingState(const FSlateClippingZone& InClipRect) const;
+	SLATECORE_API FSlateClippingState CreateClippingState(const FSlateClippingZone& InClipRect) const;
 
 private:
 	/** Maintains the current clipping stack, with the indexes in the array of clipping states.  Pushed and popped throughout the drawing process. */
diff --git a/Engine/Source/Runtime/SlateCore/Public/Layout/FlowDirection.h b/Engine/Source/Runtime/SlateCore/Public/Layout/FlowDirection.h
--- a/Engine/Source/Runtime/SlateCore/Public/Layout/FlowDirection.h
+++ b/Engine/Source/Runtime/SlateCore/Public/Layout/FlowDirection.h
@@ -45,9 +45,9 @@ extern SLATECORE_API int32 GSlateFlowDirectionShouldFollowCultureByDefault;
 /**
  * Class containing utilities for getting layout localization information.
  */
-class SLATECORE_API FLayoutLocalization
+class FLayoutLocalization
 {
 public:
 	/** Gets the current expected flow direction based on localization. */
-	static EFlowDirection GetLocalizedLayoutDirection();
+	static SLATECORE_API EFlowDirection GetLocalizedLayoutDirection();
 };
diff --git a/Engine/Source/Runtime/SlateCore/Public/Layout/Geometry.h b/Engine/Source/Runtime/SlateCore/Public/Layout/Geometry.h
--- a/Engine/Source/Runtime/SlateCore/Public/Layout/Geometry.h
+++ b/Engine/Source/Runtime/SlateCore/Public/Layout/Geometry.h
@@ -36,7 +36,7 @@ template <typename T> struct TIsPODType;
  * the corresponding parent widget.
  */
 USTRUCT(BlueprintType)
-struct SLATECORE_API FGeometry
+struct FGeometry
 {
 	GENERATED_USTRUCT_BODY()
 
@@ -45,7 +45,7 @@ public:
 	/**
 	 * Default constructor. Creates a geometry with identity transforms.
 	 */
-	FGeometry();
+	SLATECORE_API FGeometry();
 
 	/**
 	 * Copy constructor.
@@ -56,7 +56,7 @@ public:
 	 * !!! HACK!!! We're keeping members of FGeometry const to prevent mutability without making them private, for backward compatibility.
 	 * But this means the assignment operator no longer works. We implement one ourselves now and force a memcpy.
 	 */
-	FGeometry& operator=(const FGeometry& RHS);
+	SLATECORE_API FGeometry& operator=(const FGeometry& RHS);
 
 	/**
 	 * !!! DEPRECATED FUNCTION !!! Use MakeChild taking a layout transform instead!
@@ -255,7 +255,7 @@ public:
 	 *
 	 * @return					The new child geometry.
 	 */
-	FArrangedWidget MakeChild(const TSharedRef<SWidget>& ChildWidget, const FLayoutGeometry& LayoutGeometry) const;
+	SLATECORE_API FArrangedWidget MakeChild(const TSharedRef<SWidget>& ChildWidget, const FLayoutGeometry& LayoutGeometry) const;
 
 	/**
 	 * Create a child geometry+widget relative to this one with a given local space size and layout transform.
@@ -269,7 +269,7 @@ public:
 	 *
 	 * @return					The new child geometry+widget.
 	 */
-	FArrangedWidget MakeChild(const TSharedRef<SWidget>& ChildWidget, const UE::Slate::FDeprecateVector2DParameter& InLocalSize, const FSlateLayoutTransform& LayoutTransform) const;
+	SLATECORE_API FArrangedWidget MakeChild(const TSharedRef<SWidget>& ChildWidget, const UE::Slate::FDeprecateVector2DParameter& InLocalSize, const FSlateLayoutTransform& LayoutTransform) const;
 
 #if UE_ENABLE_SLATE_VECTOR_DEPRECATION_MECHANISMS
 
@@ -305,7 +305,7 @@ public:
 	 *
 	 * @return				The new child geometry+widget.
 	 */
-	FArrangedWidget MakeChild(const TSharedRef<SWidget>& ChildWidget, const UE::Slate::FDeprecateVector2DParameter& ChildOffset, const UE::Slate::FDeprecateVector2DParameter& InLocalSize, float ChildScale = 1.0f) const;
+	SLATECORE_API FArrangedWidget MakeChild(const TSharedRef<SWidget>& ChildWidget, const UE::Slate::FDeprecateVector2DParameter& ChildOffset, const UE::Slate::FDeprecateVector2DParameter& InLocalSize, float ChildScale = 1.0f) const;
 
 	/**
 	 * Create a paint geometry that represents this geometry.
@@ -493,7 +493,7 @@ public:
 	}
 	
 	/** @return A String representation of this Geometry */
-	FString ToString() const;
+	SLATECORE_API FString ToString() const;
 
 	/** 
 	 * !!! DEPRECATED !!! This legacy function does not account for render transforms.
diff --git a/Engine/Source/Runtime/SlateCore/Public/Layout/PaintGeometry.h b/Engine/Source/Runtime/SlateCore/Public/Layout/PaintGeometry.h
--- a/Engine/Source/Runtime/SlateCore/Public/Layout/PaintGeometry.h
+++ b/Engine/Source/Runtime/SlateCore/Public/Layout/PaintGeometry.h
@@ -20,7 +20,7 @@
  *
  *		The DrawScale is only applied to the internal aspects of the draw primitives. e.g. Line thickness, 3x3 grid margins, etc.
  */
-struct SLATECORE_API FPaintGeometry
+struct FPaintGeometry
 {
 	/** 
 	 * !!! DEPRECATED!!! Drawing should only happen in local space to ensure render transforms work.
diff --git a/Engine/Source/Runtime/SlateCore/Public/Layout/SlateRect.h b/Engine/Source/Runtime/SlateCore/Public/Layout/SlateRect.h
--- a/Engine/Source/Runtime/SlateCore/Public/Layout/SlateRect.h
+++ b/Engine/Source/Runtime/SlateCore/Public/Layout/SlateRect.h
@@ -22,7 +22,7 @@
  *      \/
  *      Y
  */
-class SLATECORE_API FSlateRect
+class FSlateRect
 {
 public:
 
@@ -345,7 +345,7 @@ public:
 	*
 	* @return	A string of the rect coordinates 
 	*/
-	FString ToString() const;
+	SLATECORE_API FString ToString() const;
 
 	/**
 	* Returns a string of containing the coordinates of the rect
@@ -354,7 +354,7 @@ public:
 	*
 	* @return	True if initialized successfully
 	*/
-	bool InitFromString(const FString& InSourceString);
+	SLATECORE_API bool InitFromString(const FString& InSourceString);
 
 	
 	friend FORCEINLINE uint32 GetTypeHash(const FSlateRect& Key)
diff --git a/Engine/Source/Runtime/SlateCore/Public/Layout/SlateRotatedRect.h b/Engine/Source/Runtime/SlateCore/Public/Layout/SlateRotatedRect.h
--- a/Engine/Source/Runtime/SlateCore/Public/Layout/SlateRotatedRect.h
+++ b/Engine/Source/Runtime/SlateCore/Public/Layout/SlateRotatedRect.h
@@ -16,7 +16,7 @@
  * We provide a ctor that does the work common to slate drawing, but you could technically
  * create this any way you want.
  */
-struct SLATECORE_API FSlateRotatedRect
+struct FSlateRotatedRect
 {
 public:
 	/** Default ctor. */
@@ -59,17 +59,17 @@ public:
 public:
 
 	/** Convert to a bounding, aligned rect. */
-	FSlateRect ToBoundingRect() const;
+	SLATECORE_API FSlateRect ToBoundingRect() const;
 
 	/** Point-in-rect test. */
-	bool IsUnderLocation(const UE::Slate::FDeprecateVector2DParameter Location) const;
+	SLATECORE_API bool IsUnderLocation(const UE::Slate::FDeprecateVector2DParameter Location) const;
 
 	static FSlateRotatedRect MakeRotatedRect(const FSlateRect& ClipRectInLayoutWindowSpace, const FSlateLayoutTransform& InverseLayoutTransform, const FSlateRenderTransform& RenderTransform)
 	{
 		return MakeRotatedRect(ClipRectInLayoutWindowSpace, Concatenate(InverseLayoutTransform, RenderTransform));
 	}
 
-	static FSlateRotatedRect MakeRotatedRect(const FSlateRect& ClipRectInLayoutWindowSpace, const FTransform2f& LayoutToRenderTransform);
+	static SLATECORE_API FSlateRotatedRect MakeRotatedRect(const FSlateRect& ClipRectInLayoutWindowSpace, const FTransform2f& LayoutToRenderTransform);
 
 	static FSlateRotatedRect MakeSnappedRotatedRect(const FSlateRect& ClipRectInLayoutWindowSpace, const FSlateLayoutTransform& InverseLayoutTransform, const FSlateRenderTransform& RenderTransform)
 	{
@@ -79,7 +79,7 @@ public:
 	/**
 	* Used to construct a rotated rect from an aligned clip rect and a set of layout and render transforms from the geometry, snapped to pixel boundaries. Returns a float or float16 version of the rect based on the typedef.
 	*/
-	static FSlateRotatedRect MakeSnappedRotatedRect(const FSlateRect& ClipRectInLayoutWindowSpace, const FTransform2f& LayoutToRenderTransform);
+	static SLATECORE_API FSlateRotatedRect MakeSnappedRotatedRect(const FSlateRect& ClipRectInLayoutWindowSpace, const FTransform2f& LayoutToRenderTransform);
 };
 
 /**
diff --git a/Engine/Source/Runtime/SlateCore/Public/Layout/Visibility.h b/Engine/Source/Runtime/SlateCore/Public/Layout/Visibility.h
--- a/Engine/Source/Runtime/SlateCore/Public/Layout/Visibility.h
+++ b/Engine/Source/Runtime/SlateCore/Public/Layout/Visibility.h
@@ -8,25 +8,25 @@
 #include "HAL/Platform.h"
 
 /** Is an entity visible? */
-struct SLATECORE_API EVisibility
+struct EVisibility
 {
 	/** Visible and hit-testable (can interact with cursor). Default value. */
-	static const EVisibility Visible;
+	static SLATECORE_API const EVisibility Visible;
 
 	/** Not visible and takes up no space in the layout (obviously not hit-testable). */
-	static const EVisibility Collapsed;
+	static SLATECORE_API const EVisibility Collapsed;
 
 	/** Not visible but occupies layout space (obviously not hit-testable). */
-	static const EVisibility Hidden;
+	static SLATECORE_API const EVisibility Hidden;
 
 	/** Visible but not hit-testable (cannot interact with cursor) and children in the hierarchy (if any) are also not hit-testable. */
-	static const EVisibility HitTestInvisible;
+	static SLATECORE_API const EVisibility HitTestInvisible;
 
 	/** Visible but not hit-testable (cannot interact with cursor) and doesn't affect hit-testing on children (if any). */
-	static const EVisibility SelfHitTestInvisible;
+	static SLATECORE_API const EVisibility SelfHitTestInvisible;
 
 	/** Any visibility will do */
-	static const EVisibility All;
+	static SLATECORE_API const EVisibility All;
 
 public:
 
@@ -73,7 +73,7 @@ public:
 		return 0 != (InVisibility.Value & InVisibilityFilter.Value);
 	}
 
-	FString ToString() const;
+	SLATECORE_API FString ToString() const;
 
 private:
 
diff --git a/Engine/Source/Runtime/SlateCore/Public/Rendering/DrawElements.h b/Engine/Source/Runtime/SlateCore/Public/Rendering/DrawElements.h
--- a/Engine/Source/Runtime/SlateCore/Public/Rendering/DrawElements.h
+++ b/Engine/Source/Runtime/SlateCore/Public/Rendering/DrawElements.h
@@ -641,7 +641,7 @@ public:
 
 	SLATECORE_API int32 PushClip(const FSlateClippingZone& InClipZone);
 	SLATECORE_API int32 GetClippingIndex() const;
-	SLATECORE_API int32 GetClippingStackDepth() const { return ClippingManager.GetStackDepth(); }
+	int32 GetClippingStackDepth() const { return ClippingManager.GetStackDepth(); }
 	SLATECORE_API TOptional<FSlateClippingState> GetClippingState() const;
 	SLATECORE_API void PopClip();
 	SLATECORE_API void PopClipToStackIndex(int32 Index);
@@ -666,14 +666,14 @@ public:
 	 * Some widgets may want to paint their children after after another, loosely-related widget finished painting.
 	 * Or they may want to paint "after everyone".
 	 */
-	struct SLATECORE_API FDeferredPaint
+	struct FDeferredPaint
 	{
 	public:
-		FDeferredPaint(const TSharedRef<const SWidget>& InWidgetToPaint, const FPaintArgs& InArgs, const FGeometry InAllottedGeometry, const FWidgetStyle& InWidgetStyle, bool InParentEnabled);
+		SLATECORE_API FDeferredPaint(const TSharedRef<const SWidget>& InWidgetToPaint, const FPaintArgs& InArgs, const FGeometry InAllottedGeometry, const FWidgetStyle& InWidgetStyle, bool InParentEnabled);
 
-		int32 ExecutePaint(int32 LayerId, FSlateWindowElementList& OutDrawElements, const FSlateRect& MyCullingRect) const;
+		SLATECORE_API int32 ExecutePaint(int32 LayerId, FSlateWindowElementList& OutDrawElements, const FSlateRect& MyCullingRect) const;
 
-		FDeferredPaint Copy(const FPaintArgs& InArgs);
+		SLATECORE_API FDeferredPaint Copy(const FPaintArgs& InArgs);
 
 	private:
 		// Used for making copies.
diff --git a/Engine/Source/Runtime/SlateCore/Public/Rendering/RenderingCommon.h b/Engine/Source/Runtime/SlateCore/Public/Rendering/RenderingCommon.h
--- a/Engine/Source/Runtime/SlateCore/Public/Rendering/RenderingCommon.h
+++ b/Engine/Source/Runtime/SlateCore/Public/Rendering/RenderingCommon.h
@@ -201,7 +201,7 @@ struct FShaderParams
 /** 
  * A struct which defines a basic vertex seen by the Slate vertex buffers and shaders
  */
-struct SLATECORE_API FSlateVertex
+struct FSlateVertex
 {
 	/** Texture coordinates.  The first 2 are in xy and the 2nd are in zw */
 	float TexCoords[4]; 
diff --git a/Engine/Source/Runtime/SlateCore/Public/Rendering/ShaderResourceManager.h b/Engine/Source/Runtime/SlateCore/Public/Rendering/ShaderResourceManager.h
--- a/Engine/Source/Runtime/SlateCore/Public/Rendering/ShaderResourceManager.h
+++ b/Engine/Source/Runtime/SlateCore/Public/Rendering/ShaderResourceManager.h
@@ -40,7 +40,7 @@ struct FCompareFNewTextureInfoByTextureSize
 /** 
  * Base texture manager class used by a Slate renderer to manage texture resources
  */
-class SLATECORE_API FSlateShaderResourceManager
+class FSlateShaderResourceManager
 {
 public:
 	FSlateShaderResourceManager() {};
@@ -66,8 +66,8 @@ public:
 	 * @param	DrawScale	The scale of the draw element using this brush.  This scale is used to rescale vector graphics only
 	 * @return	The created resource handle.  
 	 */
-	virtual FSlateResourceHandle GetResourceHandle(const FSlateBrush& Brush, FVector2f LocalSize, float DrawScale);
-	virtual FSlateResourceHandle GetResourceHandle(const FSlateBrush& Brush);
+	SLATECORE_API virtual FSlateResourceHandle GetResourceHandle(const FSlateBrush& Brush, FVector2f LocalSize, float DrawScale);
+	SLATECORE_API virtual FSlateResourceHandle GetResourceHandle(const FSlateBrush& Brush);
 
 	virtual FSlateShaderResource* GetFontShaderResource( int32 InTextureAtlasIndex, FSlateShaderResource* FontTextureAtlas, const class UObject* FontMaterial ) { return FontTextureAtlas; }
 
@@ -105,7 +105,7 @@ protected:
 #endif
 private:
 	// Non-copyable
-	FSlateShaderResourceManager(const FSlateShaderResourceManager&);
-	FSlateShaderResourceManager& operator=(const FSlateShaderResourceManager&);
+	SLATECORE_API FSlateShaderResourceManager(const FSlateShaderResourceManager&);
+	SLATECORE_API FSlateShaderResourceManager& operator=(const FSlateShaderResourceManager&);
 
 };
diff --git a/Engine/Source/Runtime/SlateCore/Public/Rendering/SlateDrawBuffer.h b/Engine/Source/Runtime/SlateCore/Public/Rendering/SlateDrawBuffer.h
--- a/Engine/Source/Runtime/SlateCore/Public/Rendering/SlateDrawBuffer.h
+++ b/Engine/Source/Runtime/SlateCore/Public/Rendering/SlateDrawBuffer.h
@@ -15,27 +15,27 @@ class SWindow;
 /**
  * Implements a draw buffer for Slate.
  */
-class SLATECORE_API FSlateDrawBuffer : public FGCObject
+class FSlateDrawBuffer : public FGCObject
 {
 public:
 	/** Default constructor. */
-	FSlateDrawBuffer();
+	SLATECORE_API FSlateDrawBuffer();
 
 	/** Removes all data from the buffer. */
-	void ClearBuffer();
+	SLATECORE_API void ClearBuffer();
 
 	/** Updates renderer resource version to allow the draw buffer to clean up cached resources */
-	void UpdateResourceVersion(uint32 NewResourceVersion);
+	SLATECORE_API void UpdateResourceVersion(uint32 NewResourceVersion);
 
 	/**
 	 * Creates a new FSlateWindowElementList and returns a reference to it so it can have draw elements added to it
 	 *
 	 * @param ForWindow    The window for which we are creating a list of paint elements.
 	 */
-	FSlateWindowElementList& AddWindowElementList(TSharedRef<SWindow> ForWindow);
+	SLATECORE_API FSlateWindowElementList& AddWindowElementList(TSharedRef<SWindow> ForWindow);
 
 	/** Removes any window from the draw buffer that's not in this list or whose window has become invalid. */
-	void RemoveUnusedWindowElement(const TArray<SWindow*>& AllWindows);
+	SLATECORE_API void RemoveUnusedWindowElement(const TArray<SWindow*>& AllWindows);
 
 	/**
 	 * Gets all window element lists in this buffer.
@@ -51,14 +51,14 @@ public:
 	 * @return true if the buffer could be locked.  False otherwise.
 	 * @see Unlock
 	 */
-	bool Lock();
+	SLATECORE_API bool Lock();
 
 	/**
 	 * Unlocks the buffer.  Indicates that the buffer is free.
 	 *
 	 * @see Lock
 	 */
-	void Unlock();
+	SLATECORE_API void Unlock();
 
 	/** @return true if the buffer is locked. */
 	bool IsLocked() const
@@ -67,7 +67,7 @@ public:
 	}
 
 	/** FGCObject Interface */
-	virtual void AddReferencedObjects(FReferenceCollector& Collector) override;
+	SLATECORE_API virtual void AddReferencedObjects(FReferenceCollector& Collector) override;
 	virtual FString GetReferencerName() const
 	{
 		return TEXT("FSlateDrawBuffer for Uncached Elements");
diff --git a/Engine/Source/Runtime/SlateCore/Public/Rendering/SlateRenderer.h b/Engine/Source/Runtime/SlateCore/Public/Rendering/SlateRenderer.h
--- a/Engine/Source/Runtime/SlateCore/Public/Rendering/SlateRenderer.h
+++ b/Engine/Source/Runtime/SlateCore/Public/Rendering/SlateRenderer.h
@@ -44,24 +44,24 @@ struct FRenderThreadUpdateContext
 /**
  * Provides access to the game and render thread font caches that Slate should use
  */
-class SLATECORE_API FSlateFontServices
+class FSlateFontServices
 {
 public:
 	/**
 	 * Construct the font services from the font caches (we'll create corresponding measure services ourselves)
 	 * These pointers may be the same if your renderer doesn't need a separate render thread font cache
 	 */
-	FSlateFontServices(TSharedRef<class FSlateFontCache> InGameThreadFontCache, TSharedRef<class FSlateFontCache> InRenderThreadFontCache);
+	SLATECORE_API FSlateFontServices(TSharedRef<class FSlateFontCache> InGameThreadFontCache, TSharedRef<class FSlateFontCache> InRenderThreadFontCache);
 
 	/**
 	 * Destruct the font services
 	 */
-	~FSlateFontServices();
+	SLATECORE_API ~FSlateFontServices();
 
 	/**
 	 * Get the font cache to use for the current thread
 	 */
-	TSharedRef<class FSlateFontCache> GetFontCache() const;
+	SLATECORE_API TSharedRef<class FSlateFontCache> GetFontCache() const;
 
 	/**
 	 * Get the font cache to use for the game thread
@@ -82,7 +82,7 @@ public:
 	/**
 	 * Get access to the font measure service for the current thread
 	 */
-	TSharedRef<class FSlateFontMeasure> GetFontMeasureService() const;
+	SLATECORE_API TSharedRef<class FSlateFontMeasure> GetFontMeasureService() const;
 
 	/**
 	 * Get access to the font measure service for the current thread
@@ -103,30 +103,30 @@ public:
 	/**
 	 * Flushes all cached data from the font cache for the current thread
 	 */
-	void FlushFontCache(const FString& FlushReason);
+	SLATECORE_API void FlushFontCache(const FString& FlushReason);
 
 	/**
 	 * Flushes all cached data from the font cache for the game thread
 	 */
-	void FlushGameThreadFontCache(const FString& FlushReason);
+	SLATECORE_API void FlushGameThreadFontCache(const FString& FlushReason);
 
 	/**
 	 * Flushes all cached data from the font cache for the render thread
 	 */
-	void FlushRenderThreadFontCache(const FString& FlushReason);
+	SLATECORE_API void FlushRenderThreadFontCache(const FString& FlushReason);
 
 	/**
 	 * Release any rendering resources owned by this font service
 	 */
-	void ReleaseResources();
+	SLATECORE_API void ReleaseResources();
 
 	/**
 	 * Delegate called after releasing the rendering resources used by this font service
 	 */
-	FOnReleaseFontResources& OnReleaseResources();
+	SLATECORE_API FOnReleaseFontResources& OnReleaseResources();
 
 private:
-	void HandleFontCacheReleaseResources(const class FSlateFontCache& InFontCache);
+	SLATECORE_API void HandleFontCacheReleaseResources(const class FSlateFontCache& InFontCache);
 
 	TSharedRef<class FSlateFontCache> GameThreadFontCache;
 	TSharedRef<class FSlateFontCache> RenderThreadFontCache;
@@ -167,15 +167,15 @@ struct FMappedTextureBuffer
 /**
  * Abstract base class for Slate renderers.
  */
-class SLATECORE_API FSlateRenderer
+class FSlateRenderer
 {
 public:
 
 	/** Constructor. */
-	explicit FSlateRenderer(const TSharedRef<FSlateFontServices>& InSlateFontServices);
+	SLATECORE_API explicit FSlateRenderer(const TSharedRef<FSlateFontServices>& InSlateFontServices);
 
 	/** Virtual destructor. */
-	virtual ~FSlateRenderer();
+	SLATECORE_API virtual ~FSlateRenderer();
 
 public:
 	/** Acquire the draw buffer and release it at the end of the scope. */
@@ -444,7 +444,7 @@ public:
 	 * @Window	The window to check for fullscreen
 	 * @return true if the window's viewport should be fullscreen
 	 */
-	bool IsViewportFullscreen( const SWindow& Window ) const;
+	SLATECORE_API bool IsViewportFullscreen( const SWindow& Window ) const;
 
 	/** Returns whether shaders that Slate depends on have been compiled. */
 	virtual bool AreShadersInitialized() const { return true; }
@@ -509,12 +509,12 @@ public:
 	/**
 	 * Returns the way to access the texture atlas information for this renderer
 	 */
-	virtual ISlateAtlasProvider* GetTextureAtlasProvider();
+	SLATECORE_API virtual ISlateAtlasProvider* GetTextureAtlasProvider();
 
 	/**
 	 * Returns the way to access the font atlas information for this renderer
 	 */
-	virtual ISlateAtlasProvider* GetFontAtlasProvider();
+	SLATECORE_API virtual ISlateAtlasProvider* GetFontAtlasProvider();
 
 	/**
 	 * Copies all slate windows out to a buffer at half resolution with debug information
@@ -541,8 +541,8 @@ public:
 	/** Reset the internal Scene tracking.*/
 	virtual void ClearScenes() = 0;
 
-	virtual void DestroyCachedFastPathRenderingData(struct FSlateCachedFastPathRenderingData* VertexData);
-	virtual void DestroyCachedFastPathElementData(struct FSlateCachedElementData* ElementData);
+	SLATECORE_API virtual void DestroyCachedFastPathRenderingData(struct FSlateCachedFastPathRenderingData* VertexData);
+	SLATECORE_API virtual void DestroyCachedFastPathElementData(struct FSlateCachedElementData* ElementData);
 
 	virtual bool HasLostDevice() const { return false; }
 
@@ -562,10 +562,10 @@ public:
 private:
 
 	// Non-copyable
-	FSlateRenderer(const FSlateRenderer&);
-	FSlateRenderer& operator=(const FSlateRenderer&);
+	SLATECORE_API FSlateRenderer(const FSlateRenderer&);
+	SLATECORE_API FSlateRenderer& operator=(const FSlateRenderer&);
 
-	void HandleFontCacheReleaseResources(const class FSlateFontCache& InFontCache);
+	SLATECORE_API void HandleFontCacheReleaseResources(const class FSlateFontCache& InFontCache);
 
 protected:
 
diff --git a/Engine/Source/Runtime/SlateCore/Public/Rendering/SlateResourceHandle.h b/Engine/Source/Runtime/SlateCore/Public/Rendering/SlateResourceHandle.h
--- a/Engine/Source/Runtime/SlateCore/Public/Rendering/SlateResourceHandle.h
+++ b/Engine/Source/Runtime/SlateCore/Public/Rendering/SlateResourceHandle.h
@@ -12,25 +12,25 @@ class FSlateShaderResourceProxy;
  * A SlateResourceHandle is used as fast path for looking up a rendering resource for a given brush when adding Slate draw elements
  * This can be cached and stored safely in code.  It will become invalid when a resource is destroyed
 */
-class SLATECORE_API FSlateResourceHandle
+class FSlateResourceHandle
 {
 	friend class FSlateShaderResourceManager;
 	friend class FSlateNullShaderResourceManager;
 public:
-	FSlateResourceHandle();
+	SLATECORE_API FSlateResourceHandle();
 
 	/**
 	 * @return true if the handle still points to a valid rendering resource
 	 */
-	bool IsValid() const;
+	SLATECORE_API bool IsValid() const;
 
 	/**
 	 * @return the resource proxy used to render.
 	 */
-	const FSlateShaderResourceProxy* GetResourceProxy() const;
+	SLATECORE_API const FSlateShaderResourceProxy* GetResourceProxy() const;
 
 private:
-	FSlateResourceHandle(const TSharedPtr<FSlateSharedHandleData>& InData);
+	SLATECORE_API FSlateResourceHandle(const TSharedPtr<FSlateSharedHandleData>& InData);
 
 	/** Internal data to pair the handle to the resource */
 	TSharedPtr<FSlateSharedHandleData> Data;
diff --git a/Engine/Source/Runtime/SlateCore/Public/SlotBase.h b/Engine/Source/Runtime/SlateCore/Public/SlotBase.h
--- a/Engine/Source/Runtime/SlateCore/Public/SlotBase.h
+++ b/Engine/Source/Runtime/SlateCore/Public/SlotBase.h
@@ -10,16 +10,16 @@ class FChildren;
 class SWidget;
 
 /** Slot are a container of a SWidget used by the FChildren. */
-class SLATECORE_API FSlotBase
+class FSlotBase
 {
 public:
-	FSlotBase();
-	FSlotBase(const FChildren& InParent);
-	FSlotBase(const TSharedRef<SWidget>& InWidget);
+	SLATECORE_API FSlotBase();
+	SLATECORE_API FSlotBase(const FChildren& InParent);
+	SLATECORE_API FSlotBase(const TSharedRef<SWidget>& InWidget);
 	FSlotBase& operator=(const FSlotBase&) = delete;
 	FSlotBase(const FSlotBase&) = delete;
 
-	virtual ~FSlotBase();
+	SLATECORE_API virtual ~FSlotBase();
 
 public:
 	struct FSlotArguments {};
@@ -38,13 +38,13 @@ public:
 	 * Access the widget that own the slot.
 	 * The owner can be invalid when the slot is not attached.
 	 */
-	SWidget* GetOwnerWidget() const;
+	SLATECORE_API SWidget* GetOwnerWidget() const;
 
 	/**
 	 * Set the owner of the slot.
 	 * Slots cannot be reassigned to different parents.
 	 */
-	void SetOwner(const FChildren& Children);
+	SLATECORE_API void SetOwner(const FChildren& Children);
 
 	/** Attach the child widget the slot now owns. */
 	FORCEINLINE_DEBUGGABLE void AttachWidget( const TSharedRef<SWidget>& InWidget )
@@ -72,10 +72,10 @@ public:
 	 * The removed widget is returned so that operations could be performed on it.
 	 * If the null widget was being stored, an invalid shared ptr is returned instead.
 	 */
-	const TSharedPtr<SWidget> DetachWidget();
+	SLATECORE_API const TSharedPtr<SWidget> DetachWidget();
 
 	/** Invalidate the widget's owner. */
-	void Invalidate(EInvalidateWidgetReason InvalidateReason);
+	SLATECORE_API void Invalidate(EInvalidateWidgetReason InvalidateReason);
 
 protected:
 	/**
@@ -105,8 +105,8 @@ protected:
 	}
 
 private:
-	void DetatchParentFromContent();
-	void AfterContentOrOwnerAssigned();
+	SLATECORE_API void DetatchParentFromContent();
+	SLATECORE_API void AfterContentOrOwnerAssigned();
 
 private:
 	/** The children that own the slot. */
diff --git a/Engine/Source/Runtime/SlateCore/Public/Sound/ISlateSoundDevice.h b/Engine/Source/Runtime/SlateCore/Public/Sound/ISlateSoundDevice.h
--- a/Engine/Source/Runtime/SlateCore/Public/Sound/ISlateSoundDevice.h
+++ b/Engine/Source/Runtime/SlateCore/Public/Sound/ISlateSoundDevice.h
@@ -6,7 +6,7 @@
 struct FSlateSound;
 
 /** Interface that Slate uses to play sounds. */
-class SLATECORE_API ISlateSoundDevice
+class ISlateSoundDevice
 {
 public:
 	/** Play the sound resource within the given Slate sound object */
diff --git a/Engine/Source/Runtime/SlateCore/Public/Sound/NullSlateSoundDevice.h b/Engine/Source/Runtime/SlateCore/Public/Sound/NullSlateSoundDevice.h
--- a/Engine/Source/Runtime/SlateCore/Public/Sound/NullSlateSoundDevice.h
+++ b/Engine/Source/Runtime/SlateCore/Public/Sound/NullSlateSoundDevice.h
@@ -9,10 +9,10 @@
 struct FSlateSound;
 
 /** Silent implementation of ISlateSoundDevice; it plays nothing. */
-class SLATECORE_API FNullSlateSoundDevice : public ISlateSoundDevice
+class FNullSlateSoundDevice : public ISlateSoundDevice
 {
 public:
-	virtual void PlaySound(const FSlateSound&, int32) const override;
-	virtual float GetSoundDuration(const FSlateSound& Sound) const override;
-	virtual ~FNullSlateSoundDevice();
+	SLATECORE_API virtual void PlaySound(const FSlateSound&, int32) const override;
+	SLATECORE_API virtual float GetSoundDuration(const FSlateSound& Sound) const override;
+	SLATECORE_API virtual ~FNullSlateSoundDevice();
 };
diff --git a/Engine/Source/Runtime/SlateCore/Public/Sound/SlateSound.h b/Engine/Source/Runtime/SlateCore/Public/Sound/SlateSound.h
--- a/Engine/Source/Runtime/SlateCore/Public/Sound/SlateSound.h
+++ b/Engine/Source/Runtime/SlateCore/Public/Sound/SlateSound.h
@@ -12,7 +12,7 @@
  * An intermediary to make UBaseSound available for Slate to play sounds
  */
 USTRUCT(BlueprintType)
-struct SLATECORE_API FSlateSound
+struct FSlateSound
 {
 	GENERATED_USTRUCT_BODY()
 
@@ -30,7 +30,7 @@ public:
 	 * @note Ensure that you only access the resource as a USoundBase 
 	 * @return The resource object, or null if this sound resource is empty
 	 */
-	UObject* GetResourceObject( ) const;
+	SLATECORE_API UObject* GetResourceObject( ) const;
 
 	/**
 	* Sets the sound that is supposed to be played.
@@ -43,12 +43,12 @@ public:
 	/**
 	 * Used when updating the ResourceObject within an FSlateSound from the editor to clear out any legacy data that may be set.
 	 */
-	void StripLegacyData_DEPRECATED( );
+	SLATECORE_API void StripLegacyData_DEPRECATED( );
 
 	/**
 	 * Used to upgrade a FName property to a FSlateSound property
 	 */
-	bool SerializeFromMismatchedTag(struct FPropertyTag const& Tag, FStructuredArchive::FSlot Slot);
+	SLATECORE_API bool SerializeFromMismatchedTag(struct FPropertyTag const& Tag, FStructuredArchive::FSlot Slot);
 
 public:
 
@@ -59,7 +59,7 @@ public:
 	 * @param SoundName The name of the sound to load, eg) SoundCue'/Game/QA_Assets/Sound/TEST_Music_Ambient.TEST_Music_Ambient'
 	 * @return An FSlateSound object with a ResourceObject pointer that's either set the given sound, or points to null if the sound could not be found
 	 */
-	static FSlateSound FromName_DEPRECATED( const FName& SoundName );
+	static SLATECORE_API FSlateSound FromName_DEPRECATED( const FName& SoundName );
 
 protected:
 
diff --git a/Engine/Source/Runtime/SlateCore/Public/Styling/AppStyle.h b/Engine/Source/Runtime/SlateCore/Public/Styling/AppStyle.h
--- a/Engine/Source/Runtime/SlateCore/Public/Styling/AppStyle.h
+++ b/Engine/Source/Runtime/SlateCore/Public/Styling/AppStyle.h
@@ -20,17 +20,17 @@
  *
  */
 
-class SLATECORE_API FAppStyle
+class FAppStyle
 {
 
 public:
 
-	static const ISlateStyle& Get();
+	static SLATECORE_API const ISlateStyle& Get();
 	
-	static const FName GetAppStyleSetName();
-	static void SetAppStyleSetName(const FName& InName);
+	static SLATECORE_API const FName GetAppStyleSetName();
+	static SLATECORE_API void SetAppStyleSetName(const FName& InName);
 
-	static void SetAppStyleSet(const ISlateStyle& InStyle);
+	static SLATECORE_API void SetAppStyleSet(const ISlateStyle& InStyle);
 
 	template< class T >
 	static const T& GetWidgetStyle(FName PropertyName, const ANSICHAR* Specifier = NULL)
@@ -131,6 +131,6 @@ public:
 
 private:
 
-	static FName AppStyleName;
+	static SLATECORE_API FName AppStyleName;
 
-};
\ No newline at end of file
+};
diff --git a/Engine/Source/Runtime/SlateCore/Public/Styling/CoreStyle.h b/Engine/Source/Runtime/SlateCore/Public/Styling/CoreStyle.h
--- a/Engine/Source/Runtime/SlateCore/Public/Styling/CoreStyle.h
+++ b/Engine/Source/Runtime/SlateCore/Public/Styling/CoreStyle.h
@@ -11,11 +11,11 @@ struct FSlateDynamicImageBrush;
 /**
  * Core slate style
  */
-class SLATECORE_API FCoreStyle 
+class FCoreStyle 
 {
 public:
 
-	static TSharedRef<class ISlateStyle> Create( const FName& InStyleSetName = "CoreStyle" );
+	static SLATECORE_API TSharedRef<class ISlateStyle> Create( const FName& InStyleSetName = "CoreStyle" );
 
 	/** 
 	* @return the Application Style 
@@ -42,40 +42,40 @@ public:
 	}
 
 	/** Get the default font for Slate */
-	static TSharedRef<const FCompositeFont> GetDefaultFont();
+	static SLATECORE_API TSharedRef<const FCompositeFont> GetDefaultFont();
 
 	/** Get a font style using the default for for Slate */
-	static FSlateFontInfo GetDefaultFontStyle(const FName InTypefaceFontName, const int32 InSize, const FFontOutlineSettings& InOutlineSettings = FFontOutlineSettings());
+	static SLATECORE_API FSlateFontInfo GetDefaultFontStyle(const FName InTypefaceFontName, const int32 InSize, const FFontOutlineSettings& InOutlineSettings = FFontOutlineSettings());
 
-	static void ResetToDefault( );
+	static SLATECORE_API void ResetToDefault( );
 
 	/** Used to override the default selection colors */
-	static void SetSelectorColor( const FLinearColor& NewColor );
-	static void SetSelectionColor( const FLinearColor& NewColor );
-	static void SetInactiveSelectionColor( const FLinearColor& NewColor );
-	static void SetPressedSelectionColor( const FLinearColor& NewColor );
-	static void SetFocusBrush(FSlateBrush* NewBrush);
+	static SLATECORE_API void SetSelectorColor( const FLinearColor& NewColor );
+	static SLATECORE_API void SetSelectionColor( const FLinearColor& NewColor );
+	static SLATECORE_API void SetInactiveSelectionColor( const FLinearColor& NewColor );
+	static SLATECORE_API void SetPressedSelectionColor( const FLinearColor& NewColor );
+	static SLATECORE_API void SetFocusBrush(FSlateBrush* NewBrush);
 
 	// todo: jdale - These are only here because of UTouchInterface::Activate and the fact that GetDynamicImageBrush is non-const
-	static const TSharedPtr<FSlateDynamicImageBrush> GetDynamicImageBrush( FName BrushTemplate, FName TextureName, const ANSICHAR* Specifier = nullptr );
-	static const TSharedPtr<FSlateDynamicImageBrush> GetDynamicImageBrush( FName BrushTemplate, const ANSICHAR* Specifier, class UTexture2D* TextureResource, FName TextureName );
-	static const TSharedPtr<FSlateDynamicImageBrush> GetDynamicImageBrush( FName BrushTemplate, class UTexture2D* TextureResource, FName TextureName );
+	static SLATECORE_API const TSharedPtr<FSlateDynamicImageBrush> GetDynamicImageBrush( FName BrushTemplate, FName TextureName, const ANSICHAR* Specifier = nullptr );
+	static SLATECORE_API const TSharedPtr<FSlateDynamicImageBrush> GetDynamicImageBrush( FName BrushTemplate, const ANSICHAR* Specifier, class UTexture2D* TextureResource, FName TextureName );
+	static SLATECORE_API const TSharedPtr<FSlateDynamicImageBrush> GetDynamicImageBrush( FName BrushTemplate, class UTexture2D* TextureResource, FName TextureName );
 
 	static const int32 RegularTextSize = 9;
 	static const int32 SmallTextSize = 8;
 
-	static bool IsStarshipStyle();
+	static SLATECORE_API bool IsStarshipStyle();
 
 	static bool IsInitialized() { return Instance.IsValid(); }
 
 private:
 
-	static void SetStyle( const TSharedRef< class ISlateStyle >& NewStyle );
+	static SLATECORE_API void SetStyle( const TSharedRef< class ISlateStyle >& NewStyle );
 
 private:
 
 	/** Singleton instances of this style. */
-	static TSharedPtr< class ISlateStyle > Instance;
+	static SLATECORE_API TSharedPtr< class ISlateStyle > Instance;
 };
 
 namespace CoreStyleConstants
diff --git a/Engine/Source/Runtime/SlateCore/Public/Styling/DefaultStyleCache.h b/Engine/Source/Runtime/SlateCore/Public/Styling/DefaultStyleCache.h
--- a/Engine/Source/Runtime/SlateCore/Public/Styling/DefaultStyleCache.h
+++ b/Engine/Source/Runtime/SlateCore/Public/Styling/DefaultStyleCache.h
@@ -14,7 +14,7 @@ namespace UE::Slate::Private
  * Single point of access for various default styles used in UMG for runtime and editor
  * with each style already having it's colors unlinked for reuse
  */
-struct SLATECORE_API FDefaultStyleCache
+struct FDefaultStyleCache
 {
 	/**
 	 * Internal style default holder struct, used so that one can get a bundle of defaults relevant for runtime or editor
@@ -107,11 +107,11 @@ struct SLATECORE_API FDefaultStyleCache
 private:
 
 	/** Singleton getter, however private used since GetRuntime & GetEditor is preferred for styles */
-	static FDefaultStyleCache& Get();
+	static SLATECORE_API FDefaultStyleCache& Get();
 
 	friend ::FLazySingleton;
 
-	FDefaultStyleCache();
+	SLATECORE_API FDefaultStyleCache();
 
 	FStyles Runtime;
 
@@ -120,4 +120,4 @@ private:
 #endif
 };
 
-} //namespace UE::Slate::Private
\ No newline at end of file
+} //namespace UE::Slate::Private
diff --git a/Engine/Source/Runtime/SlateCore/Public/Styling/ISlateStyle.h b/Engine/Source/Runtime/SlateCore/Public/Styling/ISlateStyle.h
--- a/Engine/Source/Runtime/SlateCore/Public/Styling/ISlateStyle.h
+++ b/Engine/Source/Runtime/SlateCore/Public/Styling/ISlateStyle.h
@@ -14,7 +14,7 @@ struct FSlateWidgetStyle;
 /**
  * 
  */
-class SLATECORE_API ISlateStyle
+class ISlateStyle
 {
 public:
 	/** Constructor. */
diff --git a/Engine/Source/Runtime/SlateCore/Public/Styling/SegmentedControlStyle.h b/Engine/Source/Runtime/SlateCore/Public/Styling/SegmentedControlStyle.h
--- a/Engine/Source/Runtime/SlateCore/Public/Styling/SegmentedControlStyle.h
+++ b/Engine/Source/Runtime/SlateCore/Public/Styling/SegmentedControlStyle.h
@@ -12,20 +12,20 @@
  * Represents the appearance of an SSegmentedControl
  */
 USTRUCT(BlueprintType)
-struct SLATECORE_API FSegmentedControlStyle: public FSlateWidgetStyle
+struct FSegmentedControlStyle: public FSlateWidgetStyle
 {
 	GENERATED_USTRUCT_BODY()
 
-	FSegmentedControlStyle();
+	SLATECORE_API FSegmentedControlStyle();
 
 	virtual ~FSegmentedControlStyle() {}
 
-	virtual void GetResources( TArray< const FSlateBrush* >& OutBrushes ) const override;
+	SLATECORE_API virtual void GetResources( TArray< const FSlateBrush* >& OutBrushes ) const override;
 
-	static const FName TypeName;
+	static SLATECORE_API const FName TypeName;
 	virtual const FName GetTypeName() const override { return TypeName; };
 
-	static const FSegmentedControlStyle& GetDefault();
+	static SLATECORE_API const FSegmentedControlStyle& GetDefault();
 
 	/**
 	 * The style to use for our Center Control
@@ -72,4 +72,4 @@ struct SLATECORE_API FSegmentedControlStyle: public FSlateWidgetStyle
 		BackgroundBrush.UnlinkColors();
 	}
 
-};
\ No newline at end of file
+};
diff --git a/Engine/Source/Runtime/SlateCore/Public/Styling/SlateBrush.h b/Engine/Source/Runtime/SlateCore/Public/Styling/SlateBrush.h
--- a/Engine/Source/Runtime/SlateCore/Public/Styling/SlateBrush.h
+++ b/Engine/Source/Runtime/SlateCore/Public/Styling/SlateBrush.h
@@ -130,7 +130,7 @@ namespace ESlateBrushRoundingType
  * Possible options for rounded box brush image
  */
 USTRUCT(BlueprintType)
-struct SLATECORE_API FSlateBrushOutlineSettings
+struct FSlateBrushOutlineSettings
 {
 	GENERATED_USTRUCT_BODY()
 
@@ -213,7 +213,7 @@ namespace SlateBrushDefs
  * A brush which contains information about how to draw a Slate element
  */
 USTRUCT(BlueprintType) //, meta = (HasNativeMake = ""))
-struct SLATECORE_API FSlateBrush
+struct FSlateBrush
 {
 	GENERATED_USTRUCT_BODY()
 
@@ -278,7 +278,7 @@ public:
 	/**
 	 * Default constructor.
 	 */
-	FSlateBrush();
+	SLATECORE_API FSlateBrush();
 
 	virtual ~FSlateBrush(){}
 
@@ -324,7 +324,7 @@ public:
 	/**
 	 * Sets the UObject that represents the brush resource.
 	 */
-	void SetResourceObject(class UObject* InResourceObject);
+	SLATECORE_API void SetResourceObject(class UObject* InResourceObject);
 
 	/**
 	 * Gets the brush's tint color.
@@ -435,7 +435,7 @@ public:
 	 *
 	 * @return Texture identifier string.
 	 */
-	static const FString UTextureIdentifier( );
+	static SLATECORE_API const FString UTextureIdentifier( );
 	
 	const FSlateResourceHandle& GetRenderingResource(UE::Slate::FDeprecateVector2DParameter LocalSize, float DrawScale) const
 	{
@@ -444,7 +444,7 @@ public:
 		return ResourceHandle;
 	}
 
-	const FSlateResourceHandle& GetRenderingResource() const;
+	SLATECORE_API const FSlateResourceHandle& GetRenderingResource() const;
 
 	bool IsSet() const { return bIsSet; }
 
@@ -456,8 +456,8 @@ public:
 #endif
 
 private:
-	void UpdateRenderingResource(FVector2f LocalSize, float DrawScale) const;
-	bool CanRenderResourceObject(UObject* InResourceObject) const;
+	SLATECORE_API void UpdateRenderingResource(FVector2f LocalSize, float DrawScale) const;
+	SLATECORE_API bool CanRenderResourceObject(UObject* InResourceObject) const;
 
 private:
 
@@ -497,11 +497,11 @@ protected:
 	 * @param InTint		  Tint to apply to the element.
 	 * @param InOutlineSettings Optional Outline Border Settings for RoundedBox mode
 	 */
-	 FORCENOINLINE FSlateBrush( ESlateBrushDrawType::Type InDrawType, const FName InResourceName, const FMargin& InMargin, ESlateBrushTileType::Type InTiling, ESlateBrushImageType::Type InImageType, const UE::Slate::FDeprecateVector2DParameter& InImageSize, const FLinearColor& InTint = FLinearColor::White, UObject* InObjectResource = nullptr, bool bInDynamicallyLoaded = false);
+	 SLATECORE_API FORCENOINLINE FSlateBrush( ESlateBrushDrawType::Type InDrawType, const FName InResourceName, const FMargin& InMargin, ESlateBrushTileType::Type InTiling, ESlateBrushImageType::Type InImageType, const UE::Slate::FDeprecateVector2DParameter& InImageSize, const FLinearColor& InTint = FLinearColor::White, UObject* InObjectResource = nullptr, bool bInDynamicallyLoaded = false);
 
-	 FORCENOINLINE FSlateBrush( ESlateBrushDrawType::Type InDrawType, const FName InResourceName, const FMargin& InMargin, ESlateBrushTileType::Type InTiling, ESlateBrushImageType::Type InImageType, const UE::Slate::FDeprecateVector2DParameter& InImageSize, const TSharedRef< FLinearColor >& InTint, UObject* InObjectResource = nullptr, bool bInDynamicallyLoaded = false);
+	 SLATECORE_API FORCENOINLINE FSlateBrush( ESlateBrushDrawType::Type InDrawType, const FName InResourceName, const FMargin& InMargin, ESlateBrushTileType::Type InTiling, ESlateBrushImageType::Type InImageType, const UE::Slate::FDeprecateVector2DParameter& InImageSize, const TSharedRef< FLinearColor >& InTint, UObject* InObjectResource = nullptr, bool bInDynamicallyLoaded = false);
 
-	 FORCENOINLINE FSlateBrush( ESlateBrushDrawType::Type InDrawType, const FName InResourceName, const FMargin& InMargin, ESlateBrushTileType::Type InTiling, ESlateBrushImageType::Type InImageType, const UE::Slate::FDeprecateVector2DParameter& InImageSize, const FSlateColor& InTint, UObject* InObjectResource = nullptr, bool bInDynamicallyLoaded = false);
+	 SLATECORE_API FORCENOINLINE FSlateBrush( ESlateBrushDrawType::Type InDrawType, const FName InResourceName, const FMargin& InMargin, ESlateBrushTileType::Type InTiling, ESlateBrushImageType::Type InImageType, const UE::Slate::FDeprecateVector2DParameter& InImageSize, const FSlateColor& InTint, UObject* InObjectResource = nullptr, bool bInDynamicallyLoaded = false);
 
 };
 
diff --git a/Engine/Source/Runtime/SlateCore/Public/Styling/SlateColor.h b/Engine/Source/Runtime/SlateCore/Public/Styling/SlateColor.h
--- a/Engine/Source/Runtime/SlateCore/Public/Styling/SlateColor.h
+++ b/Engine/Source/Runtime/SlateCore/Public/Styling/SlateColor.h
@@ -38,7 +38,7 @@ enum class EStyleColor : uint8;
  * A Slate color can be a directly specified value, or the color can be pulled from a WidgetStyle.
  */
 USTRUCT(BlueprintType)
-struct SLATECORE_API FSlateColor
+struct FSlateColor
 {
 	GENERATED_USTRUCT_BODY()
 
@@ -216,7 +216,7 @@ public:
 
 
 	/** Used to upgrade an FColor or FLinearColor property to an FSlateColor property */
-	bool SerializeFromMismatchedTag(const struct FPropertyTag& Tag, FStructuredArchive::FSlot Slot);
+	SLATECORE_API bool SerializeFromMismatchedTag(const struct FPropertyTag& Tag, FStructuredArchive::FSlot Slot);
 
 protected:
 	
@@ -227,7 +227,7 @@ protected:
 		, ColorTableId()
 	{ }
 
-	const FLinearColor& GetColorFromTable() const;
+	SLATECORE_API const FLinearColor& GetColorFromTable() const;
 protected:
 
 	// The current specified color; only meaningful when ColorToUse == UseColor_Specified.
diff --git a/Engine/Source/Runtime/SlateCore/Public/Styling/SlateIconFinder.h b/Engine/Source/Runtime/SlateCore/Public/Styling/SlateIconFinder.h
--- a/Engine/Source/Runtime/SlateCore/Public/Styling/SlateIconFinder.h
+++ b/Engine/Source/Runtime/SlateCore/Public/Styling/SlateIconFinder.h
@@ -8,7 +8,7 @@
 struct FSlateBrush;
 
 /** Class used for finding icons within a registered set of styles */
-class SLATECORE_API FSlateIconFinder
+class FSlateIconFinder
 {
 public:
 
@@ -19,7 +19,7 @@ public:
 	 * @param InDefaultName		The default icon name to use, if no specialized icon could be found
 	 * @return An FSlateIcon structure defining the class's icon
 	 */
-	static FSlateIcon FindIconForClass(const UClass* InClass, const FName& InDefaultName = FName());
+	static SLATECORE_API FSlateIcon FindIconForClass(const UClass* InClass, const FName& InDefaultName = FName());
 
 	/**
 	 * Find a custom icon to use for the supplied class, according to the specified base style
@@ -29,7 +29,7 @@ public:
 	 * @param InDefaultName		The default icon name to use, if no specialized icon could be found
 	 * @return An FSlateIcon structure defining the class's icon
 	 */
-	static FSlateIcon FindCustomIconForClass(const UClass* InClass, const TCHAR* StyleBasePath, const FName& InDefaultName = FName());
+	static SLATECORE_API FSlateIcon FindCustomIconForClass(const UClass* InClass, const TCHAR* StyleBasePath, const FName& InDefaultName = FName());
 
 	/**
 	 * Find a slate brush to use for the supplied class's icon
@@ -61,5 +61,5 @@ public:
 	 * @param InIconName 		The fully qualified style name of the icon to find
 	 * @return An FSlateIcon structure
 	 */
-	static FSlateIcon FindIcon(const FName& InIconName);
+	static SLATECORE_API FSlateIcon FindIcon(const FName& InIconName);
 };
diff --git a/Engine/Source/Runtime/SlateCore/Public/Styling/SlateStyle.h b/Engine/Source/Runtime/SlateCore/Public/Styling/SlateStyle.h
--- a/Engine/Source/Runtime/SlateCore/Public/Styling/SlateStyle.h
+++ b/Engine/Source/Runtime/SlateCore/Public/Styling/SlateStyle.h
@@ -24,59 +24,59 @@ struct FSlateWidgetStyle;
  * A slate style chunk that contains a collection of named properties that guide the appearance of Slate.
  * At the moment, basically FAppStyle.
  */
-class SLATECORE_API FSlateStyleSet : public ISlateStyle
+class FSlateStyleSet : public ISlateStyle
 {
 public:
 	/**
 	 * Construct a style chunk.
 	 * @param InStyleSetName The name used to identity this style set
 	 */
-	FSlateStyleSet(const FName& InStyleSetName);
+	SLATECORE_API FSlateStyleSet(const FName& InStyleSetName);
 
 	/** Destructor. */
-	virtual ~FSlateStyleSet();
-
-	virtual const FName& GetStyleSetName() const override;
-	virtual void GetResources(TArray< const FSlateBrush* >& OutResources) const override;
-	virtual TArray<FName> GetEntriesUsingBrush(const FName BrushName) const override;
-	virtual void SetContentRoot(const FString& InContentRootDir);
-	virtual FString RootToContentDir(const ANSICHAR* RelativePath, const TCHAR* Extension);
-	virtual FString RootToContentDir(const WIDECHAR* RelativePath, const TCHAR* Extension);
-	virtual FString RootToContentDir(const FString& RelativePath, const TCHAR* Extension);
-	virtual FString RootToContentDir(const ANSICHAR* RelativePath);
-	virtual FString RootToContentDir(const WIDECHAR* RelativePath);
-	virtual FString RootToContentDir(const FString& RelativePath);
+	SLATECORE_API virtual ~FSlateStyleSet();
+
+	SLATECORE_API virtual const FName& GetStyleSetName() const override;
+	SLATECORE_API virtual void GetResources(TArray< const FSlateBrush* >& OutResources) const override;
+	SLATECORE_API virtual TArray<FName> GetEntriesUsingBrush(const FName BrushName) const override;
+	SLATECORE_API virtual void SetContentRoot(const FString& InContentRootDir);
+	SLATECORE_API virtual FString RootToContentDir(const ANSICHAR* RelativePath, const TCHAR* Extension);
+	SLATECORE_API virtual FString RootToContentDir(const WIDECHAR* RelativePath, const TCHAR* Extension);
+	SLATECORE_API virtual FString RootToContentDir(const FString& RelativePath, const TCHAR* Extension);
+	SLATECORE_API virtual FString RootToContentDir(const ANSICHAR* RelativePath);
+	SLATECORE_API virtual FString RootToContentDir(const WIDECHAR* RelativePath);
+	SLATECORE_API virtual FString RootToContentDir(const FString& RelativePath);
 	virtual FString GetContentRootDir() const { return ContentRootDir; }
 
-	virtual void SetCoreContentRoot(const FString& InCoreContentRootDir);
+	SLATECORE_API virtual void SetCoreContentRoot(const FString& InCoreContentRootDir);
 
-	virtual FString RootToCoreContentDir(const ANSICHAR* RelativePath, const TCHAR* Extension);
-	virtual FString RootToCoreContentDir(const WIDECHAR* RelativePath, const TCHAR* Extension);
-	virtual FString RootToCoreContentDir(const FString& RelativePath, const TCHAR* Extension);
-	virtual FString RootToCoreContentDir(const ANSICHAR* RelativePath);
-	virtual FString RootToCoreContentDir(const WIDECHAR* RelativePath);
-	virtual FString RootToCoreContentDir(const FString& RelativePath);
+	SLATECORE_API virtual FString RootToCoreContentDir(const ANSICHAR* RelativePath, const TCHAR* Extension);
+	SLATECORE_API virtual FString RootToCoreContentDir(const WIDECHAR* RelativePath, const TCHAR* Extension);
+	SLATECORE_API virtual FString RootToCoreContentDir(const FString& RelativePath, const TCHAR* Extension);
+	SLATECORE_API virtual FString RootToCoreContentDir(const ANSICHAR* RelativePath);
+	SLATECORE_API virtual FString RootToCoreContentDir(const WIDECHAR* RelativePath);
+	SLATECORE_API virtual FString RootToCoreContentDir(const FString& RelativePath);
 
-	virtual float GetFloat(const FName PropertyName, const ANSICHAR* Specifier = nullptr, float DefaultValue = FStyleDefaults::GetFloat(), const ISlateStyle* RequestingStyle = nullptr) const override;
-	virtual FVector2D GetVector(const FName PropertyName, const ANSICHAR* Specifier = nullptr, FVector2D DefaultValue = FStyleDefaults::GetVector2D(), const ISlateStyle* RequestingStyle = nullptr) const override;
-	virtual const FLinearColor& GetColor(const FName PropertyName, const ANSICHAR* Specifier = nullptr, const FLinearColor& DefaultValue = FStyleDefaults::GetColor(), const ISlateStyle* RequestingStyle = nullptr) const override;
-	virtual const FSlateColor GetSlateColor(const FName PropertyName, const ANSICHAR* Specifier = nullptr, const FSlateColor& DefaultValue = FStyleDefaults::GetSlateColor(), const ISlateStyle* RequestingStyle = nullptr) const override;
-	virtual const FMargin& GetMargin(const FName PropertyName, const ANSICHAR* Specifier = nullptr, const FMargin& DefaultValue = FStyleDefaults::GetMargin(), const ISlateStyle* RequestingStyle = nullptr) const override;
+	SLATECORE_API virtual float GetFloat(const FName PropertyName, const ANSICHAR* Specifier = nullptr, float DefaultValue = FStyleDefaults::GetFloat(), const ISlateStyle* RequestingStyle = nullptr) const override;
+	SLATECORE_API virtual FVector2D GetVector(const FName PropertyName, const ANSICHAR* Specifier = nullptr, FVector2D DefaultValue = FStyleDefaults::GetVector2D(), const ISlateStyle* RequestingStyle = nullptr) const override;
+	SLATECORE_API virtual const FLinearColor& GetColor(const FName PropertyName, const ANSICHAR* Specifier = nullptr, const FLinearColor& DefaultValue = FStyleDefaults::GetColor(), const ISlateStyle* RequestingStyle = nullptr) const override;
+	SLATECORE_API virtual const FSlateColor GetSlateColor(const FName PropertyName, const ANSICHAR* Specifier = nullptr, const FSlateColor& DefaultValue = FStyleDefaults::GetSlateColor(), const ISlateStyle* RequestingStyle = nullptr) const override;
+	SLATECORE_API virtual const FMargin& GetMargin(const FName PropertyName, const ANSICHAR* Specifier = nullptr, const FMargin& DefaultValue = FStyleDefaults::GetMargin(), const ISlateStyle* RequestingStyle = nullptr) const override;
 
-	virtual const FSlateBrush* GetBrush(const FName PropertyName, const ANSICHAR* Specifier = nullptr, const ISlateStyle* RequestingStyle = nullptr) const override;
-	virtual const FSlateBrush* GetOptionalBrush(const FName PropertyName, const ANSICHAR* Specifier = nullptr, const FSlateBrush* const DefaultBrush = FStyleDefaults::GetNoBrush()) const override;
+	SLATECORE_API virtual const FSlateBrush* GetBrush(const FName PropertyName, const ANSICHAR* Specifier = nullptr, const ISlateStyle* RequestingStyle = nullptr) const override;
+	SLATECORE_API virtual const FSlateBrush* GetOptionalBrush(const FName PropertyName, const ANSICHAR* Specifier = nullptr, const FSlateBrush* const DefaultBrush = FStyleDefaults::GetNoBrush()) const override;
 
-	virtual const TSharedPtr<FSlateDynamicImageBrush> GetDynamicImageBrush(const FName BrushTemplate, const FName TextureName, const ANSICHAR* Specifier = nullptr, const ISlateStyle* RequestingStyle = nullptr) const override;
-	virtual const TSharedPtr<FSlateDynamicImageBrush> GetDynamicImageBrush(const FName BrushTemplate, const ANSICHAR* Specifier, UTexture2D* TextureResource, const FName TextureName, const ISlateStyle* RequestingStyle = nullptr) const override;
-	virtual const TSharedPtr<FSlateDynamicImageBrush> GetDynamicImageBrush(const FName BrushTemplate, UTexture2D* TextureResource, const FName TextureName, const ISlateStyle* RequestingStyle = nullptr) const override;
+	SLATECORE_API virtual const TSharedPtr<FSlateDynamicImageBrush> GetDynamicImageBrush(const FName BrushTemplate, const FName TextureName, const ANSICHAR* Specifier = nullptr, const ISlateStyle* RequestingStyle = nullptr) const override;
+	SLATECORE_API virtual const TSharedPtr<FSlateDynamicImageBrush> GetDynamicImageBrush(const FName BrushTemplate, const ANSICHAR* Specifier, UTexture2D* TextureResource, const FName TextureName, const ISlateStyle* RequestingStyle = nullptr) const override;
+	SLATECORE_API virtual const TSharedPtr<FSlateDynamicImageBrush> GetDynamicImageBrush(const FName BrushTemplate, UTexture2D* TextureResource, const FName TextureName, const ISlateStyle* RequestingStyle = nullptr) const override;
 
-	virtual FSlateBrush* GetDefaultBrush() const override;
+	SLATECORE_API virtual FSlateBrush* GetDefaultBrush() const override;
 
-	virtual const FSlateSound& GetSound(const FName PropertyName, const ANSICHAR* Specifier = nullptr, const ISlateStyle* RequestingStyle = nullptr) const override;
-	virtual FSlateFontInfo GetFontStyle(const FName PropertyName, const ANSICHAR* Specifier = nullptr) const override;
+	SLATECORE_API virtual const FSlateSound& GetSound(const FName PropertyName, const ANSICHAR* Specifier = nullptr, const ISlateStyle* RequestingStyle = nullptr) const override;
+	SLATECORE_API virtual FSlateFontInfo GetFontStyle(const FName PropertyName, const ANSICHAR* Specifier = nullptr) const override;
 
 	/** Name a Parent Style to fallback on if the requested styles are specified in this style **/
-	void SetParentStyleName(const FName& InParentStyleName);
+	SLATECORE_API void SetParentStyleName(const FName& InParentStyleName);
 public:
 
 	template< typename DefinitionType >
@@ -200,20 +200,20 @@ protected:
 	friend class FSlateStyleSet;
 	friend class ISlateStyle;
 
-	virtual const FSlateWidgetStyle* GetWidgetStyleInternal(const FName DesiredTypeName, const FName StyleName, const FSlateWidgetStyle* DefaultStyle, bool bWarnIfNotFound) const override;
+	SLATECORE_API virtual const FSlateWidgetStyle* GetWidgetStyleInternal(const FName DesiredTypeName, const FName StyleName, const FSlateWidgetStyle* DefaultStyle, bool bWarnIfNotFound) const override;
 
-	virtual void Log(ISlateStyle::EStyleMessageSeverity Severity, const FText& Message) const override;
+	SLATECORE_API virtual void Log(ISlateStyle::EStyleMessageSeverity Severity, const FText& Message) const override;
 
-	virtual void LogMissingResource(EStyleMessageSeverity Severity, const FText& Message, const FName& MissingResource) const override;
+	SLATECORE_API virtual void LogMissingResource(EStyleMessageSeverity Severity, const FText& Message, const FName& MissingResource) const override;
 
-	virtual const TSharedPtr< FSlateDynamicImageBrush > MakeDynamicImageBrush(const FName BrushTemplate, UTexture2D* TextureResource, const FName TextureName) const override;
+	SLATECORE_API virtual const TSharedPtr< FSlateDynamicImageBrush > MakeDynamicImageBrush(const FName BrushTemplate, UTexture2D* TextureResource, const FName TextureName) const override;
 
-	virtual void LogUnusedBrushResources();
+	SLATECORE_API virtual void LogUnusedBrushResources();
 
-	TSet<FName> GetStyleKeys() const;
+	SLATECORE_API TSet<FName> GetStyleKeys() const;
 protected:
 
-	bool IsBrushFromFile(const FString& FilePath, const FSlateBrush* Brush);
+	SLATECORE_API bool IsBrushFromFile(const FString& FilePath, const FSlateBrush* Brush);
 
 	/** The name used to identity this style set */
 	FName StyleSetName;
@@ -258,6 +258,6 @@ protected:
 	mutable TSet< FName > MissingResources;
 
 	FName ParentStyleName;
-	const ISlateStyle* GetParentStyle() const;
+	SLATECORE_API const ISlateStyle* GetParentStyle() const;
 };
 
diff --git a/Engine/Source/Runtime/SlateCore/Public/Styling/SlateStyleRegistry.h b/Engine/Source/Runtime/SlateCore/Public/Styling/SlateStyleRegistry.h
--- a/Engine/Source/Runtime/SlateCore/Public/Styling/SlateStyleRegistry.h
+++ b/Engine/Source/Runtime/SlateCore/Public/Styling/SlateStyleRegistry.h
@@ -10,7 +10,7 @@ struct FSlateBrush;
 /**
 * A central repository that can be used to track and manage chunks of slate style data.
 */
-class SLATECORE_API FSlateStyleRegistry
+class FSlateStyleRegistry
 {
 public:
 
@@ -19,21 +19,21 @@ public:
 	 *
 	 * @param ISlateStyle interface to a user-definable style to add.
 	 */
-	static void RegisterSlateStyle( const ISlateStyle& InSlateStyle );
+	static SLATECORE_API void RegisterSlateStyle( const ISlateStyle& InSlateStyle );
 
 	/**
 	 * Remove a slate style from the repository.
 	 *
 	 * @param ISlateStyle Interface to the style to remove.
 	 */
-	static void UnRegisterSlateStyle( const ISlateStyle& InSlateStyle );
+	static SLATECORE_API void UnRegisterSlateStyle( const ISlateStyle& InSlateStyle );
 
 	/**
 	 * Removes a slate style from the repository (by name).
 	 *
 	 * @param StyleSetName The name of the style set to remove.
 	 */
-	static void UnRegisterSlateStyle( const FName StyleSetName );
+	static SLATECORE_API void UnRegisterSlateStyle( const FName StyleSetName );
 
 	/**
 	 * Find a slate style in the repository.
@@ -41,7 +41,7 @@ public:
 	 * @param InSlateStyleName The name of the slate style to find.
 	 * @return The slate style, or null if it couldn't be found
 	 */
-	static const ISlateStyle* FindSlateStyle( const FName& InSlateStyleName );
+	static SLATECORE_API const ISlateStyle* FindSlateStyle( const FName& InSlateStyleName );
 
 	/**
 	 * Iterate all the slate style known to this registry
@@ -49,14 +49,14 @@ public:
 	 * @param Iter A predicate to call for each known style that returns true to continue iteration, and false to terminate iteration
 	 * @return true where to loop completed, of false if it was terminated by the predicate
 	 */
-	static bool IterateAllStyles(const TFunctionRef<bool(const ISlateStyle&)>& Iter);
+	static SLATECORE_API bool IterateAllStyles(const TFunctionRef<bool(const ISlateStyle&)>& Iter);
 
 	/**
 	 * Populate an array of slate brushes with all of the resources used by the registered styles.
 	 *
 	 * @param OutResources Array of slate brushes to populate.
 	 */
-	static void GetAllResources( TArray<const FSlateBrush*>& OutResources );
+	static SLATECORE_API void GetAllResources( TArray<const FSlateBrush*>& OutResources );
 
 	/**
 	 * Gets the names of every style entry using a brush
@@ -65,7 +65,7 @@ public:
 	 * @param BrushName The name of the brush to find style entries from
 	 * @return Array of style names using the brush
 	 */
-	static TArray<FName> GetSylesUsingBrush(const FName BrushName);
+	static SLATECORE_API TArray<FName> GetSylesUsingBrush(const FName BrushName);
 
 private:
 
