[SlateCore]

diff --git a/Engine/Source/Runtime/SlateCore/Public/Styling/SlateTypes.h b/Engine/Source/Runtime/SlateCore/Public/Styling/SlateTypes.h
--- a/Engine/Source/Runtime/SlateCore/Public/Styling/SlateTypes.h
+++ b/Engine/Source/Runtime/SlateCore/Public/Styling/SlateTypes.h
@@ -89,20 +89,20 @@ enum class ETextOverflowPolicy : uint8
  * Represents the appearance of an SCheckBox
  */
 USTRUCT(BlueprintType)
-struct SLATECORE_API FCheckBoxStyle : public FSlateWidgetStyle
+struct FCheckBoxStyle : public FSlateWidgetStyle
 {
 	GENERATED_USTRUCT_BODY()
 
-	FCheckBoxStyle();
+	SLATECORE_API FCheckBoxStyle();
 
 	virtual ~FCheckBoxStyle() {}
 
-	virtual void GetResources( TArray< const FSlateBrush* > & OutBrushes ) const override;
+	SLATECORE_API virtual void GetResources( TArray< const FSlateBrush* > & OutBrushes ) const override;
 
-	static const FName TypeName;
+	static SLATECORE_API const FName TypeName;
 	virtual const FName GetTypeName() const override { return TypeName; };
 
-	static const FCheckBoxStyle& GetDefault();
+	static SLATECORE_API const FCheckBoxStyle& GetDefault();
 
 	/** The visual type of the checkbox */	
 	UPROPERTY( EditAnywhere, BlueprintReadWrite, Category=Appearance )
@@ -246,7 +246,7 @@ struct SLATECORE_API FCheckBoxStyle : public FSlateWidgetStyle
 	/**
 	 * Used to upgrade the deprecated FName sound properties into the new-style FSlateSound properties
 	 */	
-	void PostSerialize(const FArchive& Ar);
+	SLATECORE_API void PostSerialize(const FArchive& Ar);
 #endif
 
 	/**
@@ -310,20 +310,20 @@ enum class ETextTransformPolicy : uint8
  * Represents the appearance of an STextBlock
  */
 USTRUCT(BlueprintType)
-struct SLATECORE_API FTextBlockStyle : public FSlateWidgetStyle
+struct FTextBlockStyle : public FSlateWidgetStyle
 {
 	GENERATED_USTRUCT_BODY()
 
-	FTextBlockStyle();
+	SLATECORE_API FTextBlockStyle();
 
 	virtual ~FTextBlockStyle() {}
 
-	virtual void GetResources( TArray< const FSlateBrush* >& OutBrushes ) const override;
+	SLATECORE_API virtual void GetResources( TArray< const FSlateBrush* >& OutBrushes ) const override;
 
-	static const FName TypeName;
+	static SLATECORE_API const FName TypeName;
 	virtual const FName GetTypeName() const override { return TypeName; };
 
-	static const FTextBlockStyle& GetDefault();
+	static SLATECORE_API const FTextBlockStyle& GetDefault();
 
 	/** Font family and size to be used when displaying this text. */
 	UPROPERTY(EditAnywhere, BlueprintReadWrite, Category=Appearance)
@@ -430,20 +430,20 @@ struct SLATECORE_API FTextBlockStyle : public FSlateWidgetStyle
  * Represents the appearance of an SButton
  */
 USTRUCT(BlueprintType)
-struct SLATECORE_API FButtonStyle : public FSlateWidgetStyle
+struct FButtonStyle : public FSlateWidgetStyle
 {
 	GENERATED_USTRUCT_BODY()
 
-	FButtonStyle();
+	SLATECORE_API FButtonStyle();
 
 	virtual ~FButtonStyle() {}
 
-	virtual void GetResources( TArray< const FSlateBrush* >& OutBrushes ) const override;
+	SLATECORE_API virtual void GetResources( TArray< const FSlateBrush* >& OutBrushes ) const override;
 
-	static const FName TypeName;
+	static SLATECORE_API const FName TypeName;
 	virtual const FName GetTypeName() const override { return TypeName; };
 
-	static const FButtonStyle& GetDefault();
+	static SLATECORE_API const FButtonStyle& GetDefault();
 
 	/** Button appearance when the button is not hovered or pressed */
 	UPROPERTY(EditAnywhere, BlueprintReadWrite, Category=Appearance)
@@ -526,7 +526,7 @@ struct SLATECORE_API FButtonStyle : public FSlateWidgetStyle
 	/**
 	 * Used to upgrade the deprecated FName sound properties into the new-style FSlateSound properties
 	 */	
-	void PostSerialize(const FArchive& Ar);
+	SLATECORE_API void PostSerialize(const FArchive& Ar);
 #endif
 
 	/**
@@ -558,20 +558,20 @@ struct TStructOpsTypeTraits<FButtonStyle> : public TStructOpsTypeTraitsBase2<FBu
  * Represents the appearance of an SComboButton
  */
 USTRUCT(BlueprintType)
-struct SLATECORE_API FComboButtonStyle : public FSlateWidgetStyle
+struct FComboButtonStyle : public FSlateWidgetStyle
 {
 	GENERATED_USTRUCT_BODY()
 
-	FComboButtonStyle();
+	SLATECORE_API FComboButtonStyle();
 
 	virtual ~FComboButtonStyle() {}
 
-	virtual void GetResources( TArray< const FSlateBrush* >& OutBrushes ) const override;
+	SLATECORE_API virtual void GetResources( TArray< const FSlateBrush* >& OutBrushes ) const override;
 
-	static const FName TypeName;
+	static SLATECORE_API const FName TypeName;
 	virtual const FName GetTypeName() const override { return TypeName; };
 
-	static const FComboButtonStyle& GetDefault();
+	static SLATECORE_API const FComboButtonStyle& GetDefault();
 
 	/**
 	 * The style to use for our SButton.
@@ -655,20 +655,20 @@ struct SLATECORE_API FComboButtonStyle : public FSlateWidgetStyle
  * Represents the appearance of an SComboBox
  */
 USTRUCT(BlueprintType)
-struct SLATECORE_API FComboBoxStyle : public FSlateWidgetStyle
+struct FComboBoxStyle : public FSlateWidgetStyle
 {
 	GENERATED_USTRUCT_BODY()
 
-	FComboBoxStyle();
+	SLATECORE_API FComboBoxStyle();
 
 	virtual ~FComboBoxStyle() {}
 
-	virtual void GetResources( TArray< const FSlateBrush* >& OutBrushes ) const override;
+	SLATECORE_API virtual void GetResources( TArray< const FSlateBrush* >& OutBrushes ) const override;
 
-	static const FName TypeName;
+	static SLATECORE_API const FName TypeName;
 	virtual const FName GetTypeName() const override { return TypeName; };
 
-	static const FComboBoxStyle& GetDefault();
+	static SLATECORE_API const FComboBoxStyle& GetDefault();
 
 	/**
 	 * The style to use for our SComboButton
@@ -716,7 +716,7 @@ struct SLATECORE_API FComboBoxStyle : public FSlateWidgetStyle
 	/**
 	 * Used to upgrade the deprecated FName sound properties into the new-style FSlateSound properties
 	 */	
-	void PostSerialize(const FArchive& Ar);
+	SLATECORE_API void PostSerialize(const FArchive& Ar);
 #endif
 
 	/**
@@ -745,20 +745,20 @@ struct TStructOpsTypeTraits<FComboBoxStyle> : public TStructOpsTypeTraitsBase2<F
  * Represents the appearance of an SHyperlink
  */
 USTRUCT(BlueprintType)
-struct SLATECORE_API FHyperlinkStyle : public FSlateWidgetStyle
+struct FHyperlinkStyle : public FSlateWidgetStyle
 {
 	GENERATED_USTRUCT_BODY()
 
-	FHyperlinkStyle();
+	SLATECORE_API FHyperlinkStyle();
 
 	virtual ~FHyperlinkStyle() {}
 
-	virtual void GetResources( TArray< const FSlateBrush* >& OutBrushes ) const override;
+	SLATECORE_API virtual void GetResources( TArray< const FSlateBrush* >& OutBrushes ) const override;
 
-	static const FName TypeName;
+	static SLATECORE_API const FName TypeName;
 	virtual const FName GetTypeName() const override { return TypeName; };
 
-	static const FHyperlinkStyle& GetDefault();
+	static SLATECORE_API const FHyperlinkStyle& GetDefault();
 
 	/** Underline style */
 	UPROPERTY(EditAnywhere, BlueprintReadWrite, Category=Appearance)
@@ -780,20 +780,20 @@ struct SLATECORE_API FHyperlinkStyle : public FSlateWidgetStyle
  * Represents the appearance of an SEditableText
  */
 USTRUCT(BlueprintType)
-struct SLATECORE_API FEditableTextStyle : public FSlateWidgetStyle
+struct FEditableTextStyle : public FSlateWidgetStyle
 {
 	GENERATED_USTRUCT_BODY()
 
-	FEditableTextStyle();
+	SLATECORE_API FEditableTextStyle();
 
 	virtual ~FEditableTextStyle() {}
 
-	virtual void GetResources( TArray< const FSlateBrush* >& OutBrushes ) const override;
+	SLATECORE_API virtual void GetResources( TArray< const FSlateBrush* >& OutBrushes ) const override;
 
-	static const FName TypeName;
+	static SLATECORE_API const FName TypeName;
 	virtual const FName GetTypeName() const override { return TypeName; }
 
-	static const FEditableTextStyle& GetDefault();
+	static SLATECORE_API const FEditableTextStyle& GetDefault();
 
 	/** Font family and size to be used when displaying this text. */
 	UPROPERTY(EditAnywhere, BlueprintReadWrite, Category=Appearance)
@@ -841,20 +841,20 @@ struct SLATECORE_API FEditableTextStyle : public FSlateWidgetStyle
  * Represents the appearance of an SScrollBar
  */
 USTRUCT(BlueprintType)
-struct SLATECORE_API FScrollBarStyle : public FSlateWidgetStyle
+struct FScrollBarStyle : public FSlateWidgetStyle
 {
 	GENERATED_USTRUCT_BODY()
 
-	FScrollBarStyle();
+	SLATECORE_API FScrollBarStyle();
 
 	virtual ~FScrollBarStyle() {}
 
-	virtual void GetResources( TArray< const FSlateBrush* >& OutBrushes ) const override;
+	SLATECORE_API virtual void GetResources( TArray< const FSlateBrush* >& OutBrushes ) const override;
 
-	static const FName TypeName;
+	static SLATECORE_API const FName TypeName;
 	virtual const FName GetTypeName() const override { return TypeName; };
 
-	static const FScrollBarStyle& GetDefault();
+	static SLATECORE_API const FScrollBarStyle& GetDefault();
 
 	/** Background image to use when the scrollbar is oriented horizontally */
 	UPROPERTY(EditAnywhere, BlueprintReadWrite, Category=Appearance)
@@ -928,25 +928,25 @@ struct SLATECORE_API FScrollBarStyle : public FSlateWidgetStyle
  * Represents the appearance of an SEditableTextBox
  */
 USTRUCT(BlueprintType)
-struct SLATECORE_API FEditableTextBoxStyle : public FSlateWidgetStyle
+struct FEditableTextBoxStyle : public FSlateWidgetStyle
 {
 	GENERATED_USTRUCT_BODY()
 
-	PRAGMA_DISABLE_DEPRECATION_WARNINGS
+	SLATECORE_API PRAGMA_DISABLE_DEPRECATION_WARNINGS
 	FEditableTextBoxStyle();
 	FEditableTextBoxStyle(const FEditableTextBoxStyle&) = default;
 
 	FEditableTextBoxStyle& operator=(const FEditableTextBoxStyle&) = default;
 
 	virtual ~FEditableTextBoxStyle() = default;
-	PRAGMA_ENABLE_DEPRECATION_WARNINGS
+	SLATECORE_API PRAGMA_ENABLE_DEPRECATION_WARNINGS
 
 	virtual void GetResources( TArray< const FSlateBrush* >& OutBrushes ) const override;
 
-	static const FName TypeName;
+	static SLATECORE_API const FName TypeName;
 	virtual const FName GetTypeName() const override { return TypeName; };
 
-	static const FEditableTextBoxStyle& GetDefault();
+	static SLATECORE_API const FEditableTextBoxStyle& GetDefault();
 
 	/** Border background image when the box is not hovered or focused */
 	UPROPERTY(EditAnywhere, BlueprintReadWrite, Category=Appearance)
@@ -1045,20 +1045,20 @@ struct SLATECORE_API FEditableTextBoxStyle : public FSlateWidgetStyle
  * Represents the appearance of an SInlineEditableTextBlock
  */
 USTRUCT(BlueprintType)
-struct SLATECORE_API FInlineEditableTextBlockStyle : public FSlateWidgetStyle
+struct FInlineEditableTextBlockStyle : public FSlateWidgetStyle
 {
 	GENERATED_USTRUCT_BODY()
 
-	FInlineEditableTextBlockStyle();
+	SLATECORE_API FInlineEditableTextBlockStyle();
 
 	virtual ~FInlineEditableTextBlockStyle() {}
 
-	virtual void GetResources( TArray< const FSlateBrush* >& OutBrushes ) const override;
+	SLATECORE_API virtual void GetResources( TArray< const FSlateBrush* >& OutBrushes ) const override;
 
-	static const FName TypeName;
+	static SLATECORE_API const FName TypeName;
 	virtual const FName GetTypeName() const override { return TypeName; };
 
-	static const FInlineEditableTextBlockStyle& GetDefault();
+	static SLATECORE_API const FInlineEditableTextBlockStyle& GetDefault();
 
 	/** The style of the editable text box, which dictates the font, color, and shadow options. */
 	UPROPERTY(EditAnywhere, BlueprintReadWrite, Category=Appearance)
@@ -1076,20 +1076,20 @@ struct SLATECORE_API FInlineEditableTextBlockStyle : public FSlateWidgetStyle
  * Represents the appearance of an SProgressBar
  */
 USTRUCT(BlueprintType)
-struct SLATECORE_API FProgressBarStyle : public FSlateWidgetStyle
+struct FProgressBarStyle : public FSlateWidgetStyle
 {
 	GENERATED_USTRUCT_BODY()
 
-	FProgressBarStyle();
+	SLATECORE_API FProgressBarStyle();
 
 	virtual ~FProgressBarStyle() {}
 
-	virtual void GetResources( TArray< const FSlateBrush* >& OutBrushes ) const override;
+	SLATECORE_API virtual void GetResources( TArray< const FSlateBrush* >& OutBrushes ) const override;
 
-	static const FName TypeName;
+	static SLATECORE_API const FName TypeName;
 	virtual const FName GetTypeName() const override { return TypeName; };
 
-	static const FProgressBarStyle& GetDefault();
+	static SLATECORE_API const FProgressBarStyle& GetDefault();
 
 	/** Background image to use for the progress bar */
 	UPROPERTY(EditAnywhere, BlueprintReadWrite, Category=Appearance)
@@ -1129,20 +1129,20 @@ struct SLATECORE_API FProgressBarStyle : public FSlateWidgetStyle
  * Represents the appearance of an SExpandableArea
  */
 USTRUCT()
-struct SLATECORE_API FExpandableAreaStyle : public FSlateWidgetStyle
+struct FExpandableAreaStyle : public FSlateWidgetStyle
 {
 	GENERATED_USTRUCT_BODY()
 
-	FExpandableAreaStyle();
+	SLATECORE_API FExpandableAreaStyle();
 
 	virtual ~FExpandableAreaStyle() {}
 
-	virtual void GetResources( TArray< const FSlateBrush* >& OutBrushes ) const override;
+	SLATECORE_API virtual void GetResources( TArray< const FSlateBrush* >& OutBrushes ) const override;
 
-	static const FName TypeName;
+	static SLATECORE_API const FName TypeName;
 	virtual const FName GetTypeName() const override { return TypeName; };
 
-	static const FExpandableAreaStyle& GetDefault();
+	static SLATECORE_API const FExpandableAreaStyle& GetDefault();
 
 	/** Image to use when the area is collapsed */
 	UPROPERTY(EditAnywhere, Category=Appearance)
@@ -1175,25 +1175,25 @@ struct SLATECORE_API FExpandableAreaStyle : public FSlateWidgetStyle
  * Represents the appearance of an SSearchBox
  */
 USTRUCT()
-struct SLATECORE_API FSearchBoxStyle : public FSlateWidgetStyle
+struct FSearchBoxStyle : public FSlateWidgetStyle
 {
 	GENERATED_USTRUCT_BODY()
 
-	FSearchBoxStyle();
+	SLATECORE_API FSearchBoxStyle();
 
 	virtual ~FSearchBoxStyle() {}
 
-	virtual void GetResources( TArray< const FSlateBrush* >& OutBrushes ) const override;
+	SLATECORE_API virtual void GetResources( TArray< const FSlateBrush* >& OutBrushes ) const override;
 
-	static const FName TypeName;
+	static SLATECORE_API const FName TypeName;
 	virtual const FName GetTypeName() const override { return TypeName; };
 
-	static const FSearchBoxStyle& GetDefault();
+	static SLATECORE_API const FSearchBoxStyle& GetDefault();
 
 	/** Style to use for the text box part of the search box */
 	UPROPERTY(EditAnywhere, Category=Appearance)
 	FEditableTextBoxStyle TextBoxStyle;
-	FSearchBoxStyle& SetTextBoxStyle( const FEditableTextBoxStyle& InTextBoxStyle );
+	SLATECORE_API FSearchBoxStyle& SetTextBoxStyle( const FEditableTextBoxStyle& InTextBoxStyle );
 
 	/** Font to use for the text box part of the search box when a search term is entered*/
 	UPROPERTY(EditAnywhere, Category = Appearance)
@@ -1228,7 +1228,7 @@ struct SLATECORE_API FSearchBoxStyle : public FSlateWidgetStyle
 	/** If true, buttons appear to the left of the search text */
 	UPROPERTY(meta = (DeprecatedProperty, DeprecationMessage="Use LeftAlignSearchResultButtons and LeftAlignGlassImageAndClearButton instead"))
 	bool bLeftAlignButtons_DEPRECATED;
-	FSearchBoxStyle& SetLeftAlignButtons(bool bInLeftAlignButtons);
+	SLATECORE_API FSearchBoxStyle& SetLeftAlignButtons(bool bInLeftAlignButtons);
 
 	/** If true, search result buttons appear to the left of the search text */
 	UPROPERTY(EditAnywhere, Category = Appearance)
@@ -1246,20 +1246,20 @@ struct SLATECORE_API FSearchBoxStyle : public FSlateWidgetStyle
  * Represents the appearance of an SSlider
  */
 USTRUCT(BlueprintType)
-struct SLATECORE_API FSliderStyle : public FSlateWidgetStyle
+struct FSliderStyle : public FSlateWidgetStyle
 {
 	GENERATED_USTRUCT_BODY()
 
-	FSliderStyle();
+	SLATECORE_API FSliderStyle();
 
 	virtual ~FSliderStyle() {}
 
-	virtual void GetResources( TArray< const FSlateBrush* >& OutBrushes ) const override;
+	SLATECORE_API virtual void GetResources( TArray< const FSlateBrush* >& OutBrushes ) const override;
 
-	static const FName TypeName;
+	static SLATECORE_API const FName TypeName;
 	virtual const FName GetTypeName() const override { return TypeName; };
 
-	static const FSliderStyle& GetDefault();
+	static SLATECORE_API const FSliderStyle& GetDefault();
 
 	/** Image to use when the slider bar is in its normal state */
 	UPROPERTY(EditAnywhere, BlueprintReadWrite, Category=Appearance)
@@ -1315,20 +1315,20 @@ struct SLATECORE_API FSliderStyle : public FSlateWidgetStyle
  * Represents the appearance of an SVolumeControl
  */
 USTRUCT()
-struct SLATECORE_API FVolumeControlStyle : public FSlateWidgetStyle
+struct FVolumeControlStyle : public FSlateWidgetStyle
 {
 	GENERATED_USTRUCT_BODY()
 
-	FVolumeControlStyle();
+	SLATECORE_API FVolumeControlStyle();
 
 	virtual ~FVolumeControlStyle() {}
 
-	virtual void GetResources( TArray< const FSlateBrush* >& OutBrushes ) const override;
+	SLATECORE_API virtual void GetResources( TArray< const FSlateBrush* >& OutBrushes ) const override;
 
-	static const FName TypeName;
+	static SLATECORE_API const FName TypeName;
 	virtual const FName GetTypeName() const override { return TypeName; };
 
-	static const FVolumeControlStyle& GetDefault();
+	static SLATECORE_API const FVolumeControlStyle& GetDefault();
 
 	/** The style of the volume control slider */
 	UPROPERTY(EditAnywhere, Category=Appearance)
@@ -1365,20 +1365,20 @@ struct SLATECORE_API FVolumeControlStyle : public FSlateWidgetStyle
  * Represents the appearance of an inline image used by rich text
  */
 USTRUCT()
-struct SLATECORE_API FInlineTextImageStyle : public FSlateWidgetStyle
+struct FInlineTextImageStyle : public FSlateWidgetStyle
 {
 	GENERATED_USTRUCT_BODY()
 
-	FInlineTextImageStyle();
+	SLATECORE_API FInlineTextImageStyle();
 
 	virtual ~FInlineTextImageStyle() {}
 
-	virtual void GetResources( TArray< const FSlateBrush* >& OutBrushes ) const override;
+	SLATECORE_API virtual void GetResources( TArray< const FSlateBrush* >& OutBrushes ) const override;
 
-	static const FName TypeName;
+	static SLATECORE_API const FName TypeName;
 	virtual const FName GetTypeName() const override { return TypeName; };
 
-	static const FInlineTextImageStyle& GetDefault();
+	static SLATECORE_API const FInlineTextImageStyle& GetDefault();
 
 	/** Image to use when the slider thumb is in its normal state */
 	UPROPERTY(EditAnywhere, Category=Appearance)
@@ -1395,20 +1395,20 @@ struct SLATECORE_API FInlineTextImageStyle : public FSlateWidgetStyle
  * Represents the appearance of an SSpinBox
  */
 USTRUCT(BlueprintType)
-struct SLATECORE_API FSpinBoxStyle : public FSlateWidgetStyle
+struct FSpinBoxStyle : public FSlateWidgetStyle
 {
 	GENERATED_USTRUCT_BODY()
 
-	FSpinBoxStyle();
+	SLATECORE_API FSpinBoxStyle();
 
 	virtual ~FSpinBoxStyle() {}
 
-	virtual void GetResources( TArray< const FSlateBrush* >& OutBrushes ) const override;
+	SLATECORE_API virtual void GetResources( TArray< const FSlateBrush* >& OutBrushes ) const override;
 
-	static const FName TypeName;
+	static SLATECORE_API const FName TypeName;
 	virtual const FName GetTypeName() const override { return TypeName; };
 
-	static const FSpinBoxStyle& GetDefault();
+	static SLATECORE_API const FSpinBoxStyle& GetDefault();
 
 	/** Brush used to draw the background of the spinbox */
 	UPROPERTY(EditAnywhere, BlueprintReadWrite, Category=Appearance)
@@ -1481,20 +1481,20 @@ struct SLATECORE_API FSpinBoxStyle : public FSlateWidgetStyle
  * Represents the appearance of an SSplitter
  */
 USTRUCT(BlueprintType)
-struct SLATECORE_API FSplitterStyle : public FSlateWidgetStyle
+struct FSplitterStyle : public FSlateWidgetStyle
 {
 	GENERATED_USTRUCT_BODY()
 
-	FSplitterStyle();
+	SLATECORE_API FSplitterStyle();
 
 	virtual ~FSplitterStyle() {}
 
-	virtual void GetResources( TArray< const FSlateBrush* >& OutBrushes ) const override;
+	SLATECORE_API virtual void GetResources( TArray< const FSlateBrush* >& OutBrushes ) const override;
 
-	static const FName TypeName;
+	static SLATECORE_API const FName TypeName;
 	virtual const FName GetTypeName() const override { return TypeName; };
 
-	static const FSplitterStyle& GetDefault();
+	static SLATECORE_API const FSplitterStyle& GetDefault();
 
 	/** Brush used to draw the handle in its normal state */
 	UPROPERTY(EditAnywhere, BlueprintReadWrite, Category=Appearance)
@@ -1512,20 +1512,20 @@ struct SLATECORE_API FSplitterStyle : public FSlateWidgetStyle
  */
 
 USTRUCT(BlueprintType)
-struct SLATECORE_API FTableViewStyle : public FSlateWidgetStyle
+struct FTableViewStyle : public FSlateWidgetStyle
 {
 	GENERATED_USTRUCT_BODY()
 
-	FTableViewStyle();
+	SLATECORE_API FTableViewStyle();
 
 	virtual ~FTableViewStyle() {}
 
-	virtual void GetResources( TArray< const FSlateBrush* >& OutBrushes ) const override;
+	SLATECORE_API virtual void GetResources( TArray< const FSlateBrush* >& OutBrushes ) const override;
 
-	static const FName TypeName;
+	static SLATECORE_API const FName TypeName;
 	virtual const FName GetTypeName() const override { return TypeName; };
 
-	static const FTableViewStyle& GetDefault();
+	static SLATECORE_API const FTableViewStyle& GetDefault();
 
 	/** Brush used when a selected row is active */
 	UPROPERTY(EditAnywhere, BlueprintReadWrite, Category=Appearance)
@@ -1547,20 +1547,20 @@ struct SLATECORE_API FTableViewStyle : public FSlateWidgetStyle
  * Represents the appearance of an STableRow
  */
 USTRUCT(BlueprintType)
-struct SLATECORE_API FTableRowStyle : public FSlateWidgetStyle
+struct FTableRowStyle : public FSlateWidgetStyle
 {
 	GENERATED_USTRUCT_BODY()
 
-	FTableRowStyle();
+	SLATECORE_API FTableRowStyle();
 
 	virtual ~FTableRowStyle() {}
 
-	virtual void GetResources( TArray< const FSlateBrush* >& OutBrushes ) const override;
+	SLATECORE_API virtual void GetResources( TArray< const FSlateBrush* >& OutBrushes ) const override;
 
-	static const FName TypeName;
+	static SLATECORE_API const FName TypeName;
 	virtual const FName GetTypeName() const override { return TypeName; };
 
-	static const FTableRowStyle& GetDefault();
+	static SLATECORE_API const FTableRowStyle& GetDefault();
 
 	/** Brush used as a selector when a row is focused */
 	UPROPERTY(EditAnywhere, BlueprintReadWrite, Category=Appearance)
@@ -1689,20 +1689,20 @@ struct SLATECORE_API FTableRowStyle : public FSlateWidgetStyle
  * Represents the appearance of an STableColumnHeader
  */
 USTRUCT()
-struct SLATECORE_API FTableColumnHeaderStyle : public FSlateWidgetStyle
+struct FTableColumnHeaderStyle : public FSlateWidgetStyle
 {
 	GENERATED_USTRUCT_BODY()
 
-	FTableColumnHeaderStyle();
+	SLATECORE_API FTableColumnHeaderStyle();
 
 	virtual ~FTableColumnHeaderStyle() {}
 
-	virtual void GetResources( TArray< const FSlateBrush* >& OutBrushes ) const override;
+	SLATECORE_API virtual void GetResources( TArray< const FSlateBrush* >& OutBrushes ) const override;
 
-	static const FName TypeName;
+	static SLATECORE_API const FName TypeName;
 	virtual const FName GetTypeName() const override { return TypeName; };
 
-	static const FTableColumnHeaderStyle& GetDefault();
+	static SLATECORE_API const FTableColumnHeaderStyle& GetDefault();
 
 	/** Image used when a column is primarily sorted in ascending order */
 	UPROPERTY(EditAnywhere, Category=Appearance)
@@ -1755,20 +1755,20 @@ struct SLATECORE_API FTableColumnHeaderStyle : public FSlateWidgetStyle
  * Represents the appearance of an SHeaderRow
  */
 USTRUCT()
-struct SLATECORE_API FHeaderRowStyle : public FSlateWidgetStyle
+struct FHeaderRowStyle : public FSlateWidgetStyle
 {
 	GENERATED_USTRUCT_BODY()
 
-	FHeaderRowStyle();
+	SLATECORE_API FHeaderRowStyle();
 
 	virtual ~FHeaderRowStyle() {}
 
-	virtual void GetResources( TArray< const FSlateBrush* >& OutBrushes ) const override;
+	SLATECORE_API virtual void GetResources( TArray< const FSlateBrush* >& OutBrushes ) const override;
 
-	static const FName TypeName;
+	static SLATECORE_API const FName TypeName;
 	virtual const FName GetTypeName() const override { return TypeName; };
 
-	static const FHeaderRowStyle& GetDefault();
+	static SLATECORE_API const FHeaderRowStyle& GetDefault();
 
 	/** Style of the normal header row columns */
 	UPROPERTY(EditAnywhere, Category=Appearance)
@@ -1817,20 +1817,20 @@ struct SLATECORE_API FHeaderRowStyle : public FSlateWidgetStyle
  * Represents the appearance of an SDockTab
  */
 USTRUCT()
-struct SLATECORE_API FDockTabStyle : public FSlateWidgetStyle
+struct FDockTabStyle : public FSlateWidgetStyle
 {
 	GENERATED_USTRUCT_BODY()
 
-	FDockTabStyle();
+	SLATECORE_API FDockTabStyle();
 
 	virtual ~FDockTabStyle() {}
 
-	virtual void GetResources( TArray< const FSlateBrush* >& OutBrushes ) const override;
+	SLATECORE_API virtual void GetResources( TArray< const FSlateBrush* >& OutBrushes ) const override;
 
-	static const FName TypeName;
+	static SLATECORE_API const FName TypeName;
 	virtual const FName GetTypeName() const override { return TypeName; };
 
-	static const FDockTabStyle& GetDefault();
+	static SLATECORE_API const FDockTabStyle& GetDefault();
 
 	/** Style used for the close button */
 	UPROPERTY(EditAnywhere, Category=Appearance)
@@ -1929,20 +1929,20 @@ struct SLATECORE_API FDockTabStyle : public FSlateWidgetStyle
  * Represents the appearance of an SScrollBox
  */
 USTRUCT(BlueprintType)
-struct SLATECORE_API FScrollBoxStyle : public FSlateWidgetStyle
+struct FScrollBoxStyle : public FSlateWidgetStyle
 {
 	GENERATED_USTRUCT_BODY()
 
-	FScrollBoxStyle();
+	SLATECORE_API FScrollBoxStyle();
 
 	virtual ~FScrollBoxStyle() {}
 
-	virtual void GetResources( TArray< const FSlateBrush* >& OutBrushes ) const override;
+	SLATECORE_API virtual void GetResources( TArray< const FSlateBrush* >& OutBrushes ) const override;
 
-	static const FName TypeName;
+	static SLATECORE_API const FName TypeName;
 	virtual const FName GetTypeName() const override { return TypeName; };
 
-	static const FScrollBoxStyle& GetDefault();
+	static SLATECORE_API const FScrollBoxStyle& GetDefault();
 
 	UPROPERTY(EditAnywhere, BlueprintReadWrite, Category=Appearance)
 	float BarThickness;
@@ -2012,20 +2012,20 @@ struct SLATECORE_API FScrollBoxStyle : public FSlateWidgetStyle
 * Represents the appearance of an FScrollBorderStyle
 */
 USTRUCT(BlueprintType)
-struct SLATECORE_API FScrollBorderStyle : public FSlateWidgetStyle
+struct FScrollBorderStyle : public FSlateWidgetStyle
 {
 	GENERATED_USTRUCT_BODY()
 
-	FScrollBorderStyle();
+	SLATECORE_API FScrollBorderStyle();
 
 	virtual ~FScrollBorderStyle() {}
 
-	virtual void GetResources( TArray< const FSlateBrush* >& OutBrushes ) const override;
+	SLATECORE_API virtual void GetResources( TArray< const FSlateBrush* >& OutBrushes ) const override;
 
-	static const FName TypeName;
+	static SLATECORE_API const FName TypeName;
 	virtual const FName GetTypeName() const override { return TypeName; };
 
-	static const FScrollBorderStyle& GetDefault();
+	static SLATECORE_API const FScrollBorderStyle& GetDefault();
 
 	/** Brush used to draw the top shadow of a scrollborder */
 	UPROPERTY(EditAnywhere, BlueprintReadWrite, Category = Appearance)
@@ -2043,20 +2043,20 @@ struct SLATECORE_API FScrollBorderStyle : public FSlateWidgetStyle
  * Represents the appearance of an SWindow
  */
 USTRUCT(BlueprintType)
-struct SLATECORE_API FWindowStyle : public FSlateWidgetStyle
+struct FWindowStyle : public FSlateWidgetStyle
 {
 	GENERATED_USTRUCT_BODY()
 
-	FWindowStyle();
+	SLATECORE_API FWindowStyle();
 
 	virtual ~FWindowStyle() {}
 
-	virtual void GetResources( TArray< const FSlateBrush* >& OutBrushes ) const override;
+	SLATECORE_API virtual void GetResources( TArray< const FSlateBrush* >& OutBrushes ) const override;
 
-	static const FName TypeName;
+	static SLATECORE_API const FName TypeName;
 	virtual const FName GetTypeName() const override { return TypeName; };
 
-	static const FWindowStyle& GetDefault();
+	static SLATECORE_API const FWindowStyle& GetDefault();
 
 	/** Style used to draw the window minimize button */
 	UPROPERTY(EditAnywhere, BlueprintReadWrite, Category=Appearance)
@@ -2146,7 +2146,7 @@ struct SLATECORE_API FWindowStyle : public FSlateWidgetStyle
 };
 
 
-class SLATECORE_API FInvalidatableBrushAttribute
+class FInvalidatableBrushAttribute
 {
 public:
 	FInvalidatableBrushAttribute() { }
@@ -2163,7 +2163,7 @@ public:
 
 	const FSlateBrush* Get() const { return Image.Get(); }
 	TAttribute< const FSlateBrush* > GetImage() const { return Image; }
-	void SetImage(SWidget& ThisWidget, const TAttribute< const FSlateBrush* >& InImage);
+	SLATECORE_API void SetImage(SWidget& ThisWidget, const TAttribute< const FSlateBrush* >& InImage);
 
 private:
 	/** The slate brush to draw for the image, or a bound delegate to a brush. */
diff --git a/Engine/Source/Runtime/SlateCore/Public/Styling/SlateWidgetStyle.h b/Engine/Source/Runtime/SlateCore/Public/Styling/SlateWidgetStyle.h
--- a/Engine/Source/Runtime/SlateCore/Public/Styling/SlateWidgetStyle.h
+++ b/Engine/Source/Runtime/SlateCore/Public/Styling/SlateWidgetStyle.h
@@ -12,14 +12,14 @@ struct FSlateBrush;
  * Base structure for widget styles.
  */
 USTRUCT()
-struct SLATECORE_API FSlateWidgetStyle
+struct FSlateWidgetStyle
 {
 	GENERATED_USTRUCT_BODY()
 
 public:
 
 	/** Default constructor. */
-	FSlateWidgetStyle( );
+	SLATECORE_API FSlateWidgetStyle( );
 
 	/** Virtual Destructor. */
 	virtual ~FSlateWidgetStyle() { }
diff --git a/Engine/Source/Runtime/SlateCore/Public/Styling/SlateWidgetStyleContainerBase.h b/Engine/Source/Runtime/SlateCore/Public/Styling/SlateWidgetStyleContainerBase.h
--- a/Engine/Source/Runtime/SlateCore/Public/Styling/SlateWidgetStyleContainerBase.h
+++ b/Engine/Source/Runtime/SlateCore/Public/Styling/SlateWidgetStyleContainerBase.h
@@ -16,12 +16,12 @@ SLATECORE_API DECLARE_LOG_CATEGORY_EXTERN(LogSlateStyle, Log, All);
 /**
  * Just a wrapper for the struct with real data in it.
  */
-UCLASS(hidecategories=Object)
-class SLATECORE_API USlateWidgetStyleContainerBase : public UObject, public ISlateWidgetStyleContainerInterface
+UCLASS(hidecategories=Object, MinimalAPI)
+class USlateWidgetStyleContainerBase : public UObject, public ISlateWidgetStyleContainerInterface
 {
 	GENERATED_BODY()
 
 public:
 
-	virtual const struct FSlateWidgetStyle* const GetStyle() const override;
+	SLATECORE_API virtual const struct FSlateWidgetStyle* const GetStyle() const override;
 };
diff --git a/Engine/Source/Runtime/SlateCore/Public/Styling/SlateWidgetStyleContainerInterface.h b/Engine/Source/Runtime/SlateCore/Public/Styling/SlateWidgetStyleContainerInterface.h
--- a/Engine/Source/Runtime/SlateCore/Public/Styling/SlateWidgetStyleContainerInterface.h
+++ b/Engine/Source/Runtime/SlateCore/Public/Styling/SlateWidgetStyleContainerInterface.h
@@ -11,8 +11,8 @@
 
 class UObject;
 
-UINTERFACE()
-class SLATECORE_API USlateWidgetStyleContainerInterface : public UInterface
+UINTERFACE(MinimalAPI)
+class USlateWidgetStyleContainerInterface : public UInterface
 {
 	GENERATED_UINTERFACE_BODY()
 };
diff --git a/Engine/Source/Runtime/SlateCore/Public/Styling/StarshipCoreStyle.h b/Engine/Source/Runtime/SlateCore/Public/Styling/StarshipCoreStyle.h
--- a/Engine/Source/Runtime/SlateCore/Public/Styling/StarshipCoreStyle.h
+++ b/Engine/Source/Runtime/SlateCore/Public/Styling/StarshipCoreStyle.h
@@ -11,11 +11,11 @@ struct FSlateDynamicImageBrush;
 /**
  * Core slate style
  */
-class SLATECORE_API FStarshipCoreStyle
+class FStarshipCoreStyle
 {
 public:
 
-	static TSharedRef<class ISlateStyle> Create();
+	static SLATECORE_API TSharedRef<class ISlateStyle> Create();
 
 	/** 
 	* @return the Application Style 
@@ -42,24 +42,24 @@ public:
 	}
 
 	/** Get the default font for Slate */
-	static TSharedRef<const FCompositeFont> GetDefaultFont();
+	static SLATECORE_API TSharedRef<const FCompositeFont> GetDefaultFont();
 
 	/** Get a font style using the default for for Slate */
-	static FSlateFontInfo GetDefaultFontStyle(const FName InTypefaceFontName, const int32 InSize, const FFontOutlineSettings& InOutlineSettings = FFontOutlineSettings());
+	static SLATECORE_API FSlateFontInfo GetDefaultFontStyle(const FName InTypefaceFontName, const int32 InSize, const FFontOutlineSettings& InOutlineSettings = FFontOutlineSettings());
 
-	static void ResetToDefault( );
+	static SLATECORE_API void ResetToDefault( );
 
 	/** Used to override the default selection colors */
-	static void SetSelectorColor( const FLinearColor& NewColor );
-	static void SetSelectionColor( const FLinearColor& NewColor );
-	static void SetInactiveSelectionColor( const FLinearColor& NewColor );
-	static void SetPressedSelectionColor( const FLinearColor& NewColor );
-	static void SetFocusBrush(FSlateBrush* NewBrush);
+	static SLATECORE_API void SetSelectorColor( const FLinearColor& NewColor );
+	static SLATECORE_API void SetSelectionColor( const FLinearColor& NewColor );
+	static SLATECORE_API void SetInactiveSelectionColor( const FLinearColor& NewColor );
+	static SLATECORE_API void SetPressedSelectionColor( const FLinearColor& NewColor );
+	static SLATECORE_API void SetFocusBrush(FSlateBrush* NewBrush);
 
 	// todo: jdale - These are only here because of UTouchInterface::Activate and the fact that GetDynamicImageBrush is non-const
-	static const TSharedPtr<FSlateDynamicImageBrush> GetDynamicImageBrush( FName BrushTemplate, FName TextureName, const ANSICHAR* Specifier = nullptr );
-	static const TSharedPtr<FSlateDynamicImageBrush> GetDynamicImageBrush( FName BrushTemplate, const ANSICHAR* Specifier, class UTexture2D* TextureResource, FName TextureName );
-	static const TSharedPtr<FSlateDynamicImageBrush> GetDynamicImageBrush( FName BrushTemplate, class UTexture2D* TextureResource, FName TextureName );
+	static SLATECORE_API const TSharedPtr<FSlateDynamicImageBrush> GetDynamicImageBrush( FName BrushTemplate, FName TextureName, const ANSICHAR* Specifier = nullptr );
+	static SLATECORE_API const TSharedPtr<FSlateDynamicImageBrush> GetDynamicImageBrush( FName BrushTemplate, const ANSICHAR* Specifier, class UTexture2D* TextureResource, FName TextureName );
+	static SLATECORE_API const TSharedPtr<FSlateDynamicImageBrush> GetDynamicImageBrush( FName BrushTemplate, class UTexture2D* TextureResource, FName TextureName );
 
 	static const int32 RegularTextSize = 10;
 	static const int32 SmallTextSize = 8;
@@ -68,26 +68,26 @@ public:
 private:
 	class FStyle;
 
-	static void SetStyle(const TSharedRef<class ISlateStyle>& NewStyle);
+	static SLATECORE_API void SetStyle(const TSharedRef<class ISlateStyle>& NewStyle);
 	
-	static void SetupColors(TSharedRef<FStyle>& Style);
-	static void SetupTextStyles(TSharedRef<FStyle>& Style);
-	static void SetupButtonStyles(TSharedRef<FStyle>& Style);
-	static void SetupComboButtonStyles(TSharedRef<FStyle>& Style);
-	static void SetupCheckboxStyles(TSharedRef<FStyle>& Style);
-	static void SetupDockingStyles(TSharedRef<FStyle>& Style);
-	static void SetupColorPickerStyles(TSharedRef<FStyle>& Style);
-	static void SetupTableViewStyles(TSharedRef<FStyle>& Style);
-	static void SetupMultiboxStyles(TSharedRef<FStyle>& Style);
+	static SLATECORE_API void SetupColors(TSharedRef<FStyle>& Style);
+	static SLATECORE_API void SetupTextStyles(TSharedRef<FStyle>& Style);
+	static SLATECORE_API void SetupButtonStyles(TSharedRef<FStyle>& Style);
+	static SLATECORE_API void SetupComboButtonStyles(TSharedRef<FStyle>& Style);
+	static SLATECORE_API void SetupCheckboxStyles(TSharedRef<FStyle>& Style);
+	static SLATECORE_API void SetupDockingStyles(TSharedRef<FStyle>& Style);
+	static SLATECORE_API void SetupColorPickerStyles(TSharedRef<FStyle>& Style);
+	static SLATECORE_API void SetupTableViewStyles(TSharedRef<FStyle>& Style);
+	static SLATECORE_API void SetupMultiboxStyles(TSharedRef<FStyle>& Style);
 
 private:
 
 	/** Singleton instances of this style. */
-	static TSharedPtr< class ISlateStyle > Instance;
+	static SLATECORE_API TSharedPtr< class ISlateStyle > Instance;
 
 };
 
-struct SLATECORE_API FStyleFonts
+struct FStyleFonts
 {
 
   public:
@@ -114,7 +114,7 @@ struct SLATECORE_API FStyleFonts
 	const FSlateFontInfo HeadingExtraSmall;
 
   private: 
-  	FStyleFonts();
-  	static TUniquePtr<struct FStyleFonts> Instance;
+  	SLATECORE_API FStyleFonts();
+  	static SLATECORE_API TUniquePtr<struct FStyleFonts> Instance;
 };
 
diff --git a/Engine/Source/Runtime/SlateCore/Public/Styling/StyleColors.h b/Engine/Source/Runtime/SlateCore/Public/Styling/StyleColors.h
--- a/Engine/Source/Runtime/SlateCore/Public/Styling/StyleColors.h
+++ b/Engine/Source/Runtime/SlateCore/Public/Styling/StyleColors.h
@@ -141,8 +141,8 @@ struct FStyleTheme
 	}
 };
 
-UCLASS(Config=EditorSettings)
-class SLATECORE_API USlateThemeManager : public UObject 
+UCLASS(Config=EditorSettings, MinimalAPI)
+class USlateThemeManager : public UObject 
 {
 	GENERATED_BODY()
 public:
@@ -161,12 +161,12 @@ public:
 		return CurrentThemeId; 
 	}
 
-	USlateThemeManager();
+	SLATECORE_API USlateThemeManager();
 
 	/**
 	 * Initializes default colors
 	 */
-	void InitalizeDefaults();
+	SLATECORE_API void InitalizeDefaults();
 
 	/**
 	 * Sets a default color to be used as a fallback if no theme is loaded
@@ -197,52 +197,52 @@ public:
 	/**
 	 * Load all known themes from engine, project, and user directories
 	 */
-	void LoadThemes();
+	SLATECORE_API void LoadThemes();
 
 	/**
 	 * Saves the current theme
 	 * 
 	 * @param The filename to save the current theme as
 	 */
-	void SaveCurrentThemeAs(const FString& Filename);
+	SLATECORE_API void SaveCurrentThemeAs(const FString& Filename);
 
 	/**
 	 * Applies a theme as the active theme
 	 */
-	void ApplyTheme(FGuid ThemeId);
+	SLATECORE_API void ApplyTheme(FGuid ThemeId);
 
 	/**
 	 * Applies the default dark theme as the active theme
 	*/
-	void ApplyDefaultTheme();
+	SLATECORE_API void ApplyDefaultTheme();
 
 	/**
 	* Returns true if the active theme is an engine-specific theme
 	*/
-	bool IsEngineTheme() const;
+	SLATECORE_API bool IsEngineTheme() const;
 
 	/**
 	* Returns true if the active theme is a project-specific theme
 	*/
-	bool IsProjectTheme() const;
+	SLATECORE_API bool IsProjectTheme() const;
 
 	/**
 	 * Removes a theme. 
 	 * Note: The active theme cannot be removed and there must always be an active theme.  Apply a new theme first before removing the current theme.
 	 */
-	void RemoveTheme(FGuid ThemeId);
+	SLATECORE_API void RemoveTheme(FGuid ThemeId);
 
 	/**
 	 * Duplicates the active theme
 	 * 
 	 * @return the id of the new theme
 	 */
-	FGuid DuplicateActiveTheme();
+	SLATECORE_API FGuid DuplicateActiveTheme();
 
 	/**
 	 * Sets the display name for the current theme
 	 */
-	void SetCurrentThemeDisplayName(FText NewDisplayName);
+	SLATECORE_API void SetCurrentThemeDisplayName(FText NewDisplayName);
 
 	/**
 	 * @return the current theme
@@ -257,32 +257,32 @@ public:
 	/**
 	 * Resets an active color to the default color for the curerent theme
 	 */
-	void ResetActiveColorToDefault(EStyleColor Color);
+	SLATECORE_API void ResetActiveColorToDefault(EStyleColor Color);
 
 	/**
 	 * Validate that there is an active loaded theme
 	 */
-	void ValidateActiveTheme();
+	SLATECORE_API void ValidateActiveTheme();
 
 	/**
 	 * @return the engine theme dir.  Engine themes are project agnostic
 	 */
-	FString GetEngineThemeDir() const;
+	SLATECORE_API FString GetEngineThemeDir() const;
 
 	/**
 	 * @return the project theme dir. Project themes can override engine themes
 	 */
-	FString GetProjectThemeDir() const;
+	SLATECORE_API FString GetProjectThemeDir() const;
 
 	/**
 	 * @return the user theme dir. Themes in this dir are per-user and override engine and project themes
 	 */
-	FString GetUserThemeDir() const;
+	SLATECORE_API FString GetUserThemeDir() const;
 		
 	/**
 	 * @return true if the ThemeID already exists in the theme dropdown. 
 	 */
-	bool DoesThemeExist(const FGuid& ThemeID) const;
+	SLATECORE_API bool DoesThemeExist(const FGuid& ThemeID) const;
 
 	DECLARE_EVENT_OneParam(USlateThemeManager, FThemeChangedEvent, FGuid);
 	/**
@@ -295,10 +295,10 @@ public:
 private:
 
 	FStyleTheme& GetMutableCurrentTheme() { return *Themes.FindByKey(CurrentThemeId); }
-	void LoadThemesFromDirectory(const FString& Directory);
-	bool ReadTheme(const FString& ThemeData, FStyleTheme& OutTheme);
-	void EnsureValidCurrentTheme();
-	void LoadThemeColors(FStyleTheme& Theme);
+	SLATECORE_API void LoadThemesFromDirectory(const FString& Directory);
+	SLATECORE_API bool ReadTheme(const FString& ThemeData, FStyleTheme& OutTheme);
+	SLATECORE_API void EnsureValidCurrentTheme();
+	SLATECORE_API void LoadThemeColors(FStyleTheme& Theme);
 
 private:
 	FStyleTheme DefaultDarkTheme;
@@ -330,58 +330,58 @@ private:
  * Common/themeable colors used by all styles
  * Please avoid adding new generic colors to this list without discussion first
  */
-struct SLATECORE_API FStyleColors
+struct FStyleColors
 {
-	static const FSlateColor Transparent;
-	static const FSlateColor Black;
-	static const FSlateColor Title;
-	static const FSlateColor WindowBorder;
-	static const FSlateColor Foldout;
-	static const FSlateColor Input;
-	static const FSlateColor InputOutline;
-	static const FSlateColor Recessed;
-	static const FSlateColor Background;
-	static const FSlateColor Panel;
-	static const FSlateColor Header;
-	static const FSlateColor Dropdown;
-	static const FSlateColor DropdownOutline;
-	static const FSlateColor Hover;
-	static const FSlateColor Hover2;
-	static const FSlateColor White;
-	static const FSlateColor White25;
-	static const FSlateColor Highlight;
-
-	static const FSlateColor Primary;
-	static const FSlateColor PrimaryHover;
-	static const FSlateColor PrimaryPress;
-	static const FSlateColor Secondary;
-
-	static const FSlateColor Foreground;
-	static const FSlateColor ForegroundHover;
-	static const FSlateColor ForegroundInverted;
-	static const FSlateColor ForegroundHeader;
-
-	static const FSlateColor Select;
-	static const FSlateColor SelectInactive;
-	static const FSlateColor SelectParent;
-	static const FSlateColor SelectHover;
-
-	static const FSlateColor Notifications;
-
-	static const FSlateColor AccentBlue;
-	static const FSlateColor AccentPurple;
-	static const FSlateColor AccentPink;
-	static const FSlateColor AccentRed;
-	static const FSlateColor AccentOrange;
-	static const FSlateColor AccentYellow;
-	static const FSlateColor AccentGreen;
-	static const FSlateColor AccentBrown;
-	static const FSlateColor AccentBlack;
-	static const FSlateColor AccentGray;
-	static const FSlateColor AccentWhite;
-	static const FSlateColor AccentFolder;
-
-	static const FSlateColor Warning;
-	static const FSlateColor Error;
-	static const FSlateColor Success;
+	static SLATECORE_API const FSlateColor Transparent;
+	static SLATECORE_API const FSlateColor Black;
+	static SLATECORE_API const FSlateColor Title;
+	static SLATECORE_API const FSlateColor WindowBorder;
+	static SLATECORE_API const FSlateColor Foldout;
+	static SLATECORE_API const FSlateColor Input;
+	static SLATECORE_API const FSlateColor InputOutline;
+	static SLATECORE_API const FSlateColor Recessed;
+	static SLATECORE_API const FSlateColor Background;
+	static SLATECORE_API const FSlateColor Panel;
+	static SLATECORE_API const FSlateColor Header;
+	static SLATECORE_API const FSlateColor Dropdown;
+	static SLATECORE_API const FSlateColor DropdownOutline;
+	static SLATECORE_API const FSlateColor Hover;
+	static SLATECORE_API const FSlateColor Hover2;
+	static SLATECORE_API const FSlateColor White;
+	static SLATECORE_API const FSlateColor White25;
+	static SLATECORE_API const FSlateColor Highlight;
+
+	static SLATECORE_API const FSlateColor Primary;
+	static SLATECORE_API const FSlateColor PrimaryHover;
+	static SLATECORE_API const FSlateColor PrimaryPress;
+	static SLATECORE_API const FSlateColor Secondary;
+
+	static SLATECORE_API const FSlateColor Foreground;
+	static SLATECORE_API const FSlateColor ForegroundHover;
+	static SLATECORE_API const FSlateColor ForegroundInverted;
+	static SLATECORE_API const FSlateColor ForegroundHeader;
+
+	static SLATECORE_API const FSlateColor Select;
+	static SLATECORE_API const FSlateColor SelectInactive;
+	static SLATECORE_API const FSlateColor SelectParent;
+	static SLATECORE_API const FSlateColor SelectHover;
+
+	static SLATECORE_API const FSlateColor Notifications;
+
+	static SLATECORE_API const FSlateColor AccentBlue;
+	static SLATECORE_API const FSlateColor AccentPurple;
+	static SLATECORE_API const FSlateColor AccentPink;
+	static SLATECORE_API const FSlateColor AccentRed;
+	static SLATECORE_API const FSlateColor AccentOrange;
+	static SLATECORE_API const FSlateColor AccentYellow;
+	static SLATECORE_API const FSlateColor AccentGreen;
+	static SLATECORE_API const FSlateColor AccentBrown;
+	static SLATECORE_API const FSlateColor AccentBlack;
+	static SLATECORE_API const FSlateColor AccentGray;
+	static SLATECORE_API const FSlateColor AccentWhite;
+	static SLATECORE_API const FSlateColor AccentFolder;
+
+	static SLATECORE_API const FSlateColor Warning;
+	static SLATECORE_API const FSlateColor Error;
+	static SLATECORE_API const FSlateColor Success;
 };
diff --git a/Engine/Source/Runtime/SlateCore/Public/Styling/StyleDefaults.h b/Engine/Source/Runtime/SlateCore/Public/Styling/StyleDefaults.h
--- a/Engine/Source/Runtime/SlateCore/Public/Styling/StyleDefaults.h
+++ b/Engine/Source/Runtime/SlateCore/Public/Styling/StyleDefaults.h
@@ -8,7 +8,7 @@
 #include "Misc/Paths.h"
 #include "Brushes/SlateNoResource.h"
 
-class SLATECORE_API FStyleDefaults
+class FStyleDefaults
 {
 public:
 
@@ -60,13 +60,13 @@ public:
 	/**
 	 * @return - Returns no brush.
 	 */
-	static const FSlateBrush* GetNoBrush();
+	static SLATECORE_API const FSlateBrush* GetNoBrush();
 
 	/**
 	 * Get default font.
 	 * @return - The default font.
 	 */
-	static const FSlateFontInfo GetFontInfo(uint16 Size = 10);
+	static SLATECORE_API const FSlateFontInfo GetFontInfo(uint16 Size = 10);
 
 	/**
 	 * Return the static default sound value.
@@ -79,10 +79,10 @@ public:
 
 private:
 
-	static float DefaultFloat;
-	static FVector2D DefaultFVector2D;
-	static FLinearColor DefaultColor;
-	static FMargin DefaultMargin;
-	static FSlateSound DefaultSound;
-	static FSlateColor DefaultSlateColor;
+	static SLATECORE_API float DefaultFloat;
+	static SLATECORE_API FVector2D DefaultFVector2D;
+	static SLATECORE_API FLinearColor DefaultColor;
+	static SLATECORE_API FMargin DefaultMargin;
+	static SLATECORE_API FSlateSound DefaultSound;
+	static SLATECORE_API FSlateColor DefaultSlateColor;
 };
diff --git a/Engine/Source/Runtime/SlateCore/Public/Styling/ToolBarStyle.h b/Engine/Source/Runtime/SlateCore/Public/Styling/ToolBarStyle.h
--- a/Engine/Source/Runtime/SlateCore/Public/Styling/ToolBarStyle.h
+++ b/Engine/Source/Runtime/SlateCore/Public/Styling/ToolBarStyle.h
@@ -12,20 +12,20 @@
  * Represents the appearance of a toolbar 
  */
 USTRUCT(BlueprintType)
-struct SLATECORE_API FToolBarStyle : public FSlateWidgetStyle
+struct FToolBarStyle : public FSlateWidgetStyle
 {
 	GENERATED_BODY()
 
-	FToolBarStyle();
+	SLATECORE_API FToolBarStyle();
 
 	virtual ~FToolBarStyle() {}
 
-	virtual void GetResources(TArray<const FSlateBrush*>& OutBrushes) const override;
+	SLATECORE_API virtual void GetResources(TArray<const FSlateBrush*>& OutBrushes) const override;
 
-	static const FName TypeName;
+	static SLATECORE_API const FName TypeName;
 	virtual const FName GetTypeName() const override { return TypeName; };
 
-	static const FToolBarStyle& GetDefault();
+	static SLATECORE_API const FToolBarStyle& GetDefault();
 
 	/** The brush used for the background of the toolbar */
 	UPROPERTY(EditAnywhere, BlueprintReadWrite, Category = Appearance)
diff --git a/Engine/Source/Runtime/SlateCore/Public/Styling/UMGCoreStyle.h b/Engine/Source/Runtime/SlateCore/Public/Styling/UMGCoreStyle.h
--- a/Engine/Source/Runtime/SlateCore/Public/Styling/UMGCoreStyle.h
+++ b/Engine/Source/Runtime/SlateCore/Public/Styling/UMGCoreStyle.h
@@ -9,11 +9,11 @@
 /**
  * Core slate style
  */
-class SLATECORE_API FUMGCoreStyle
+class FUMGCoreStyle
 {
 public:
 
-	static TSharedRef<class ISlateStyle> Create();
+	static SLATECORE_API TSharedRef<class ISlateStyle> Create();
 
 	static const ISlateStyle& Get()
 	{
@@ -21,19 +21,19 @@ public:
 	}
 
 	/** Get the default font for Slate */
-	static TSharedRef<const FCompositeFont> GetDefaultFont();
+	static SLATECORE_API TSharedRef<const FCompositeFont> GetDefaultFont();
 
 	/** Get a font style using the default for for Slate */
-	static FSlateFontInfo GetDefaultFontStyle(const FName InTypefaceFontName, const int32 InSize, const FFontOutlineSettings& InOutlineSettings = FFontOutlineSettings());
+	static SLATECORE_API FSlateFontInfo GetDefaultFontStyle(const FName InTypefaceFontName, const int32 InSize, const FFontOutlineSettings& InOutlineSettings = FFontOutlineSettings());
 
-	static void ResetToDefault( );
+	static SLATECORE_API void ResetToDefault( );
 
 	/** Used to override the default selection colors */
-	static void SetSelectorColor( const FLinearColor& NewColor );
-	static void SetSelectionColor( const FLinearColor& NewColor );
-	static void SetInactiveSelectionColor( const FLinearColor& NewColor );
-	static void SetPressedSelectionColor( const FLinearColor& NewColor );
-	static void SetFocusBrush(FSlateBrush* NewBrush);
+	static SLATECORE_API void SetSelectorColor( const FLinearColor& NewColor );
+	static SLATECORE_API void SetSelectionColor( const FLinearColor& NewColor );
+	static SLATECORE_API void SetInactiveSelectionColor( const FLinearColor& NewColor );
+	static SLATECORE_API void SetPressedSelectionColor( const FLinearColor& NewColor );
+	static SLATECORE_API void SetFocusBrush(FSlateBrush* NewBrush);
 
 	static const int32 RegularTextSize = 10;
 	static const int32 SmallTextSize = 8;
@@ -41,10 +41,10 @@ public:
 	static bool IsInitialized() { return Instance.IsValid(); }
 private:
 
-	static void SetStyle(const TSharedRef<class ISlateStyle>& NewStyle);
+	static SLATECORE_API void SetStyle(const TSharedRef<class ISlateStyle>& NewStyle);
 
 private:
 
 	/** Singleton instances of this style. */
-	static TSharedPtr< class ISlateStyle > Instance;
+	static SLATECORE_API TSharedPtr< class ISlateStyle > Instance;
 };
diff --git a/Engine/Source/Runtime/SlateCore/Public/Styling/WidgetStyle.h b/Engine/Source/Runtime/SlateCore/Public/Styling/WidgetStyle.h
--- a/Engine/Source/Runtime/SlateCore/Public/Styling/WidgetStyle.h
+++ b/Engine/Source/Runtime/SlateCore/Public/Styling/WidgetStyle.h
@@ -11,7 +11,7 @@ template< typename ObjectType > class TAttribute;
 /**
  * Contains info about those aspects of widget appearance that should be propagated hierarchically.
  */
-class SLATECORE_API FWidgetStyle
+class FWidgetStyle
 {
 public:
 	
@@ -74,7 +74,7 @@ public:
 	 * @return This instance (for method chaining).
 	 * @see GetForegroundColor
 	 */
-	FWidgetStyle& SetForegroundColor( const TAttribute<FSlateColor>& InForeground );
+	SLATECORE_API FWidgetStyle& SetForegroundColor( const TAttribute<FSlateColor>& InForeground );
 
 public:
 
@@ -119,5 +119,5 @@ private:
 
 private:
 
-	static const float SubdueAmount;
+	static SLATECORE_API const float SubdueAmount;
 };
diff --git a/Engine/Source/Runtime/SlateCore/Public/Textures/SlateIcon.h b/Engine/Source/Runtime/SlateCore/Public/Textures/SlateIcon.h
--- a/Engine/Source/Runtime/SlateCore/Public/Textures/SlateIcon.h
+++ b/Engine/Source/Runtime/SlateCore/Public/Textures/SlateIcon.h
@@ -9,12 +9,12 @@ struct FSlateBrush;
 /**
  * Struct used to represent an icon in Slate
  */
-struct SLATECORE_API FSlateIcon
+struct FSlateIcon
 {
 	/**
 	 * Default constructor (empty icon).
 	 */
-	FSlateIcon();
+	SLATECORE_API FSlateIcon();
 
 	/**
 	 * Creates and initializes a new icon from a style set and style name
@@ -24,7 +24,7 @@ struct SLATECORE_API FSlateIcon
 	 * @param InSmallStyleName The name of the style for the small icon
 	 * @param InStatusOverlayStyleName The name of the style for a status icon to overlay on the base image
 	 */
-	FSlateIcon(const FName InStyleSetName, const FName InStyleName, const FName InSmallStyleName = NAME_None, const FName InStatusOverlayStyleName = NAME_None);
+	SLATECORE_API FSlateIcon(const FName InStyleSetName, const FName InStyleName, const FName InSmallStyleName = NAME_None, const FName InStatusOverlayStyleName = NAME_None);
 
 public:
 	
@@ -56,14 +56,14 @@ public:
 	 * @return Icon brush, or FStyleDefaults::GetNoBrush() if the icon wasn't found.
 	 * @see GetSmallIcon
 	 */
-	const FSlateBrush* GetIcon( ) const;
+	SLATECORE_API const FSlateBrush* GetIcon( ) const;
 
 	/**
 	 * Optionally gets the resolved icon, returning nullptr if it's not defined
 	 *
 	 * @return Icon brush, or nullptr if the icon wasn't found.
 	 */
-	const FSlateBrush* GetOptionalIcon( ) const;
+	SLATECORE_API const FSlateBrush* GetOptionalIcon( ) const;
 
 	/**
 	 * Gets the resolved small icon.
@@ -71,16 +71,16 @@ public:
 	 * @return Icon brush, or FStyleDefaults::GetNoBrush() if the icon wasn't found.
 	 * @see GetIcon
 	 */
-	const FSlateBrush* GetSmallIcon( ) const;
+	SLATECORE_API const FSlateBrush* GetSmallIcon( ) const;
 
 	/**
 	 * Optionally gets the resolved small icon, returning nullptr if it's not defined
 	 *
 	 * @return Icon brush, or nullptr if the icon wasn't found.
 	 */
-	const FSlateBrush* GetOptionalSmallIcon( ) const;
+	SLATECORE_API const FSlateBrush* GetOptionalSmallIcon( ) const;
 
-	const FSlateBrush* GetOverlayIcon() const;
+	SLATECORE_API const FSlateBrush* GetOverlayIcon() const;
 
 	/**
 	 * Gets the name of the style for the icon.
@@ -110,7 +110,7 @@ public:
 	 * @return Style set, or nullptr if the style set wasn't found.
 	 * @see GetSmallStyleName, GetStyleName, GetStyleSetName
 	 */
-	const class ISlateStyle* GetStyleSet() const;
+	SLATECORE_API const class ISlateStyle* GetStyleSet() const;
 
 	/**
 	 * Gets the name of the style set the icon can be found in.
diff --git a/Engine/Source/Runtime/SlateCore/Public/Textures/SlateShaderResource.h b/Engine/Source/Runtime/SlateCore/Public/Textures/SlateShaderResource.h
--- a/Engine/Source/Runtime/SlateCore/Public/Textures/SlateShaderResource.h
+++ b/Engine/Source/Runtime/SlateCore/Public/Textures/SlateShaderResource.h
@@ -39,7 +39,7 @@ namespace ESlateShaderResource
 /** 
  * Base class for all platform independent texture types
  */
-class SLATECORE_API FSlateShaderResource
+class FSlateShaderResource
 {
 public:
 
@@ -109,7 +109,7 @@ public:
  * May point to a full resource or point or to a texture resource in an atlas
  * Note: This class does not free any resources.  Resources should be owned and freed elsewhere
  */
-class SLATECORE_API FSlateShaderResourceProxy
+class FSlateShaderResourceProxy
 {
 public:
 
diff --git a/Engine/Source/Runtime/SlateCore/Public/Textures/SlateTextureData.h b/Engine/Source/Runtime/SlateCore/Public/Textures/SlateTextureData.h
--- a/Engine/Source/Runtime/SlateCore/Public/Textures/SlateTextureData.h
+++ b/Engine/Source/Runtime/SlateCore/Public/Textures/SlateTextureData.h
@@ -11,7 +11,7 @@ DECLARE_MEMORY_STAT_EXTERN(TEXT("Texture Data Memory (CPU)"), STAT_SlateTextureD
 /**
  * Holds texture data for upload to a rendering resource
  */
-struct SLATECORE_API FSlateTextureData
+struct FSlateTextureData
 {
 	FSlateTextureData( uint32 InWidth = 0, uint32 InHeight = 0, uint32 InBytesPerPixel = 0, const TArray<uint8>& InBytes = TArray<uint8>() )
 		: Bytes(InBytes)
diff --git a/Engine/Source/Runtime/SlateCore/Public/Textures/TextureAtlas.h b/Engine/Source/Runtime/SlateCore/Public/Textures/TextureAtlas.h
--- a/Engine/Source/Runtime/SlateCore/Public/Textures/TextureAtlas.h
+++ b/Engine/Source/Runtime/SlateCore/Public/Textures/TextureAtlas.h
@@ -77,7 +77,7 @@ struct FAtlasedTextureSlot : public TIntrusiveLinkedList<FAtlasedTextureSlot>
 /**
  * Base class texture atlases in Slate
  */
-class SLATECORE_API FSlateTextureAtlas
+class FSlateTextureAtlas
 {
 public:
 	FSlateTextureAtlas( uint32 InWidth, uint32 InHeight, uint32 InBytesPerPixel, ESlateTextureAtlasPaddingStyle InPaddingStyle, bool bInUpdatesAfterInitialization )
@@ -97,12 +97,12 @@ public:
 		InitAtlasData();
 	}
 
-	virtual ~FSlateTextureAtlas();
+	SLATECORE_API virtual ~FSlateTextureAtlas();
 
 	/**
 	 * Clears atlas cpu data.  It does not clear rendering data
 	 */
-	void EmptyAtlasData();
+	SLATECORE_API void EmptyAtlasData();
 
 	/**
 	 * Adds a texture to the atlas
@@ -111,7 +111,7 @@ public:
 	 * @param TextureHeight	Height of the texture
 	 * @param Data			Raw texture data
 	 */
-	const FAtlasedTextureSlot* AddTexture(uint32 TextureWidth, uint32 TextureHeight, const TArray<uint8>& Data);
+	SLATECORE_API const FAtlasedTextureSlot* AddTexture(uint32 TextureWidth, uint32 TextureHeight, const TArray<uint8>& Data);
 
 	/** @return the width of the atlas */
 	uint32 GetWidth() const { return AtlasWidth; }
@@ -119,7 +119,7 @@ public:
 	uint32 GetHeight() const { return AtlasHeight; }
 
 	/** Marks the texture as dirty and needing its rendering resources updated */
-	void MarkTextureDirty();
+	SLATECORE_API void MarkTextureDirty();
 	
 	/**
 	 * Updates the texture used for rendering if needed
@@ -134,7 +134,7 @@ public:
 	virtual FSlateShaderResource* GetAtlasTexture() const = 0;
 
 #if WITH_ATLAS_DEBUGGING
-	const FAtlasedTextureSlot* GetSlotAtPosition(FIntPoint InPosition) const;
+	SLATECORE_API const FAtlasedTextureSlot* GetSlotAtPosition(FIntPoint InPosition) const;
 #endif
 protected:
 	/**
@@ -143,12 +143,12 @@ protected:
 	 * @param Width The width of the texture we are adding
 	 * @param Height The height of the texture we are adding
 	 */
-	const FAtlasedTextureSlot* FindSlotForTexture( uint32 InWidth, uint32 InHeight );
+	SLATECORE_API const FAtlasedTextureSlot* FindSlotForTexture( uint32 InWidth, uint32 InHeight );
 
 	/**
 	 * Creates enough space for a single texture the width and height of the atlas
 	 */
-	void InitAtlasData();
+	SLATECORE_API void InitAtlasData();
 
 	struct FCopyRowData
 	{
@@ -174,7 +174,7 @@ protected:
 	 *
 	 * @param CopyRowData	Information for how to copy a row
 	 */
-	void CopyRow( const FCopyRowData& CopyRowData );
+	SLATECORE_API void CopyRow( const FCopyRowData& CopyRowData );
 
 	/**
 	 * Zeros out a row in the dest texture (used with PaddingStyle == PadWithZero).
@@ -182,7 +182,7 @@ protected:
 	 *
 	 * @param CopyRowData	Information for how to copy a row
 	 */
-	void ZeroRow( const FCopyRowData& CopyRowData );
+	SLATECORE_API void ZeroRow( const FCopyRowData& CopyRowData );
 
 	/** 
 	 * Copies texture data into the atlas at a given slot
@@ -190,7 +190,7 @@ protected:
 	 * @param SlotToCopyTo	The occupied slot in the atlas where texture data should be copied to
 	 * @param Data			The data to copy into the atlas
 	 */
-	void CopyDataIntoSlot( const FAtlasedTextureSlot* SlotToCopyTo, const TArray<uint8>& Data );
+	SLATECORE_API void CopyDataIntoSlot( const FAtlasedTextureSlot* SlotToCopyTo, const TArray<uint8>& Data );
 
 private:
 	/** Returns the amount of padding needed for the current padding style */
diff --git a/Engine/Source/Runtime/SlateCore/Public/Types/Attributes/SlateAttributeDefinition.inl b/Engine/Source/Runtime/SlateCore/Public/Types/Attributes/SlateAttributeDefinition.inl
--- a/Engine/Source/Runtime/SlateCore/Public/Types/Attributes/SlateAttributeDefinition.inl
+++ b/Engine/Source/Runtime/SlateCore/Public/Types/Attributes/SlateAttributeDefinition.inl
@@ -39,7 +39,7 @@ namespace SlateAttributePrivate
 
 
 	/** Interface for structure that can be used to contain a SlateAttribute instead of a SWidget. */
-	class SLATECORE_API ISlateAttributeContainer
+	class ISlateAttributeContainer
 	{
 	public:
 		virtual SWidget& GetContainerWidget() const = 0;
@@ -47,8 +47,8 @@ namespace SlateAttributePrivate
 		virtual uint32 GetContainerSortOrder() const = 0;
 
 	protected:
-		void RemoveContainerWidget(SWidget& Widget);
-		void UpdateContainerSortOrder(SWidget& Widget);
+		SLATECORE_API void RemoveContainerWidget(SWidget& Widget);
+		SLATECORE_API void UpdateContainerSortOrder(SWidget& Widget);
 	};
 
 
@@ -80,32 +80,32 @@ namespace SlateAttributePrivate
 
 
 	/** */
-	struct SLATECORE_API FSlateAttributeImpl : public FSlateAttributeBase
+	struct FSlateAttributeImpl : public FSlateAttributeBase
 	{
 	protected:
-		void ProtectedUnregisterAttribute(SWidget& Widget, ESlateAttributeType AttributeType) const;
-		void ProtectedRegisterAttribute(SWidget& Widget, ESlateAttributeType AttributeType, TUniquePtr<ISlateAttributeGetter>&& Wrapper);
-		void ProtectedInvalidateWidget(SWidget& Widget, ESlateAttributeType AttributeType, EInvalidateWidgetReason InvalidationReason) const;
-		bool ProtectedIsBound(const SWidget& Widget, ESlateAttributeType AttributeType) const;
-		ISlateAttributeGetter* ProtectedFindGetter(const SWidget& Widget, ESlateAttributeType AttributeType) const;
-		FDelegateHandle ProtectedFindGetterHandle(const SWidget& Widget, ESlateAttributeType AttributeType) const;
-		void ProtectedUpdateNow(SWidget& Widget, ESlateAttributeType AttributeType);
+		SLATECORE_API void ProtectedUnregisterAttribute(SWidget& Widget, ESlateAttributeType AttributeType) const;
+		SLATECORE_API void ProtectedRegisterAttribute(SWidget& Widget, ESlateAttributeType AttributeType, TUniquePtr<ISlateAttributeGetter>&& Wrapper);
+		SLATECORE_API void ProtectedInvalidateWidget(SWidget& Widget, ESlateAttributeType AttributeType, EInvalidateWidgetReason InvalidationReason) const;
+		SLATECORE_API bool ProtectedIsBound(const SWidget& Widget, ESlateAttributeType AttributeType) const;
+		SLATECORE_API ISlateAttributeGetter* ProtectedFindGetter(const SWidget& Widget, ESlateAttributeType AttributeType) const;
+		SLATECORE_API FDelegateHandle ProtectedFindGetterHandle(const SWidget& Widget, ESlateAttributeType AttributeType) const;
+		SLATECORE_API void ProtectedUpdateNow(SWidget& Widget, ESlateAttributeType AttributeType);
 
 		//~ For Member
-		bool ProtectedIsWidgetInDestructionPath(SWidget* Widget) const;
-		bool ProtectedIsImplemented(const SWidget& Widget) const;
+		SLATECORE_API bool ProtectedIsWidgetInDestructionPath(SWidget* Widget) const;
+		SLATECORE_API bool ProtectedIsImplemented(const SWidget& Widget) const;
 
 		//~ For Manage
-		void ProtectedMoveAttribute(SWidget& Widget, ESlateAttributeType AttributeType, const FSlateAttributeBase* Other);
+		SLATECORE_API void ProtectedMoveAttribute(SWidget& Widget, ESlateAttributeType AttributeType, const FSlateAttributeBase* Other);
 
 		//~ For Contain
-		void ProtectedUnregisterAttribute(ISlateAttributeContainer& Widget, ESlateAttributeType AttributeType) const;
-		void ProtectedRegisterAttribute(ISlateAttributeContainer& Widget, ESlateAttributeType AttributeType, TUniquePtr<ISlateAttributeGetter>&& Wrapper);
-		void ProtectedInvalidateWidget(ISlateAttributeContainer& Widget, ESlateAttributeType AttributeType, EInvalidateWidgetReason InvalidationReason) const;
-		bool ProtectedIsBound(const ISlateAttributeContainer& Widget, ESlateAttributeType AttributeType) const;
-		ISlateAttributeGetter* ProtectedFindGetter(const ISlateAttributeContainer& Widget, ESlateAttributeType AttributeType) const;
-		FDelegateHandle ProtectedFindGetterHandle(const ISlateAttributeContainer& Widget, ESlateAttributeType AttributeType) const;
-		void ProtectedUpdateNow(ISlateAttributeContainer& Widget, ESlateAttributeType AttributeType);
+		SLATECORE_API void ProtectedUnregisterAttribute(ISlateAttributeContainer& Widget, ESlateAttributeType AttributeType) const;
+		SLATECORE_API void ProtectedRegisterAttribute(ISlateAttributeContainer& Widget, ESlateAttributeType AttributeType, TUniquePtr<ISlateAttributeGetter>&& Wrapper);
+		SLATECORE_API void ProtectedInvalidateWidget(ISlateAttributeContainer& Widget, ESlateAttributeType AttributeType, EInvalidateWidgetReason InvalidationReason) const;
+		SLATECORE_API bool ProtectedIsBound(const ISlateAttributeContainer& Widget, ESlateAttributeType AttributeType) const;
+		SLATECORE_API ISlateAttributeGetter* ProtectedFindGetter(const ISlateAttributeContainer& Widget, ESlateAttributeType AttributeType) const;
+		SLATECORE_API FDelegateHandle ProtectedFindGetterHandle(const ISlateAttributeContainer& Widget, ESlateAttributeType AttributeType) const;
+		SLATECORE_API void ProtectedUpdateNow(ISlateAttributeContainer& Widget, ESlateAttributeType AttributeType);
 	};
 
 } // SlateAttributePrivate
diff --git a/Engine/Source/Runtime/SlateCore/Public/Types/NavigationMetaData.h b/Engine/Source/Runtime/SlateCore/Public/Types/NavigationMetaData.h
--- a/Engine/Source/Runtime/SlateCore/Public/Types/NavigationMetaData.h
+++ b/Engine/Source/Runtime/SlateCore/Public/Types/NavigationMetaData.h
@@ -12,7 +12,7 @@ class SWidget;
 /**
  *  Metadata to override the navigation behavior or regular SWidget
  */
-class SLATECORE_API FNavigationMetaData : public ISlateMetaData
+class FNavigationMetaData : public ISlateMetaData
 {
 public:
 	SLATE_METADATA_TYPE(FNavigationMetaData, ISlateMetaData)
diff --git a/Engine/Source/Runtime/SlateCore/Public/Types/PaintArgs.h b/Engine/Source/Runtime/SlateCore/Public/Types/PaintArgs.h
--- a/Engine/Source/Runtime/SlateCore/Public/Types/PaintArgs.h
+++ b/Engine/Source/Runtime/SlateCore/Public/Types/PaintArgs.h
@@ -19,13 +19,13 @@ class SWindow;
  * sole parameter in order to ease the burden of passing
  * through multiple fields.
  */
-class SLATECORE_API FPaintArgs
+class FPaintArgs
 {
 	friend class SInvalidationPanel;
 	friend class SRetainerWidget;
 public:
-	FPaintArgs(const SWidget* PaintParent, FHittestGrid& InRootHittestGrid, FHittestGrid& InCurrentHitTestGrid, UE::Slate::FDeprecateVector2DParameter InWindowOffset, double InCurrentTime, float InDeltaTime);
-	FPaintArgs(const SWidget* PaintParent, FHittestGrid& InRootHittestGrid, UE::Slate::FDeprecateVector2DParameter InWindowOffset, double InCurrentTime, float InDeltaTime);
+	SLATECORE_API FPaintArgs(const SWidget* PaintParent, FHittestGrid& InRootHittestGrid, FHittestGrid& InCurrentHitTestGrid, UE::Slate::FDeprecateVector2DParameter InWindowOffset, double InCurrentTime, float InDeltaTime);
+	SLATECORE_API FPaintArgs(const SWidget* PaintParent, FHittestGrid& InRootHittestGrid, UE::Slate::FDeprecateVector2DParameter InWindowOffset, double InCurrentTime, float InDeltaTime);
 
 	UE_NODISCARD FORCEINLINE_DEBUGGABLE FPaintArgs WithNewParent(const SWidget* PaintParent) const
 	{
@@ -41,7 +41,7 @@ public:
 		return NewArgs;
 	}
 
-	FPaintArgs InsertCustomHitTestPath(const SWidget* Widget, TSharedRef<ICustomHitTestPath> CustomHitTestPath) const;
+	SLATECORE_API FPaintArgs InsertCustomHitTestPath(const SWidget* Widget, TSharedRef<ICustomHitTestPath> CustomHitTestPath) const;
 
 	void SetInheritedHittestability(bool InInheritedHittestability)
 	{
diff --git a/Engine/Source/Runtime/SlateCore/Public/Types/ReflectionMetadata.h b/Engine/Source/Runtime/SlateCore/Public/Types/ReflectionMetadata.h
--- a/Engine/Source/Runtime/SlateCore/Public/Types/ReflectionMetadata.h
+++ b/Engine/Source/Runtime/SlateCore/Public/Types/ReflectionMetadata.h
@@ -15,7 +15,7 @@ class SWidget;
  * additional information about slate widgets that are constructed by UObject
  * classes for UMG.
  */
-class SLATECORE_API FReflectionMetaData : public ISlateMetaData
+class FReflectionMetaData : public ISlateMetaData
 {
 public:
 	SLATE_METADATA_TYPE(FReflectionMetaData, ISlateMetaData)
@@ -42,11 +42,11 @@ public:
 
 public:
 
-	static FString GetWidgetPath(const SWidget* InWidget, bool bShort = true, bool bNativePathOnly = false);
-	static FString GetWidgetPath(const SWidget& InWidget, bool bShort = true, bool bNativePathOnly = false);
+	static SLATECORE_API FString GetWidgetPath(const SWidget* InWidget, bool bShort = true, bool bNativePathOnly = false);
+	static SLATECORE_API FString GetWidgetPath(const SWidget& InWidget, bool bShort = true, bool bNativePathOnly = false);
 
-	static FString GetWidgetDebugInfo(const SWidget* InWidget);
-	static FString GetWidgetDebugInfo(const SWidget& InWidget);
+	static SLATECORE_API FString GetWidgetDebugInfo(const SWidget* InWidget);
+	static SLATECORE_API FString GetWidgetDebugInfo(const SWidget& InWidget);
 
-	static TSharedPtr<FReflectionMetaData> GetWidgetOrParentMetaData(const SWidget* InWidget);
+	static SLATECORE_API TSharedPtr<FReflectionMetaData> GetWidgetOrParentMetaData(const SWidget* InWidget);
 };
diff --git a/Engine/Source/Runtime/SlateCore/Public/Types/SlateAttributeDescriptor.h b/Engine/Source/Runtime/SlateCore/Public/Types/SlateAttributeDescriptor.h
--- a/Engine/Source/Runtime/SlateCore/Public/Types/SlateAttributeDescriptor.h
+++ b/Engine/Source/Runtime/SlateCore/Public/Types/SlateAttributeDescriptor.h
@@ -28,7 +28,7 @@ namespace SlateAttributePrivate
 /**
  * Describes the static information about a Widget's type SlateAttributes.
  **/
-class SLATECORE_API FSlateAttributeDescriptor
+class FSlateAttributeDescriptor
 {
 public:
 	/**
@@ -199,28 +199,28 @@ public:
 	};
 
 	/** Internal class to initialize the SlateAttributeDescriptor::FContainer attributes (Add attributes or modify existing attributes). */
-	struct SLATECORE_API FContainerInitializer
+	struct FContainerInitializer
 	{
 	private:
 		friend FSlateAttributeDescriptor;
-		FContainerInitializer(FSlateAttributeDescriptor& InDescriptor, FName ContainerName);
-		FContainerInitializer(FSlateAttributeDescriptor& InDescriptor, const FSlateAttributeDescriptor& ParentDescriptor, FName ContainerName);
+		SLATECORE_API FContainerInitializer(FSlateAttributeDescriptor& InDescriptor, FName ContainerName);
+		SLATECORE_API FContainerInitializer(FSlateAttributeDescriptor& InDescriptor, const FSlateAttributeDescriptor& ParentDescriptor, FName ContainerName);
 
 	public:
 		FContainerInitializer() = delete;
 		FContainerInitializer(const FContainerInitializer&) = delete;
 		FContainerInitializer& operator= (const FContainerInitializer&) = delete;
 
-		struct SLATECORE_API FAttributeEntry
+		struct FAttributeEntry
 		{
-			FAttributeEntry(FSlateAttributeDescriptor& Descriptor, FName ContainerName, int32 AttributeIndex);
+			SLATECORE_API FAttributeEntry(FSlateAttributeDescriptor& Descriptor, FName ContainerName, int32 AttributeIndex);
 
 			/**
 			 * Update the attribute after the prerequisite.
 			 * The order is guaranteed but other attributes may be updated in between.
 			 * No order is guaranteed if the prerequisite or this property is updated manually.
 			 */
-			FAttributeEntry& UpdatePrerequisite(FName Prerequisite);
+			SLATECORE_API FAttributeEntry& UpdatePrerequisite(FName Prerequisite);
 
 			/**
 			 * Notified when the attribute value changed.
@@ -229,7 +229,7 @@ public:
 			 * It will not be called when the SWidget is in its construction phase.
 			 * @see SWidget::IsConstructed
 			 */
-			FAttributeEntry& OnValueChanged(FAttributeValueChangedDelegate Callback);
+			SLATECORE_API FAttributeEntry& OnValueChanged(FAttributeValueChangedDelegate Callback);
 
 		private:
 			FSlateAttributeDescriptor& Descriptor;
@@ -237,17 +237,17 @@ public:
 			int32 AttributeIndex;
 		};
 
-		FAttributeEntry AddContainedAttribute(FName AttributeName, OffsetType Offset, const FInvalidateWidgetReasonAttribute& ReasonGetter);
-		FAttributeEntry AddContainedAttribute(FName AttributeName, OffsetType Offset, FInvalidateWidgetReasonAttribute&& ReasonGetter);
+		SLATECORE_API FAttributeEntry AddContainedAttribute(FName AttributeName, OffsetType Offset, const FInvalidateWidgetReasonAttribute& ReasonGetter);
+		SLATECORE_API FAttributeEntry AddContainedAttribute(FName AttributeName, OffsetType Offset, FInvalidateWidgetReasonAttribute&& ReasonGetter);
 
 	public:
 		/** Change the InvalidationReason of an attribute defined in a base class. */
-		void OverrideInvalidationReason(FName AttributeName, const FInvalidateWidgetReasonAttribute& Reason);
+		SLATECORE_API void OverrideInvalidationReason(FName AttributeName, const FInvalidateWidgetReasonAttribute& Reason);
 		/** Change the InvalidationReason of an attribute defined in a base class. */
-		void OverrideInvalidationReason(FName AttributeName, FInvalidateWidgetReasonAttribute&& Reason);
+		SLATECORE_API void OverrideInvalidationReason(FName AttributeName, FInvalidateWidgetReasonAttribute&& Reason);
 
 		/** Change the FAttributeValueChangedDelegate of an attribute defined in a base class. */
-		void OverrideOnValueChanged(FName AttributeName, ECallbackOverrideType OverrideType, FAttributeValueChangedDelegate Callback);
+		SLATECORE_API void OverrideOnValueChanged(FName AttributeName, ECallbackOverrideType OverrideType, FAttributeValueChangedDelegate Callback);
 
 	private:
 		FSlateAttributeDescriptor& Descriptor;
@@ -255,35 +255,35 @@ public:
 	};
 
 	/** Internal class to initialize the SlateAttributeDescriptor (Add attributes or modify existing attributes). */
-	struct SLATECORE_API FInitializer
+	struct FInitializer
 	{
 	private:
 		friend FSlateWidgetClassData;
-		FInitializer(FSlateAttributeDescriptor& InDescriptor);
-		FInitializer(FSlateAttributeDescriptor& InDescriptor, const FSlateAttributeDescriptor& ParentDescriptor);
+		SLATECORE_API FInitializer(FSlateAttributeDescriptor& InDescriptor);
+		SLATECORE_API FInitializer(FSlateAttributeDescriptor& InDescriptor, const FSlateAttributeDescriptor& ParentDescriptor);
 		FInitializer(const FInitializer&) = delete;
 		FInitializer& operator=(const FInitializer&) = delete;
 
 	public:
-		~FInitializer();
+		SLATECORE_API ~FInitializer();
 
-		struct SLATECORE_API FAttributeEntry
+		struct FAttributeEntry
 		{
-			FAttributeEntry(FSlateAttributeDescriptor& Descriptor, int32 InAttributeIndex);
+			SLATECORE_API FAttributeEntry(FSlateAttributeDescriptor& Descriptor, int32 InAttributeIndex);
 
 			/**
 			 * Update the attribute after the prerequisite.
 			 * The order is guaranteed but other attributes may be updated in between.
 			 * No order is guaranteed if the prerequisite or this property is updated manually.
 			 */
-			FAttributeEntry& UpdatePrerequisite(FName Prerequisite);
+			SLATECORE_API FAttributeEntry& UpdatePrerequisite(FName Prerequisite);
 
 			/**
 			 * The attribute affect the visibility of the widget.
 			 * We only update the attributes that can change the visibility of the widget when the widget is collapsed.
 			 * Attributes that affect visibility must have the Visibility attribute as a Prerequisite or the Visibility attribute must have it as a Prerequisite.
 			 */
-			FAttributeEntry& AffectVisibility();
+			SLATECORE_API FAttributeEntry& AffectVisibility();
 
 			/**
 			 * Notified when the attribute value changed.
@@ -292,29 +292,29 @@ public:
 			 * It will not be called when the SWidget is in its construction phase.
 			 * @see SWidget::IsConstructed
 			 */
-			FAttributeEntry& OnValueChanged(FAttributeValueChangedDelegate Callback);
+			SLATECORE_API FAttributeEntry& OnValueChanged(FAttributeValueChangedDelegate Callback);
 
 		private:
 			FSlateAttributeDescriptor& Descriptor;
 			int32 AttributeIndex;
 		};
 
-		FAttributeEntry AddMemberAttribute(FName AttributeName, OffsetType Offset, const FInvalidateWidgetReasonAttribute& ReasonGetter);
-		FAttributeEntry AddMemberAttribute(FName AttributeName, OffsetType Offset, FInvalidateWidgetReasonAttribute&& ReasonGetter);
+		SLATECORE_API FAttributeEntry AddMemberAttribute(FName AttributeName, OffsetType Offset, const FInvalidateWidgetReasonAttribute& ReasonGetter);
+		SLATECORE_API FAttributeEntry AddMemberAttribute(FName AttributeName, OffsetType Offset, FInvalidateWidgetReasonAttribute&& ReasonGetter);
 
-		FContainerInitializer AddContainer(FName ContainerName, OffsetType Offset);
+		SLATECORE_API FContainerInitializer AddContainer(FName ContainerName, OffsetType Offset);
 
 	public:
 		/** Change the InvalidationReason of an attribute defined in a base class. */
-		void OverrideInvalidationReason(FName AttributeName, const FInvalidateWidgetReasonAttribute& Reason);
+		SLATECORE_API void OverrideInvalidationReason(FName AttributeName, const FInvalidateWidgetReasonAttribute& Reason);
 		/** Change the InvalidationReason of an attribute defined in a base class. */
-		void OverrideInvalidationReason(FName AttributeName, FInvalidateWidgetReasonAttribute&& Reason);
+		SLATECORE_API void OverrideInvalidationReason(FName AttributeName, FInvalidateWidgetReasonAttribute&& Reason);
 
 		/** Change the FAttributeValueChangedDelegate of an attribute defined in a base class. */
-		void OverrideOnValueChanged(FName AttributeName, ECallbackOverrideType OverrideType, FAttributeValueChangedDelegate Callback);
+		SLATECORE_API void OverrideOnValueChanged(FName AttributeName, ECallbackOverrideType OverrideType, FAttributeValueChangedDelegate Callback);
 
 		/** Change the update type of an attribute defined in a base class. */
-		void SetAffectVisibility(FName AttributeName, bool bAffectVisibility);
+		SLATECORE_API void SetAffectVisibility(FName AttributeName, bool bAffectVisibility);
 
 	private:
 		FSlateAttributeDescriptor& Descriptor;
@@ -327,42 +327,42 @@ public:
 	}
 
 	/** @returns the Attribute at the index previously found with IndexOfMemberAttribute */
-	const FAttribute& GetAttributeAtIndex(int32 Index) const;
+	SLATECORE_API const FAttribute& GetAttributeAtIndex(int32 Index) const;
 
 	/** @returns the Container with the corresponding name. */
-	const FContainer* FindContainer(FName ContainerName) const;
+	SLATECORE_API const FContainer* FindContainer(FName ContainerName) const;
 
 	/** @returns the Attribute with the corresponding name. */
-	const FAttribute* FindAttribute(FName AttributeName) const;
+	SLATECORE_API const FAttribute* FindAttribute(FName AttributeName) const;
 
 	/** @returns the Attribute of a SlateAttribute that have the corresponding memory offset. */
-	const FAttribute* FindMemberAttribute(OffsetType AttributeOffset) const;
+	SLATECORE_API const FAttribute* FindMemberAttribute(OffsetType AttributeOffset) const;
 
 	/** @returns the Attribute of a SlateAttribute that have the corresponding memory offset. */
-	const FAttribute* FindContainedAttribute(FName ContainerName, OffsetType AttributeOffset) const;
+	SLATECORE_API const FAttribute* FindContainedAttribute(FName ContainerName, OffsetType AttributeOffset) const;
 
 	/** @returns the index of the Container with the corresponding name. */
-	int32 IndexOfContainer(FName AttributeName) const;
+	SLATECORE_API int32 IndexOfContainer(FName AttributeName) const;
 
 	/** @returns the index of a SlateAttribute that have the corresponding memory offset. */
-	int32 IndexOfAttribute(FName AttributeName) const;
+	SLATECORE_API int32 IndexOfAttribute(FName AttributeName) const;
 
 	/** @returns the index of a SlateAttribute that have the corresponding memory offset. */
-	int32 IndexOfMemberAttribute(OffsetType AttributeOffset) const;
+	SLATECORE_API int32 IndexOfMemberAttribute(OffsetType AttributeOffset) const;
 
 	/** @returns the index of a SlateAttribute that have the corresponding memory offset. */
-	int32 IndexOfContainedAttribute(FName ContainerName, OffsetType AttributeOffset) const;
+	SLATECORE_API int32 IndexOfContainedAttribute(FName ContainerName, OffsetType AttributeOffset) const;
 
 private:
-	FAttribute* FindAttribute(FName AttributeName);
-
-	FContainerInitializer AddContainer(FName AttributeName, OffsetType Offset);
-	FInitializer::FAttributeEntry AddMemberAttribute(FName AttributeName, OffsetType Offset, FInvalidateWidgetReasonAttribute ReasonGetter);
-	FContainerInitializer::FAttributeEntry AddContainedAttribute(FName ContainerName, FName AttributeName, OffsetType Offset, FInvalidateWidgetReasonAttribute ReasonGetter);
-	void OverrideInvalidationReason(FName ContainerName, FName AttributeName, FInvalidateWidgetReasonAttribute ReasonGetter);
-	void OverrideOnValueChanged(FName ContainerName, FName AttributeName, ECallbackOverrideType OverrideType, FAttributeValueChangedDelegate Callback);
-	void SetPrerequisite(FName ContainerName, FAttribute& Attribute, FName Prerequisite);
-	void SetAffectVisibility(FAttribute& Attribute, bool bUpdate);
+	SLATECORE_API FAttribute* FindAttribute(FName AttributeName);
+
+	SLATECORE_API FContainerInitializer AddContainer(FName AttributeName, OffsetType Offset);
+	SLATECORE_API FInitializer::FAttributeEntry AddMemberAttribute(FName AttributeName, OffsetType Offset, FInvalidateWidgetReasonAttribute ReasonGetter);
+	SLATECORE_API FContainerInitializer::FAttributeEntry AddContainedAttribute(FName ContainerName, FName AttributeName, OffsetType Offset, FInvalidateWidgetReasonAttribute ReasonGetter);
+	SLATECORE_API void OverrideInvalidationReason(FName ContainerName, FName AttributeName, FInvalidateWidgetReasonAttribute ReasonGetter);
+	SLATECORE_API void OverrideOnValueChanged(FName ContainerName, FName AttributeName, ECallbackOverrideType OverrideType, FAttributeValueChangedDelegate Callback);
+	SLATECORE_API void SetPrerequisite(FName ContainerName, FAttribute& Attribute, FName Prerequisite);
+	SLATECORE_API void SetAffectVisibility(FAttribute& Attribute, bool bUpdate);
 
 private:
 	TArray<FAttribute> Attributes;
diff --git a/Engine/Source/Runtime/SlateCore/Public/Types/SlateAttributeMetaData.h b/Engine/Source/Runtime/SlateCore/Public/Types/SlateAttributeMetaData.h
--- a/Engine/Source/Runtime/SlateCore/Public/Types/SlateAttributeMetaData.h
+++ b/Engine/Source/Runtime/SlateCore/Public/Types/SlateAttributeMetaData.h
@@ -14,7 +14,7 @@
 class SWidget;
 
 /** */
-class SLATECORE_API FSlateAttributeMetaData : public ISlateMetaData
+class FSlateAttributeMetaData : public ISlateMetaData
 {
 	friend SlateAttributePrivate::FSlateAttributeImpl;
 	friend SlateAttributePrivate::ISlateAttributeContainer;
@@ -28,7 +28,7 @@ public:
 
 public:
 	/** @return the instance associated to the SWidget (if it exists). */
-	static FSlateAttributeMetaData* FindMetaData(const SWidget& OwningWidget);
+	static SLATECORE_API FSlateAttributeMetaData* FindMetaData(const SWidget& OwningWidget);
 
 	enum class EInvalidationPermission : uint8
 	{
@@ -48,26 +48,26 @@ public:
 	 * Update all the attributes.
 	 * @param InvalidationStyle if we should invalidate the widget.
 	 */
-	static void UpdateAllAttributes(SWidget& OwningWidget, EInvalidationPermission InvalidationStyle);
+	static SLATECORE_API void UpdateAllAttributes(SWidget& OwningWidget, EInvalidationPermission InvalidationStyle);
 	/**
 	 * Update attributes that are responsible to change visibility of the widget.
 	 * @param InvalidationStyle if we should invalidate the widget.
 	 */
-	static void UpdateOnlyVisibilityAttributes(SWidget& OwningWidget, EInvalidationPermission InvalidationStyle);
+	static SLATECORE_API void UpdateOnlyVisibilityAttributes(SWidget& OwningWidget, EInvalidationPermission InvalidationStyle);
 	/**
 	 * Update attributes that are NOT responsible to change visibility of the widget.
 	 * @param InvalidationStyle if we should invalidate the widget.
 	 */
-	static void UpdateExceptVisibilityAttributes(SWidget& OwningWidget, EInvalidationPermission InvalidationStyle);
+	static SLATECORE_API void UpdateExceptVisibilityAttributes(SWidget& OwningWidget, EInvalidationPermission InvalidationStyle);
 	/**
 	 * Execute UpdateOnlyVisibilityAttributes on every children of the widget.
 	 * @param InvalidationStyle if we should invalidate the widget.
 	 */
-	static void UpdateChildrenOnlyVisibilityAttributes(SWidget& OwningWidget, EInvalidationPermission InvalidationStyle, bool bRecursive);
+	static SLATECORE_API void UpdateChildrenOnlyVisibilityAttributes(SWidget& OwningWidget, EInvalidationPermission InvalidationStyle, bool bRecursive);
 	/**
 	 * Apply any delayed invalidation previously delayed with EInvalidationPermission::DelayInvalidation.
 	 */
-	static void ApplyDelayedInvalidation(SWidget& OwningWidget);
+	static SLATECORE_API void ApplyDelayedInvalidation(SWidget& OwningWidget);
 
 public:
 	bool IsBound(const FSlateAttributeBase& Attribute) const
@@ -86,7 +86,7 @@ public:
 	}
 
 	/** Get the name of all the attributes, if available. */
-	static TArray<FName> GetAttributeNames(const SWidget& OwningWidget);
+	static SLATECORE_API TArray<FName> GetAttributeNames(const SWidget& OwningWidget);
 
 	/**
 	 * Helper function that can be used inside the debuggers watch window.
@@ -99,31 +99,31 @@ private:
 	using ISlateAttributeContainer = SlateAttributePrivate::ISlateAttributeContainer;
 	using ESlateAttributeType = SlateAttributePrivate::ESlateAttributeType;
 	using ISlateAttributeGetter = SlateAttributePrivate::ISlateAttributeGetter;
-	static void RegisterAttribute(SWidget& OwningWidget, FSlateAttributeBase& Attribute, ESlateAttributeType AttributeType, TUniquePtr<ISlateAttributeGetter>&& Wrapper);
-	static void RegisterAttribute(ISlateAttributeContainer& OwningContainer, FSlateAttributeBase& Attribute, ESlateAttributeType AttributeType, TUniquePtr<ISlateAttributeGetter>&& Wrapper);
-	static bool UnregisterAttribute(SWidget& OwningWidget, const FSlateAttributeBase& Attribute);
-	static bool UnregisterAttribute(ISlateAttributeContainer& OwningContainer, const FSlateAttributeBase& Attribute);
-	static void InvalidateWidget(SWidget& OwningWidget, const FSlateAttributeBase& Attribute, ESlateAttributeType AttributeType, EInvalidateWidgetReason Reason);
-	static void InvalidateWidget(ISlateAttributeContainer& OwningContainer, const FSlateAttributeBase& Attribute, ESlateAttributeType AttributeType, EInvalidateWidgetReason Reason);
-	static void UpdateAttribute(SWidget& OwningWidget, FSlateAttributeBase& Attribute);
-	static bool IsAttributeBound(const SWidget& OwningWidget, const FSlateAttributeBase& Attribute);
-	static SlateAttributePrivate::ISlateAttributeGetter* GetAttributeGetter(const SWidget& OwningWidget, const FSlateAttributeBase& Attribute);
-	static FDelegateHandle GetAttributeGetterHandle(const SWidget& OwningWidget, const FSlateAttributeBase& Attribute);
-	static void MoveAttribute(const SWidget& OwningWidget, FSlateAttributeBase& NewAttribute, ESlateAttributeType AttributeType, const FSlateAttributeBase* PreviousAttribute);
+	static SLATECORE_API void RegisterAttribute(SWidget& OwningWidget, FSlateAttributeBase& Attribute, ESlateAttributeType AttributeType, TUniquePtr<ISlateAttributeGetter>&& Wrapper);
+	static SLATECORE_API void RegisterAttribute(ISlateAttributeContainer& OwningContainer, FSlateAttributeBase& Attribute, ESlateAttributeType AttributeType, TUniquePtr<ISlateAttributeGetter>&& Wrapper);
+	static SLATECORE_API bool UnregisterAttribute(SWidget& OwningWidget, const FSlateAttributeBase& Attribute);
+	static SLATECORE_API bool UnregisterAttribute(ISlateAttributeContainer& OwningContainer, const FSlateAttributeBase& Attribute);
+	static SLATECORE_API void InvalidateWidget(SWidget& OwningWidget, const FSlateAttributeBase& Attribute, ESlateAttributeType AttributeType, EInvalidateWidgetReason Reason);
+	static SLATECORE_API void InvalidateWidget(ISlateAttributeContainer& OwningContainer, const FSlateAttributeBase& Attribute, ESlateAttributeType AttributeType, EInvalidateWidgetReason Reason);
+	static SLATECORE_API void UpdateAttribute(SWidget& OwningWidget, FSlateAttributeBase& Attribute);
+	static SLATECORE_API bool IsAttributeBound(const SWidget& OwningWidget, const FSlateAttributeBase& Attribute);
+	static SLATECORE_API SlateAttributePrivate::ISlateAttributeGetter* GetAttributeGetter(const SWidget& OwningWidget, const FSlateAttributeBase& Attribute);
+	static SLATECORE_API FDelegateHandle GetAttributeGetterHandle(const SWidget& OwningWidget, const FSlateAttributeBase& Attribute);
+	static SLATECORE_API void MoveAttribute(const SWidget& OwningWidget, FSlateAttributeBase& NewAttribute, ESlateAttributeType AttributeType, const FSlateAttributeBase* PreviousAttribute);
 
 	//~ For Container
-	static void RemoveContainerWidget(SWidget& Widget, ISlateAttributeContainer& Container);
-	static void UpdateContainerSortOrder(SWidget& Widget, ISlateAttributeContainer& Container);
+	static SLATECORE_API void RemoveContainerWidget(SWidget& Widget, ISlateAttributeContainer& Container);
+	static SLATECORE_API void UpdateContainerSortOrder(SWidget& Widget, ISlateAttributeContainer& Container);
 
 
 private:
-	void RegisterMemberAttributeImpl(SWidget& OwningWidget, FSlateAttributeBase& Attribute, TUniquePtr<ISlateAttributeGetter>&& Getter);
-	void RegisterManagedAttributeImpl(SWidget& OwningWidget, FSlateAttributeBase& Attribute, TUniquePtr<ISlateAttributeGetter>&& Getter);
-	void RegisterContainAttributeImpl(ISlateAttributeContainer& OwningContainer, SWidget& OwningWidget, FSlateAttributeBase& Attribute, TUniquePtr<ISlateAttributeGetter>&& Getter);
-	bool UnregisterAttributeImpl(const FSlateAttributeBase& Attribute);
-	void RemoveMetaDataIfNeeded(SWidget& OwningWidget, bool bRemoved) const;
-	bool InvalidateWidgetAttribute(SWidget& OwningWidget, const FSlateAttributeBase& Attribute, EInvalidateWidgetReason Reason);
-	void UpdateAttributesImpl(SWidget& OwningWidget, EInvalidationPermission InvaldiationStyle, int32 StartIndex, int32 IndexNum);
+	SLATECORE_API void RegisterMemberAttributeImpl(SWidget& OwningWidget, FSlateAttributeBase& Attribute, TUniquePtr<ISlateAttributeGetter>&& Getter);
+	SLATECORE_API void RegisterManagedAttributeImpl(SWidget& OwningWidget, FSlateAttributeBase& Attribute, TUniquePtr<ISlateAttributeGetter>&& Getter);
+	SLATECORE_API void RegisterContainAttributeImpl(ISlateAttributeContainer& OwningContainer, SWidget& OwningWidget, FSlateAttributeBase& Attribute, TUniquePtr<ISlateAttributeGetter>&& Getter);
+	SLATECORE_API bool UnregisterAttributeImpl(const FSlateAttributeBase& Attribute);
+	SLATECORE_API void RemoveMetaDataIfNeeded(SWidget& OwningWidget, bool bRemoved) const;
+	SLATECORE_API bool InvalidateWidgetAttribute(SWidget& OwningWidget, const FSlateAttributeBase& Attribute, EInvalidateWidgetReason Reason);
+	SLATECORE_API void UpdateAttributesImpl(SWidget& OwningWidget, EInvalidationPermission InvaldiationStyle, int32 StartIndex, int32 IndexNum);
 
 private:
 	int32 IndexOfAttribute(const FSlateAttributeBase& Attribute) const
diff --git a/Engine/Source/Runtime/SlateCore/Public/Types/SlateVector2.h b/Engine/Source/Runtime/SlateCore/Public/Types/SlateVector2.h
--- a/Engine/Source/Runtime/SlateCore/Public/Types/SlateVector2.h
+++ b/Engine/Source/Runtime/SlateCore/Public/Types/SlateVector2.h
@@ -74,16 +74,16 @@ namespace UE::Slate
 	const FVector2f& CastToVector2f(const FDeprecateSlateVector2D& InValue);
 	const FVector2f& CastToVector2f(const FDeprecateVector2DParameter& InValue);
 
-	struct SLATECORE_API FDeprecateSlateVectorPtrVariant
+	struct FDeprecateSlateVectorPtrVariant
 	{
-		explicit FDeprecateSlateVectorPtrVariant(FDeprecateSlateVector2D* InInstance);
+		SLATECORE_API explicit FDeprecateSlateVectorPtrVariant(FDeprecateSlateVector2D* InInstance);
 
 		UE_SLATE_VECTOR_DEPRECATED_DEFAULT()
-		operator const FVector2D*() const &;
+		SLATECORE_API operator const FVector2D*() const &;
 
-		operator const FVector2f*() const;
+		SLATECORE_API operator const FVector2f*() const;
 
-		operator FVector2f*();
+		SLATECORE_API operator FVector2f*();
 
 		FDeprecateSlateVector2D* Instance;
 		mutable FVector2D DoubleVector;
@@ -596,4 +596,4 @@ namespace UE::Slate
 		return FVector2f(X, Y);
 	}
 
-} // namespace UE::Slate
\ No newline at end of file
+} // namespace UE::Slate
diff --git a/Engine/Source/Runtime/SlateCore/Public/Types/TrackedMetaData.h b/Engine/Source/Runtime/SlateCore/Public/Types/TrackedMetaData.h
--- a/Engine/Source/Runtime/SlateCore/Public/Types/TrackedMetaData.h
+++ b/Engine/Source/Runtime/SlateCore/Public/Types/TrackedMetaData.h
@@ -11,15 +11,15 @@ class SWidget;
 /**
  * MetaData used to add and remove widgets to the Slate Widget Tracker for the specified tags.
  */
-class SLATECORE_API FTrackedMetaData : public ISlateMetaData
+class FTrackedMetaData : public ISlateMetaData
 {
 public:
 
 	SLATE_METADATA_TYPE(FTrackedMetaData, ISlateMetaData)
 
-	FTrackedMetaData(const SWidget* InTrackedWidget, TArray<FName>&& InTags);
-	FTrackedMetaData(const SWidget* InTrackedWidget, FName InTags);
-	~FTrackedMetaData();
+	SLATECORE_API FTrackedMetaData(const SWidget* InTrackedWidget, TArray<FName>&& InTags);
+	SLATECORE_API FTrackedMetaData(const SWidget* InTrackedWidget, FName InTags);
+	SLATECORE_API ~FTrackedMetaData();
 
 	FTrackedMetaData(const FTrackedMetaData&) = delete;
 	FTrackedMetaData& operator=(const FTrackedMetaData&) = delete;
diff --git a/Engine/Source/Runtime/SlateCore/Public/Widgets/Accessibility/SlateAccessibleMessageHandler.h b/Engine/Source/Runtime/SlateCore/Public/Widgets/Accessibility/SlateAccessibleMessageHandler.h
--- a/Engine/Source/Runtime/SlateCore/Public/Widgets/Accessibility/SlateAccessibleMessageHandler.h
+++ b/Engine/Source/Runtime/SlateCore/Public/Widgets/Accessibility/SlateAccessibleMessageHandler.h
@@ -18,19 +18,19 @@ class SWindow;
  * Message handling system for Slate Accessibility API, dealing with both receiving events and pushing them back to the platform layer.
  * The message handler is also responsible for processing the Slate widget tree and queuing/processing accessibility requests for widget data from the platform.
  */
-class SLATECORE_API FSlateAccessibleMessageHandler : public FGenericAccessibleMessageHandler
+class FSlateAccessibleMessageHandler : public FGenericAccessibleMessageHandler
 {
 public:
-	FSlateAccessibleMessageHandler();
+	SLATECORE_API FSlateAccessibleMessageHandler();
 
 	// FGenericAccessibleMessageHandler
-	virtual void OnActivate() override;
-	virtual void OnDeactivate() override;
-	virtual TSharedPtr<IAccessibleWidget> GetAccessibleWindow(const TSharedRef<FGenericWindow>& InWindow) const override;
-	virtual AccessibleWidgetId GetAccessibleWindowId(const TSharedRef<FGenericWindow>& InWindow) const override;
-	virtual TSharedPtr<IAccessibleWidget> GetAccessibleWidgetFromId(AccessibleWidgetId Id) const override;
-	virtual void RunInThread(const TFunction<void()>& InFunction, bool bWaitForCompletion = true, ENamedThreads::Type InThread = ENamedThreads::GameThread) override;
-	virtual void MakeAccessibleAnnouncement(const FString& AnnouncementString) override;
+	SLATECORE_API virtual void OnActivate() override;
+	SLATECORE_API virtual void OnDeactivate() override;
+	SLATECORE_API virtual TSharedPtr<IAccessibleWidget> GetAccessibleWindow(const TSharedRef<FGenericWindow>& InWindow) const override;
+	SLATECORE_API virtual AccessibleWidgetId GetAccessibleWindowId(const TSharedRef<FGenericWindow>& InWindow) const override;
+	SLATECORE_API virtual TSharedPtr<IAccessibleWidget> GetAccessibleWidgetFromId(AccessibleWidgetId Id) const override;
+	SLATECORE_API virtual void RunInThread(const TFunction<void()>& InFunction, bool bWaitForCompletion = true, ENamedThreads::Type InThread = ENamedThreads::GameThread) override;
+	SLATECORE_API virtual void MakeAccessibleAnnouncement(const FString& AnnouncementString) override;
 	//~
 
 	/**
@@ -38,7 +38,7 @@ public:
 	 *
 	 * @param Widget The widget that is being deleted.
 	 */
-	void OnWidgetRemoved(SWidget* Widget);
+	SLATECORE_API void OnWidgetRemoved(SWidget* Widget);
 
 	/**
 	 * The arguments that make up an accessible event raised by an SWidget.
@@ -74,7 +74,7 @@ public:
 	 * @param Args The arguments for the accessible event being raised.
 	 * @see FSlateWidgetAccessibleEventArgs
 	 */
-	void OnWidgetEventRaised(const FSlateWidgetAccessibleEventArgs& Args);
+	SLATECORE_API void OnWidgetEventRaised(const FSlateWidgetAccessibleEventArgs& Args);
 
 	/**
 	 * Refresh the accessible widget tree next available tick. This should be called any time the Slate tree changes.
@@ -85,16 +85,16 @@ public:
 	* Processes all the queued tasks in the AccessibleTaskQueue
 	* Should only be called from FSlateApplication::TickPlatform() 
 	*/
-	void ProcessAccessibleTasks();
+	SLATECORE_API void ProcessAccessibleTasks();
 
 	/**
 	 * Process any pending Slate widgets and update the accessible widget tree.
 	 */
-	void Tick();
+	SLATECORE_API void Tick();
 
 
 private:
-	void HandleAccessibleWidgetFocusChangeEvent(const TSharedRef<IAccessibleWidget>& FocusWidget, bool bIsWidgetGainingFocus, FAccessibleUserIndex UserIndex);
+	SLATECORE_API void HandleAccessibleWidgetFocusChangeEvent(const TSharedRef<IAccessibleWidget>& FocusWidget, bool bIsWidgetGainingFocus, FAccessibleUserIndex UserIndex);
 	/**
 	*  A helper class that wraps an accessibility task and the event to be triggered when the task finishes executing. 
 	*/
@@ -129,7 +129,7 @@ private:
 	};
 
 	/** Queues an FSlateAccessibleTask to be processed in the Game Thread */
-	void EnqueueAccessibleTask(const FSlateAccessibleTask& InAccessibleTask);
+	SLATECORE_API void EnqueueAccessibleTask(const FSlateAccessibleTask& InAccessibleTask);
 
 	struct FWidgetAndParent {
 		FWidgetAndParent(TWeakPtr<SWidget> InWidget, TSharedRef<FSlateAccessibleWidget> InParent)
diff --git a/Engine/Source/Runtime/SlateCore/Public/Widgets/Accessibility/SlateAccessibleWidgetCache.h b/Engine/Source/Runtime/SlateCore/Public/Widgets/Accessibility/SlateAccessibleWidgetCache.h
--- a/Engine/Source/Runtime/SlateCore/Public/Widgets/Accessibility/SlateAccessibleWidgetCache.h
+++ b/Engine/Source/Runtime/SlateCore/Public/Widgets/Accessibility/SlateAccessibleWidgetCache.h
@@ -13,11 +13,11 @@ class SWidget;
  * Singleton used to retrieve accessible widgets for a given Slate widget. Accessible widgets will persist
  * so longer as the OS is sending accessibility events, after which the platform layer should call ClearAll().
  */
-class SLATECORE_API FSlateAccessibleWidgetCache
+class FSlateAccessibleWidgetCache
 {
 public:
 	/** Empty the cache and release all cached accessible widgets. This should generally only be done with accessibility is turned off. */
-	static void ClearAll();
+	static SLATECORE_API void ClearAll();
 
 	/**
 	 * Callback for when an SWidget is deleted. FSlateAccessibleMessageHandler should be the only thing to call this.
@@ -25,7 +25,7 @@ public:
 	 *
 	 * @param Widget the Slate widget that is being deleted
 	 */
-	static TSharedPtr<FSlateAccessibleWidget> RemoveWidget(SWidget* Widget);
+	static SLATECORE_API TSharedPtr<FSlateAccessibleWidget> RemoveWidget(SWidget* Widget);
 
 	/**
 	 * Get a cached accessible widget for a given Slate widget. This may return nullptr in the case where the
@@ -35,7 +35,7 @@ public:
 	 * @param Widget The Slate widget to get the accessible widget for
 	 * @return The accessible widget for the given Slate widget, or nullptr if the widget is not accessible.
 	 */
-	static TSharedPtr<FSlateAccessibleWidget> GetAccessibleWidgetChecked(const TSharedPtr<SWidget>& Widget);
+	static SLATECORE_API TSharedPtr<FSlateAccessibleWidget> GetAccessibleWidgetChecked(const TSharedPtr<SWidget>& Widget);
 	/**
 	 * Get a cached accessible widget for a given Slate widget. The caller is responsible for making sure
 	 * that the widget is fully accessible in the current Slate widget tree.
@@ -43,7 +43,7 @@ public:
 	 * @param Widget The Slate widget to get the accessible widget for
 	 * @return The accessible widget for the given Slate widget
 	 */
-	static TSharedRef<FSlateAccessibleWidget> GetAccessibleWidget(const TSharedRef<SWidget>& Widget);
+	static SLATECORE_API TSharedRef<FSlateAccessibleWidget> GetAccessibleWidget(const TSharedRef<SWidget>& Widget);
 
 	/**
 	 * Get a cached accessible widget for an identifier that matches the accessible widget's GetId() return value.
@@ -53,12 +53,12 @@ public:
 	 * @param Id The Id of a widget that corresponds to that widget's GetId() function
 	 * @return The accessible widget for the given Id, or nullptr if there is no widget with that Id.
 	 */
-	static TSharedPtr<FSlateAccessibleWidget> GetAccessibleWidgetFromId(AccessibleWidgetId Id);
+	static SLATECORE_API TSharedPtr<FSlateAccessibleWidget> GetAccessibleWidgetFromId(AccessibleWidgetId Id);
 
-	static TMap<SWidget*, TSharedRef<FSlateAccessibleWidget>>::TConstIterator GetAllWidgets();
+	static SLATECORE_API TMap<SWidget*, TSharedRef<FSlateAccessibleWidget>>::TConstIterator GetAllWidgets();
 
 #if !UE_BUILD_SHIPPING
-	static void DumpAccessibilityStats();
+	static SLATECORE_API void DumpAccessibilityStats();
 #endif
 
 private:
