[SlateCore]

diff --git a/Engine/Source/Runtime/SlateCore/Public/Widgets/Accessibility/SlateCoreAccessibleWidgets.h b/Engine/Source/Runtime/SlateCore/Public/Widgets/Accessibility/SlateCoreAccessibleWidgets.h
--- a/Engine/Source/Runtime/SlateCore/Public/Widgets/Accessibility/SlateCoreAccessibleWidgets.h
+++ b/Engine/Source/Runtime/SlateCore/Public/Widgets/Accessibility/SlateCoreAccessibleWidgets.h
@@ -15,37 +15,37 @@ class SWindow;
  * inherit directly from FSlateAccessibleWidget, and optionally inherit from other IAccessible interfaces to
  * provide more functionality.
  */
-class SLATECORE_API FSlateAccessibleWidget : public IAccessibleWidget
+class FSlateAccessibleWidget : public IAccessibleWidget
 {
 	friend class FSlateAccessibleMessageHandler;
 public:
-	FSlateAccessibleWidget(TWeakPtr<SWidget> InWidget, EAccessibleWidgetType InWidgetType = EAccessibleWidgetType::Unknown);
-	virtual ~FSlateAccessibleWidget();
+	SLATECORE_API FSlateAccessibleWidget(TWeakPtr<SWidget> InWidget, EAccessibleWidgetType InWidgetType = EAccessibleWidgetType::Unknown);
+	SLATECORE_API virtual ~FSlateAccessibleWidget();
 
 	// IAccessibleWidget
-	virtual AccessibleWidgetId GetId() const override final;
-	virtual bool IsValid() const override final;
-	virtual TSharedPtr<IAccessibleWidget> GetWindow() const override final;
-	virtual FBox2D GetBounds() const override final;
-	virtual TSharedPtr<IAccessibleWidget> GetParent() override final;
-	virtual TSharedPtr<IAccessibleWidget> GetNextSibling() override final;
-	virtual TSharedPtr<IAccessibleWidget> GetPreviousSibling() override final;
-	virtual TSharedPtr<IAccessibleWidget> GetNextWidgetInHierarchy() override final;
-	virtual TSharedPtr<IAccessibleWidget> GetPreviousWidgetInHierarchy() override final;
-	virtual TSharedPtr<IAccessibleWidget> GetChildAt(int32 Index) override final;
-	virtual int32 GetNumberOfChildren() override final;
-	virtual FString GetClassName() const override final;
-	virtual bool IsEnabled() const override final;
-	virtual bool IsHidden() const override final;
-	virtual bool SupportsFocus() const override final;
-	virtual bool SupportsAccessibleFocus() const override final;
-	virtual bool CanCurrentlyAcceptAccessibleFocus() const override final;
-	virtual bool HasUserFocus(const FAccessibleUserIndex UserIndex) const override final;
-	virtual bool SetUserFocus(const FAccessibleUserIndex UserIndex) override final;
+	SLATECORE_API virtual AccessibleWidgetId GetId() const override final;
+	SLATECORE_API virtual bool IsValid() const override final;
+	SLATECORE_API virtual TSharedPtr<IAccessibleWidget> GetWindow() const override final;
+	SLATECORE_API virtual FBox2D GetBounds() const override final;
+	SLATECORE_API virtual TSharedPtr<IAccessibleWidget> GetParent() override final;
+	SLATECORE_API virtual TSharedPtr<IAccessibleWidget> GetNextSibling() override final;
+	SLATECORE_API virtual TSharedPtr<IAccessibleWidget> GetPreviousSibling() override final;
+	SLATECORE_API virtual TSharedPtr<IAccessibleWidget> GetNextWidgetInHierarchy() override final;
+	SLATECORE_API virtual TSharedPtr<IAccessibleWidget> GetPreviousWidgetInHierarchy() override final;
+	SLATECORE_API virtual TSharedPtr<IAccessibleWidget> GetChildAt(int32 Index) override final;
+	SLATECORE_API virtual int32 GetNumberOfChildren() override final;
+	SLATECORE_API virtual FString GetClassName() const override final;
+	SLATECORE_API virtual bool IsEnabled() const override final;
+	SLATECORE_API virtual bool IsHidden() const override final;
+	SLATECORE_API virtual bool SupportsFocus() const override final;
+	SLATECORE_API virtual bool SupportsAccessibleFocus() const override final;
+	SLATECORE_API virtual bool CanCurrentlyAcceptAccessibleFocus() const override final;
+	SLATECORE_API virtual bool HasUserFocus(const FAccessibleUserIndex UserIndex) const override final;
+	SLATECORE_API virtual bool SetUserFocus(const FAccessibleUserIndex UserIndex) override final;
 
 	virtual EAccessibleWidgetType GetWidgetType() const override { return WidgetType; }
-	virtual FString GetWidgetName() const override;
-	virtual FString GetHelpText() const override;
+	SLATECORE_API virtual FString GetWidgetName() const override;
+	SLATECORE_API virtual FString GetHelpText() const override;
 	// ~
 
 	/**
@@ -53,7 +53,7 @@ public:
 	 *
 	 * @param NewParent The widget to assign as the new parent widget.
 	 */
-	void UpdateParent(TSharedPtr<IAccessibleWidget> NewParent);
+	SLATECORE_API void UpdateParent(TSharedPtr<IAccessibleWidget> NewParent);
 
 protected:
 	/** The underlying Slate widget backing this accessible widget. */
@@ -77,7 +77,7 @@ private:
 	 *
 	 * @return The parent SWindow for the Slate widget referenced by this accessible widget.
 	 */
-	TSharedPtr<SWindow> GetSlateWindow() const;
+	SLATECORE_API TSharedPtr<SWindow> GetSlateWindow() const;
 };
 
 // SWindow
@@ -108,7 +108,7 @@ public:
 // ~
 
 // SImage
-class SLATECORE_API FSlateAccessibleImage
+class FSlateAccessibleImage
 	: public FSlateAccessibleWidget
 {
 public:
@@ -116,7 +116,7 @@ public:
 	virtual ~FSlateAccessibleImage() {}
 
 	// IAccessibleWidget
-	virtual FString GetHelpText() const override;
+	SLATECORE_API virtual FString GetHelpText() const override;
 	// ~
 };
 // ~
diff --git a/Engine/Source/Runtime/SlateCore/Public/Widgets/Accessibility/SlateWidgetTracker.h b/Engine/Source/Runtime/SlateCore/Public/Widgets/Accessibility/SlateWidgetTracker.h
--- a/Engine/Source/Runtime/SlateCore/Public/Widgets/Accessibility/SlateWidgetTracker.h
+++ b/Engine/Source/Runtime/SlateCore/Public/Widgets/Accessibility/SlateWidgetTracker.h
@@ -19,27 +19,27 @@ DECLARE_EVENT_ThreeParams(FSlateWidgetTracker, FTrackedWidgetsChangedEvent, cons
  * The Slate Widget Tracker tracks widgets by tags and notifies listeners when widgets of a certain tag are added or removed.
  * Widgets can be added to the tracker using the FTrackedMetaData.
  */
-class SLATECORE_API FSlateWidgetTracker
+class FSlateWidgetTracker
 {
 
 public:
 
 	/** Get the Singleton instance of the Slate Widget Tracker. */
-	static FSlateWidgetTracker& Get();
+	static SLATECORE_API FSlateWidgetTracker& Get();
 
 	/** Return true if Slate.EnableSlateWidgetTracker was set to true via config files or C++. */
-	bool IsEnabled() const;
+	SLATECORE_API bool IsEnabled() const;
 
 	/** Adds a tracked Widget for the specified tags. */
-	void AddTrackedWidget(const SWidget* WidgetToTrack, const TArray<FName>& Tags);
+	SLATECORE_API void AddTrackedWidget(const SWidget* WidgetToTrack, const TArray<FName>& Tags);
 	/** Adds a tracked Widget for the specified tag. */
-	void AddTrackedWidget(const SWidget* WidgetToTrack, FName Tag);
+	SLATECORE_API void AddTrackedWidget(const SWidget* WidgetToTrack, FName Tag);
 
 	/** Removes a tracked Widget for all tags. */
-	void RemoveTrackedWidget(const SWidget* WidgetToStopTracking);
+	SLATECORE_API void RemoveTrackedWidget(const SWidget* WidgetToStopTracking);
 	
 	/** Returns an event that gets called when widgets tagged with the passed tag are added or removed from the tracker. */
-	FTrackedWidgetsChangedEvent& OnTrackedWidgetsChanged(const FName& Tag);
+	SLATECORE_API FTrackedWidgetsChangedEvent& OnTrackedWidgetsChanged(const FName& Tag);
 
 	/** Calls Predicate for each of the tracked widgets of the specified tag. */
 	template<class Predicate>
diff --git a/Engine/Source/Runtime/SlateCore/Public/Widgets/Images/SImage.h b/Engine/Source/Runtime/SlateCore/Public/Widgets/Images/SImage.h
--- a/Engine/Source/Runtime/SlateCore/Public/Widgets/Images/SImage.h
+++ b/Engine/Source/Runtime/SlateCore/Public/Widgets/Images/SImage.h
@@ -17,10 +17,10 @@ class FSlateWindowElementList;
 /**
  * Implements a widget that displays an image with a desired width and height.
  */
-class SLATECORE_API SImage
+class SImage
 	: public SLeafWidget
 {
-	SLATE_DECLARE_WIDGET(SImage, SLeafWidget)
+	SLATE_DECLARE_WIDGET_API(SImage, SLeafWidget, SLATECORE_API)
 
 public:
 	SLATE_BEGIN_ARGS( SImage )
@@ -46,46 +46,46 @@ public:
 	SLATE_END_ARGS()
 
 	/** Constructor */
-	SImage();
+	SLATECORE_API SImage();
 
 	/**
 	 * Construct this widget
 	 *
 	 * @param	InArgs	The declaration data for this widget
 	 */
-	void Construct( const FArguments& InArgs );
+	SLATECORE_API void Construct( const FArguments& InArgs );
 
 public:
 
 	/** Set the ColorAndOpacity attribute */
-	void SetColorAndOpacity(TAttribute<FSlateColor> InColorAndOpacity);
+	SLATECORE_API void SetColorAndOpacity(TAttribute<FSlateColor> InColorAndOpacity);
 
 	/** See the ColorAndOpacity attribute */
-	void SetColorAndOpacity( FLinearColor InColorAndOpacity );
+	SLATECORE_API void SetColorAndOpacity( FLinearColor InColorAndOpacity );
 
 	/** Set the Image attribute */
-	void SetImage(TAttribute<const FSlateBrush*> InImage);
+	SLATECORE_API void SetImage(TAttribute<const FSlateBrush*> InImage);
 
 	/** Invalidate the Image */
-	void InvalidateImage();
+	SLATECORE_API void InvalidateImage();
 
 	/** Set SizeOverride attribute */
-	void SetDesiredSizeOverride(TAttribute<TOptional<FVector2D>> InDesiredSizeOverride);
+	SLATECORE_API void SetDesiredSizeOverride(TAttribute<TOptional<FVector2D>> InDesiredSizeOverride);
 
 	/** Set FlipForRightToLeftFlowDirection */
-	void FlipForRightToLeftFlowDirection(bool InbFlipForRightToLeftFlowDirection);
+	SLATECORE_API void FlipForRightToLeftFlowDirection(bool InbFlipForRightToLeftFlowDirection);
 
 public:
 
 	// SWidget overrides
-	virtual int32 OnPaint( const FPaintArgs& Args, const FGeometry& AllottedGeometry, const FSlateRect& MyCullingRect, FSlateWindowElementList& OutDrawElements, int32 LayerId, const FWidgetStyle& InWidgetStyle, bool bParentEnabled ) const override;
+	SLATECORE_API virtual int32 OnPaint( const FPaintArgs& Args, const FGeometry& AllottedGeometry, const FSlateRect& MyCullingRect, FSlateWindowElementList& OutDrawElements, int32 LayerId, const FWidgetStyle& InWidgetStyle, bool bParentEnabled ) const override;
 #if WITH_ACCESSIBILITY
-	virtual TSharedRef<FSlateAccessibleWidget> CreateAccessibleWidget() override;
+	SLATECORE_API virtual TSharedRef<FSlateAccessibleWidget> CreateAccessibleWidget() override;
 #endif
 
 protected:
 	// Begin SWidget overrides.
-	virtual FVector2D ComputeDesiredSize(float) const override;
+	SLATECORE_API virtual FVector2D ComputeDesiredSize(float) const override;
 	// End SWidget overrides.
 
 	/** @return an attribute reference of Image */
diff --git a/Engine/Source/Runtime/SlateCore/Public/Widgets/Images/SLayeredImage.h b/Engine/Source/Runtime/SlateCore/Public/Widgets/Images/SLayeredImage.h
--- a/Engine/Source/Runtime/SlateCore/Public/Widgets/Images/SLayeredImage.h
+++ b/Engine/Source/Runtime/SlateCore/Public/Widgets/Images/SLayeredImage.h
@@ -8,47 +8,47 @@
 struct FSlateIcon;
 
 /** A slate widget to draw an arbitrary number of images over top of each other */
-class SLATECORE_API SLayeredImage : public SImage
+class SLayeredImage : public SImage
 {
 public:
 	typedef TTuple<TAttribute<const FSlateBrush*>, TAttribute<FSlateColor> > ImageLayer;
 
 	/** Constructor that adds InLayers on top of the base image  */
-	void Construct(const FArguments& InArgs, const TArray<ImageLayer>& InLayers);
+	SLATECORE_API void Construct(const FArguments& InArgs, const TArray<ImageLayer>& InLayers);
 	
 	/** Constructor that adds InLayers on top of the base image  */
-	void Construct(const FArguments& InArgs, TArray<ImageLayer>&& InLayers);
+	SLATECORE_API void Construct(const FArguments& InArgs, TArray<ImageLayer>&& InLayers);
 
 	/** Constructor that adds a single layer on top of the base image */
-	void Construct(const FArguments& InArgs, TAttribute<const FSlateBrush*> Brush, TAttribute<FSlateColor> Color);
+	SLATECORE_API void Construct(const FArguments& InArgs, TAttribute<const FSlateBrush*> Brush, TAttribute<FSlateColor> Color);
 
 	/** Constructor that adds NumLayers blank layers on top of the image for later use */
-	void Construct(const FArguments& InArgs, int32 NumLayers = 0);
+	SLATECORE_API void Construct(const FArguments& InArgs, int32 NumLayers = 0);
 
 	/** Constructor that sets the base image and any overlay layers defined in the Slate icon */
-	void Construct(const FArguments& InArgs, const FSlateIcon& InIcon);
+	SLATECORE_API void Construct(const FArguments& InArgs, const FSlateIcon& InIcon);
 
 	//~ Begin SWidget Interface
-	virtual int32 OnPaint(const FPaintArgs& Args, const FGeometry& AllottedGeometry, const FSlateRect& MyCullingRect, FSlateWindowElementList& OutDrawElements, int32 LayerId, const FWidgetStyle& InWidgetStyle, bool bParentEnabled) const override;
+	SLATECORE_API virtual int32 OnPaint(const FPaintArgs& Args, const FGeometry& AllottedGeometry, const FSlateRect& MyCullingRect, FSlateWindowElementList& OutDrawElements, int32 LayerId, const FWidgetStyle& InWidgetStyle, bool bParentEnabled) const override;
 	//~ End SWidget Interface
 
 	/** Adds a new layer to the top of the Image given a brush and color. Use White as the default color. */
-	void AddLayer(TAttribute<const FSlateBrush*> Brush);
+	SLATECORE_API void AddLayer(TAttribute<const FSlateBrush*> Brush);
 
 	/** Adds a new layer to the top of the Image given a brush and color */
-	void AddLayer(TAttribute<const FSlateBrush*> Brush, TAttribute<FSlateColor> Color);
+	SLATECORE_API void AddLayer(TAttribute<const FSlateBrush*> Brush, TAttribute<FSlateColor> Color);
 
 	/** Sets the base image and any overlay layers defined in the Slate icon. Note: removes all layers */
-	void SetFromSlateIcon(const FSlateIcon& InIcon);
+	SLATECORE_API void SetFromSlateIcon(const FSlateIcon& InIcon);
 
 	/** Removes all layers */
-	void RemoveAllLayers();
+	SLATECORE_API void RemoveAllLayers();
 
 	/** Returns the total number of layers, including the base image. */
-	int32 GetNumLayers() const;
+	SLATECORE_API int32 GetNumLayers() const;
 
 	/** Determines whether a Layer exists at the given index */
-	bool IsValidIndex(int32 Index) const;
+	SLATECORE_API bool IsValidIndex(int32 Index) const;
 
 	/**
 	 * Gets the brush for a given layer
@@ -56,10 +56,10 @@ public:
 	 * @return Null if the layer index is invalid. Otherwise, the brush for the given layer
 	 */
 	UE_DEPRECATED(5.0, "GetLayerBrush is not accessible anymore since it's attribute value may not have been updated yet.")
-	const FSlateBrush* GetLayerBrush(int32 Index) const;
+	SLATECORE_API const FSlateBrush* GetLayerBrush(int32 Index) const;
 
 	/** Sets the brush for a given layer, if it exists */
-	void SetLayerBrush(int32 Index, TAttribute<const FSlateBrush*> Brush);
+	SLATECORE_API void SetLayerBrush(int32 Index, TAttribute<const FSlateBrush*> Brush);
 
 	/** 
 	 * Gets the Color for a given layer
@@ -67,10 +67,10 @@ public:
 	 * @return Uninitialized fuschia color if the layer index is invalid. Otherwise the color for the given layer
 	 */
 	UE_DEPRECATED(5.0, "GetLayerColor is not accessible anymore since it's attribute value may not have been updated yet.")
-	FSlateColor GetLayerColor(int32 Index) const;
+	SLATECORE_API FSlateColor GetLayerColor(int32 Index) const;
 
 	/** Sets the color for a given layer, if it exists. */
-	void SetLayerColor(int32 Index, TAttribute<FSlateColor> Color);
+	SLATECORE_API void SetLayerColor(int32 Index, TAttribute<FSlateColor> Color);
 
 private:
 	typedef TSlateManagedAttribute<const FSlateBrush*, EInvalidateWidgetReason::Paint> BrushAttributeType;
diff --git a/Engine/Source/Runtime/SlateCore/Public/Widgets/SBoxPanel.h b/Engine/Source/Runtime/SlateCore/Public/Widgets/SBoxPanel.h
--- a/Engine/Source/Runtime/SlateCore/Public/Widgets/SBoxPanel.h
+++ b/Engine/Source/Runtime/SlateCore/Public/Widgets/SBoxPanel.h
@@ -23,9 +23,9 @@ class FArrangedChildren;
 /**
  * A BoxPanel contains one child and describes how that child should be arranged on the screen.
  */
-class SLATECORE_API SBoxPanel : public SPanel
+class SBoxPanel : public SPanel
 {
-	SLATE_DECLARE_WIDGET(SBoxPanel, SPanel)
+	SLATE_DECLARE_WIDGET_API(SBoxPanel, SPanel, SLATECORE_API)
 
 protected:
 	/**
@@ -185,10 +185,10 @@ public:
 	 * @param SlotWidget The widget to match when searching through the slots
 	 * @returns The index in the children array where the slot was removed and -1 if no slot was found matching the widget
 	 */
-	int32 RemoveSlot( const TSharedRef<SWidget>& SlotWidget );
+	SLATECORE_API int32 RemoveSlot( const TSharedRef<SWidget>& SlotWidget );
 
 	/** Removes all children from the box. */
-	void ClearChildren();
+	SLATECORE_API void ClearChildren();
 
 	/** @return the number of slots. */
 	int32 NumSlots() const { return Children.Num(); }
@@ -198,20 +198,20 @@ public:
 
 	//~ Begin SWidget overrides.
 public:
-	virtual void OnArrangeChildren(const FGeometry& AllottedGeometry, FArrangedChildren& ArrangedChildren) const override;
-	virtual FChildren* GetChildren() override;
+	SLATECORE_API virtual void OnArrangeChildren(const FGeometry& AllottedGeometry, FArrangedChildren& ArrangedChildren) const override;
+	SLATECORE_API virtual FChildren* GetChildren() override;
 
 protected:
-	virtual FVector2D ComputeDesiredSize(float) const override;
+	SLATECORE_API virtual FVector2D ComputeDesiredSize(float) const override;
 	//~ End SWidget overrides.
 
 protected:
-	SBoxPanel();
+	SLATECORE_API SBoxPanel();
 	/** A Box Panel's orientation. */
-	SBoxPanel(EOrientation InOrientation);
+	SLATECORE_API SBoxPanel(EOrientation InOrientation);
 
 	/** Set the orientation of the Box. It will do a full invalidation of the widget. */
-	void SetOrientation(EOrientation InOrientation);
+	SLATECORE_API void SetOrientation(EOrientation InOrientation);
 	EOrientation GetOrientation() const
 	{
 		return Orientation;
@@ -227,9 +227,9 @@ protected:
 
 
 /** A Horizontal Box Panel. See SBoxPanel for more info. */
-class SLATECORE_API SHorizontalBox : public SBoxPanel
+class SHorizontalBox : public SBoxPanel
 {
-	SLATE_DECLARE_WIDGET(SHorizontalBox, SBoxPanel)
+	SLATE_DECLARE_WIDGET_API(SHorizontalBox, SBoxPanel, SLATECORE_API)
 public:
 	class FSlot : public SBoxPanel::TSlot<FSlot>
 	{
@@ -311,8 +311,8 @@ public:
 		return FScopedWidgetSlotArguments(MakeUnique<FSlot>(), this->Children, Index);
 	}
 
-	FSlot& GetSlot(int32 SlotIndex);
-	const FSlot& GetSlot(int32 SlotIndex) const;
+	SLATECORE_API FSlot& GetSlot(int32 SlotIndex);
+	SLATECORE_API const FSlot& GetSlot(int32 SlotIndex) const;
 
 	FORCENOINLINE SHorizontalBox()
 		: SBoxPanel( Orient_Horizontal )
@@ -326,13 +326,13 @@ public:
 	 *
 	 * @param	InArgs	The declaration data for this widget
 	 */
-	void Construct( const FArguments& InArgs );
+	SLATECORE_API void Construct( const FArguments& InArgs );
 };
 
 /** A Vertical Box Panel. See SBoxPanel for more info. */
-class SLATECORE_API SVerticalBox : public SBoxPanel
+class SVerticalBox : public SBoxPanel
 {
-	SLATE_DECLARE_WIDGET(SVerticalBox, SBoxPanel)
+	SLATE_DECLARE_WIDGET_API(SVerticalBox, SBoxPanel, SLATECORE_API)
 public:
 	class FSlot : public SBoxPanel::TSlot<FSlot>
 	{
@@ -415,8 +415,8 @@ public:
 		return FScopedWidgetSlotArguments(MakeUnique<FSlot>(), this->Children, Index);
 	}
 
-	FSlot& GetSlot(int32 SlotIndex);
-	const FSlot& GetSlot(int32 SlotIndex) const;
+	SLATECORE_API FSlot& GetSlot(int32 SlotIndex);
+	SLATECORE_API const FSlot& GetSlot(int32 SlotIndex) const;
 
 	FORCENOINLINE SVerticalBox()
 		: SBoxPanel( Orient_Vertical )
@@ -430,13 +430,13 @@ public:
 	 *
 	 * @param	InArgs	The declaration data for this widget
 	 */
-	void Construct( const FArguments& InArgs );
+	SLATECORE_API void Construct( const FArguments& InArgs );
 };
 
 /** A Stack Box Panel that stack vertically or horizontally. See SBoxPanel for more info. */
-class SLATECORE_API SStackBox : public SBoxPanel
+class SStackBox : public SBoxPanel
 {
-	SLATE_DECLARE_WIDGET(SStackBox, SBoxPanel)
+	SLATE_DECLARE_WIDGET_API(SStackBox, SBoxPanel, SLATECORE_API)
 public:
 	class FSlot : public SBoxPanel::TSlot<FSlot>
 	{
@@ -493,8 +493,8 @@ public:
 		return FScopedWidgetSlotArguments(MakeUnique<FSlot>(), this->Children, Index);
 	}
 
-	FSlot& GetSlot(int32 SlotIndex);
-	const FSlot& GetSlot(int32 SlotIndex) const;
+	SLATECORE_API FSlot& GetSlot(int32 SlotIndex);
+	SLATECORE_API const FSlot& GetSlot(int32 SlotIndex) const;
 
 	FORCENOINLINE SStackBox()
 	{
@@ -507,7 +507,7 @@ public:
 	 *
 	 * @param	InArgs	The declaration data for this widget
 	 */
-	void Construct(const FArguments& InArgs);
+	SLATECORE_API void Construct(const FArguments& InArgs);
 
 	using SBoxPanel::SetOrientation;
 	using SBoxPanel::GetOrientation;
@@ -527,7 +527,7 @@ public:
 };
 
 /** A Vertical Box Panel. See SBoxPanel for more info. */
-class SLATECORE_API SDragAndDropVerticalBox : public SVerticalBox
+class SDragAndDropVerticalBox : public SVerticalBox
 {
 public:
 	/**
@@ -578,22 +578,22 @@ public:
 		SLATE_EVENT(FOnDragAndDropVerticalBoxDrop, OnDrop)
 	SLATE_END_ARGS()
 
-	void Construct(const FArguments& InArgs);
+	SLATECORE_API void Construct(const FArguments& InArgs);
 
 	/** Set the Drop indicators */
 	SDragAndDropVerticalBox& SetDropIndicator_Above(const FSlateBrush& InValue) { DropIndicator_Above = InValue; return *this; }
 	SDragAndDropVerticalBox& SetDropIndicator_Below(const FSlateBrush& InValue) { DropIndicator_Below = InValue; return *this; }
 
 	/** Drag detection and handling */
-	virtual FReply OnMouseButtonDown(const FGeometry& MyGeometry, const FPointerEvent& MouseEvent) override;
+	SLATECORE_API virtual FReply OnMouseButtonDown(const FGeometry& MyGeometry, const FPointerEvent& MouseEvent) override;
 		
-	virtual FReply OnDragDetected(const FGeometry& MyGeometry, const FPointerEvent& MouseEvent) override;
-	virtual void OnDragEnter(FGeometry const& MyGeometry, FDragDropEvent const& DragDropEvent) override;
-	virtual void OnDragLeave(FDragDropEvent const& DragDropEvent) override;
-	virtual FReply OnDragOver(const FGeometry& MyGeometry, const FDragDropEvent& DragDropEvent) override;
-	virtual FReply OnDrop(const FGeometry& MyGeometry, const FDragDropEvent& DragDropEvent) override;
+	SLATECORE_API virtual FReply OnDragDetected(const FGeometry& MyGeometry, const FPointerEvent& MouseEvent) override;
+	SLATECORE_API virtual void OnDragEnter(FGeometry const& MyGeometry, FDragDropEvent const& DragDropEvent) override;
+	SLATECORE_API virtual void OnDragLeave(FDragDropEvent const& DragDropEvent) override;
+	SLATECORE_API virtual FReply OnDragOver(const FGeometry& MyGeometry, const FDragDropEvent& DragDropEvent) override;
+	SLATECORE_API virtual FReply OnDrop(const FGeometry& MyGeometry, const FDragDropEvent& DragDropEvent) override;
 
-	virtual int32 OnPaint(const FPaintArgs& Args, const FGeometry& AllottedGeometry, const FSlateRect& MyCullingRect, FSlateWindowElementList& OutDrawElements, int32 LayerId, const FWidgetStyle& InWidgetStyle, bool bParentEnabled) const override;
+	SLATECORE_API virtual int32 OnPaint(const FPaintArgs& Args, const FGeometry& AllottedGeometry, const FSlateRect& MyCullingRect, FSlateWindowElementList& OutDrawElements, int32 LayerId, const FWidgetStyle& InWidgetStyle, bool bParentEnabled) const override;
 
 private:
 	/** @see SDragAndDropVerticalBox's OnCanAcceptDrop event */
@@ -626,5 +626,5 @@ private:
 	int32 CurrentDragOverSlotIndex;
 
 	/** @return the zone (above, below) based on where the user is hovering over */
-	EItemDropZone ZoneFromPointerPosition(FVector2f LocalPointerPos, const FGeometry& CurrentGeometry, const FGeometry& StartGeometry) const;
+	SLATECORE_API EItemDropZone ZoneFromPointerPosition(FVector2f LocalPointerPos, const FGeometry& CurrentGeometry, const FGeometry& StartGeometry) const;
 };
diff --git a/Engine/Source/Runtime/SlateCore/Public/Widgets/SCompoundWidget.h b/Engine/Source/Runtime/SlateCore/Public/Widgets/SCompoundWidget.h
--- a/Engine/Source/Runtime/SlateCore/Public/Widgets/SCompoundWidget.h
+++ b/Engine/Source/Runtime/SlateCore/Public/Widgets/SCompoundWidget.h
@@ -18,9 +18,9 @@ class FSlateWindowElementList;
  * A CompoundWidget is the base from which most non-primitive widgets should be built.
  * CompoundWidgets have a protected member named ChildSlot.
  */
-class SLATECORE_API SCompoundWidget : public SWidget
+class SCompoundWidget : public SWidget
 {
-	SLATE_DECLARE_WIDGET(SCompoundWidget, SWidget)
+	SLATE_DECLARE_WIDGET_API(SCompoundWidget, SWidget, SLATECORE_API)
 
 public:
 
@@ -75,17 +75,17 @@ public:
 public:
 
 	// SWidgetOverrides
-	virtual int32 OnPaint( const FPaintArgs& Args, const FGeometry& AllottedGeometry, const FSlateRect& MyCullingRect, FSlateWindowElementList& OutDrawElements, int32 LayerId, const FWidgetStyle& InWidgetStyle, bool bParentEnabled ) const override;
-	virtual FChildren* GetChildren() override;
-	virtual void OnArrangeChildren( const FGeometry& AllottedGeometry, FArrangedChildren& ArrangedChildren ) const override;
-	virtual FSlateColor GetForegroundColor() const override;
+	SLATECORE_API virtual int32 OnPaint( const FPaintArgs& Args, const FGeometry& AllottedGeometry, const FSlateRect& MyCullingRect, FSlateWindowElementList& OutDrawElements, int32 LayerId, const FWidgetStyle& InWidgetStyle, bool bParentEnabled ) const override;
+	SLATECORE_API virtual FChildren* GetChildren() override;
+	SLATECORE_API virtual void OnArrangeChildren( const FGeometry& AllottedGeometry, FArrangedChildren& ArrangedChildren ) const override;
+	SLATECORE_API virtual FSlateColor GetForegroundColor() const override;
 
 public:
-	virtual void SetVisibility( TAttribute<EVisibility> InVisibility ) override final;
+	SLATECORE_API virtual void SetVisibility( TAttribute<EVisibility> InVisibility ) override final;
 
 protected:
 	// Begin SWidget overrides.
-	virtual FVector2D ComputeDesiredSize(float) const override;
+	SLATECORE_API virtual FVector2D ComputeDesiredSize(float) const override;
 	// End SWidget overrides.
 
 	/** @return an attribute reference of ContentScale */
@@ -100,7 +100,7 @@ protected:
 protected:
 
 	/** Disallow public construction */
-	SCompoundWidget();
+	SLATECORE_API SCompoundWidget();
 
 	struct FCompoundWidgetOneChildSlot : ::TSingleWidgetChildrenWithBasicLayoutSlot<EInvalidateWidgetReason::None>
 	{
diff --git a/Engine/Source/Runtime/SlateCore/Public/Widgets/SLeafWidget.h b/Engine/Source/Runtime/SlateCore/Public/Widgets/SLeafWidget.h
--- a/Engine/Source/Runtime/SlateCore/Public/Widgets/SLeafWidget.h
+++ b/Engine/Source/Runtime/SlateCore/Public/Widgets/SLeafWidget.h
@@ -18,7 +18,7 @@ class FSlateWindowElementList;
  * A LeafWidget is a Widget that has no slots for children.
  * LeafWidgets are usually intended as building blocks for aggregate widgets.
  */
-class SLATECORE_API SLeafWidget
+class SLeafWidget
 	: public SWidget
 {
 public:
@@ -27,7 +27,7 @@ public:
 		bCanHaveChildren = false;
 	}
 
-	virtual void SetVisibility( TAttribute<EVisibility> InVisibility ) override final;
+	SLATECORE_API virtual void SetVisibility( TAttribute<EVisibility> InVisibility ) override final;
 
 private:
 
@@ -59,7 +59,7 @@ private:
 	 *
 	 * Leaf widgets never have children.
 	 */
-	virtual FChildren* GetChildren() override;
+	SLATECORE_API virtual FChildren* GetChildren() override;
 
-	virtual void OnArrangeChildren( const FGeometry& AllottedGeometry, FArrangedChildren& ArrangedChildren ) const override;
+	SLATECORE_API virtual void OnArrangeChildren( const FGeometry& AllottedGeometry, FArrangedChildren& ArrangedChildren ) const override;
 };
diff --git a/Engine/Source/Runtime/SlateCore/Public/Widgets/SNullWidget.h b/Engine/Source/Runtime/SlateCore/Public/Widgets/SNullWidget.h
--- a/Engine/Source/Runtime/SlateCore/Public/Widgets/SNullWidget.h
+++ b/Engine/Source/Runtime/SlateCore/Public/Widgets/SNullWidget.h
@@ -10,7 +10,7 @@
  * Widgets that support slots, such as SOverlay and SHorizontalBox should initialize
  * their slots' child widgets to SNullWidget if no user defined widget was provided.
  */
-class SLATECORE_API SNullWidget
+class SNullWidget
 {
 public:
 
@@ -19,5 +19,5 @@ public:
 	 *
 	 * @return The widget.
 	 */
-	static TSharedRef<class SWidget> NullWidget;
+	static SLATECORE_API TSharedRef<class SWidget> NullWidget;
 };
diff --git a/Engine/Source/Runtime/SlateCore/Public/Widgets/SOverlay.h b/Engine/Source/Runtime/SlateCore/Public/Widgets/SOverlay.h
--- a/Engine/Source/Runtime/SlateCore/Public/Widgets/SOverlay.h
+++ b/Engine/Source/Runtime/SlateCore/Public/Widgets/SOverlay.h
@@ -31,9 +31,9 @@ class FSlateWindowElementList;
  *
  *		Note that SWidget3 will be drawn on top of SWidget2 and SWidget1.
  */
-class SLATECORE_API SOverlay : public SPanel
+class SOverlay : public SPanel
 {
-	SLATE_DECLARE_WIDGET(SOverlay, SPanel)
+	SLATE_DECLARE_WIDGET_API(SOverlay, SPanel, SLATECORE_API)
 public:	
 
 	/** A slot that support alignment of content and padding and z-order */
@@ -76,57 +76,57 @@ public:
 		SLATE_SLOT_ARGUMENT( SOverlay::FOverlaySlot, Slots )
 	SLATE_END_ARGS()
 
-	SOverlay();
+	SLATECORE_API SOverlay();
 
 	/**
 	 * Construct this widget.
 	 *
 	 * @param	InArgs	The declaration data for this widget
 	 */
-	void Construct( const FArguments& InArgs );
+	SLATECORE_API void Construct( const FArguments& InArgs );
 
 	/** Returns the number of child widgets */
-	int32 GetNumWidgets() const;
+	SLATECORE_API int32 GetNumWidgets() const;
 
 	/**
 	 * Removes a widget from this overlay.
 	 *
 	 * @param	Widget	The widget content to remove
 	 */
-	bool RemoveSlot( TSharedRef< SWidget > Widget );
+	SLATECORE_API bool RemoveSlot( TSharedRef< SWidget > Widget );
 
 	using FScopedWidgetSlotArguments = TPanelChildren<FOverlaySlot>::FScopedWidgetSlotArguments;
 	/** Adds a slot at the specified location (ignores Z-order) */
-	FScopedWidgetSlotArguments AddSlot(int32 ZOrder=INDEX_NONE);
+	SLATECORE_API FScopedWidgetSlotArguments AddSlot(int32 ZOrder=INDEX_NONE);
 
 	/** Returns true if there is a child slot with the specified z-order */
-	bool HasSlotWithZOrder(int32 ZOrder) const;
+	SLATECORE_API bool HasSlotWithZOrder(int32 ZOrder) const;
 
 	/** Removes a slot at the specified location */
-	void RemoveSlot(int32 ZOrder=INDEX_NONE);
+	SLATECORE_API void RemoveSlot(int32 ZOrder=INDEX_NONE);
 
 	/** Removes all children from the overlay */
-	void ClearChildren();
+	SLATECORE_API void ClearChildren();
 
 	/**
 	 * @param ZOrder Default z-order of the slot
 	 * @return a new slot. Slots contain children for SOverlay
 	 */
-	static FOverlaySlot::FSlotArguments Slot(int32 ZOrder = 0);
+	static SLATECORE_API FOverlaySlot::FSlotArguments Slot(int32 ZOrder = 0);
 
 	//~ Begin of SWidget interface
-	virtual void OnArrangeChildren( const FGeometry& AllottedGeometry, FArrangedChildren& ArrangedChildren ) const override;
-	virtual FChildren* GetChildren() override;
-	virtual int32 OnPaint( const FPaintArgs& Args, const FGeometry& AllottedGeometry, const FSlateRect& MyCullingRect, FSlateWindowElementList& OutDrawElements, int32 LayerId, const FWidgetStyle& InWidgetStyle, bool bParentEnabled ) const override;
+	SLATECORE_API virtual void OnArrangeChildren( const FGeometry& AllottedGeometry, FArrangedChildren& ArrangedChildren ) const override;
+	SLATECORE_API virtual FChildren* GetChildren() override;
+	SLATECORE_API virtual int32 OnPaint( const FPaintArgs& Args, const FGeometry& AllottedGeometry, const FSlateRect& MyCullingRect, FSlateWindowElementList& OutDrawElements, int32 LayerId, const FWidgetStyle& InWidgetStyle, bool bParentEnabled ) const override;
 	//~ End of SWidget interface
 
 protected:
 	//~ Begin SWidget overrides.
-	virtual FVector2D ComputeDesiredSize(float) const override;
+	SLATECORE_API virtual FVector2D ComputeDesiredSize(float) const override;
 	//~ End SWidget overrides.
 
 	/** Returns the index of the child slot at the specified z-order */
-	int32 GetChildIndexByZOrder(int32 ZOrder) const;
+	SLATECORE_API int32 GetChildIndexByZOrder(int32 ZOrder) const;
 
 protected:
 	/** The SOverlay's slots; each slot contains a child widget. */
diff --git a/Engine/Source/Runtime/SlateCore/Public/Widgets/SPanel.h b/Engine/Source/Runtime/SlateCore/Public/Widgets/SPanel.h
--- a/Engine/Source/Runtime/SlateCore/Public/Widgets/SPanel.h
+++ b/Engine/Source/Runtime/SlateCore/Public/Widgets/SPanel.h
@@ -19,7 +19,7 @@ class FSlateWindowElementList;
  * respect to its parent (i.e. the Panel) and its peers Widgets (i.e. the Panel's other children.)
  * For a simple example see StackPanel.
  */
-class SLATECORE_API SPanel
+class SPanel
 	: public SWidget
 {
 public:	
@@ -62,14 +62,14 @@ public:
 
 	// SWidget overrides
 
-	virtual int32 OnPaint( const FPaintArgs& Args, const FGeometry& AllottedGeometry, const FSlateRect& MyCullingRect, FSlateWindowElementList& OutDrawElements, int32 LayerId, const FWidgetStyle& InWidgetStyle, bool bParentEnabled ) const override;
+	SLATECORE_API virtual int32 OnPaint( const FPaintArgs& Args, const FGeometry& AllottedGeometry, const FSlateRect& MyCullingRect, FSlateWindowElementList& OutDrawElements, int32 LayerId, const FWidgetStyle& InWidgetStyle, bool bParentEnabled ) const override;
 
 protected:
 
 	/**
 	 * Just like OnPaint, but takes already arranged children. Can be handy for writing custom SPanels.
 	 */
-	int32 PaintArrangedChildren( const FPaintArgs& Args, const FArrangedChildren& ArrangedChildren, const FGeometry& AllottedGeometry, const FSlateRect& MyCullingRect, FSlateWindowElementList& OutDrawElements, int32 LayerId, const FWidgetStyle& InWidgetStyle, bool bParentEnabled  ) const;
+	SLATECORE_API int32 PaintArrangedChildren( const FPaintArgs& Args, const FArrangedChildren& ArrangedChildren, const FGeometry& AllottedGeometry, const FSlateRect& MyCullingRect, FSlateWindowElementList& OutDrawElements, int32 LayerId, const FWidgetStyle& InWidgetStyle, bool bParentEnabled  ) const;
 	
 protected:
 
@@ -77,5 +77,5 @@ protected:
 	SPanel( ) { }
 
 public:
-	virtual void SetVisibility( TAttribute<EVisibility> InVisibility ) override final;
+	SLATECORE_API virtual void SetVisibility( TAttribute<EVisibility> InVisibility ) override final;
 };
diff --git a/Engine/Source/Runtime/SlateCore/Public/Widgets/SWidget.h b/Engine/Source/Runtime/SlateCore/Public/Widgets/SWidget.h
--- a/Engine/Source/Runtime/SlateCore/Public/Widgets/SWidget.h
+++ b/Engine/Source/Runtime/SlateCore/Public/Widgets/SWidget.h
@@ -150,11 +150,11 @@ class IToolTip;
  *   Some events are able to reply to the system by returning an FReply, FCursorReply, or similar
  *   object. 
  */
-class SLATECORE_API SWidget
+class SWidget
 	: public FSlateControlledConstruction,
 	public TSharedFromThis<SWidget>		// Enables 'this->AsShared()'
 {
-	SLATE_DECLARE_WIDGET(SWidget, FSlateControlledConstruction)
+	SLATE_DECLARE_WIDGET_API(SWidget, FSlateControlledConstruction, SLATECORE_API)
 
 	friend class FWidgetProxy;
 	friend class FSlateAttributeMetaData;
@@ -265,7 +265,7 @@ public:
 
 	/** Construct a SWidget based on initial parameters. */
 	 UE_DEPRECATED(4.27, "SWidget::Construct should not be called directly. Use SNew or SAssignNew to create a SWidget")
-	void Construct(
+	SLATECORE_API void Construct(
 		const TAttribute<FText>& InToolTipText,
 		const TSharedPtr<IToolTip>& InToolTip,
 		const TAttribute< TOptional<EMouseCursor::Type> >& InCursor,
@@ -282,7 +282,7 @@ public:
 		const TArray<TSharedRef<ISlateMetaData>>& InMetaData);
 
 	UE_DEPRECATED(4.27, "SWidget::SWidgetConstruct should not be called directly. Use SNew or SAssignNew to create a SWidget")
-	void SWidgetConstruct(const TAttribute<FText>& InToolTipText,
+	SLATECORE_API void SWidgetConstruct(const TAttribute<FText>& InToolTipText,
 		const TSharedPtr<IToolTip>& InToolTip,
 		const TAttribute< TOptional<EMouseCursor::Type> >& InCursor,
 		const TAttribute<bool>& InEnabledState,
@@ -316,7 +316,7 @@ public:
 	 * @param bParentEnabled	True if the parent of this widget is enabled.
 	 * @return The maximum layer ID attained by this widget or any of its children.
 	 */
-	int32 Paint(const FPaintArgs& Args, const FGeometry& AllottedGeometry, const FSlateRect& MyCullingRect, FSlateWindowElementList& OutDrawElements, int32 LayerId, const FWidgetStyle& InWidgetStyle, bool bParentEnabled) const;
+	SLATECORE_API int32 Paint(const FPaintArgs& Args, const FGeometry& AllottedGeometry, const FSlateRect& MyCullingRect, FSlateWindowElementList& OutDrawElements, int32 LayerId, const FWidgetStyle& InWidgetStyle, bool bParentEnabled) const;
 
 	/**
 	 * Ticks this widget with Geometry.  Override in derived classes, but always call the parent implementation.
@@ -325,7 +325,7 @@ public:
 	 * @param  InCurrentTime  Current absolute real time
 	 * @param  InDeltaTime  Real time passed since last tick
 	 */
-	virtual void Tick(const FGeometry& AllottedGeometry, const double InCurrentTime, const float InDeltaTime);
+	SLATECORE_API virtual void Tick(const FGeometry& AllottedGeometry, const double InCurrentTime, const float InDeltaTime);
 
 	//
 	// KEY INPUT
@@ -338,17 +338,17 @@ public:
 	 * @param  InFocusEvent  The FocusEvent
 	 * @return  Returns whether the event was handled, along with other possible actions
 	 */
-	virtual FReply OnFocusReceived(const FGeometry& MyGeometry, const FFocusEvent& InFocusEvent);
+	SLATECORE_API virtual FReply OnFocusReceived(const FGeometry& MyGeometry, const FFocusEvent& InFocusEvent);
 
 	/**
 	 * Called when this widget loses focus.  This event does not bubble.
 	 *
 	 * @param InFocusEvent The FocusEvent
 	 */
-	virtual void OnFocusLost(const FFocusEvent& InFocusEvent);
+	SLATECORE_API virtual void OnFocusLost(const FFocusEvent& InFocusEvent);
 
 	/** Called whenever a focus path is changing on all the widgets within the old and new focus paths */
-	virtual void OnFocusChanging(const FWeakWidgetPath& PreviousFocusPath, const FWidgetPath& NewWidgetPath, const FFocusEvent& InFocusEvent);
+	SLATECORE_API virtual void OnFocusChanging(const FWeakWidgetPath& PreviousFocusPath, const FWidgetPath& NewWidgetPath, const FFocusEvent& InFocusEvent);
 
 	/**
 	 * Called after a character is entered while this widget has keyboard focus
@@ -357,7 +357,7 @@ public:
 	 * @param  InCharacterEvent  Character event
 	 * @return  Returns whether the event was handled, along with other possible actions
 	 */
-	virtual FReply OnKeyChar(const FGeometry& MyGeometry, const FCharacterEvent& InCharacterEvent);
+	SLATECORE_API virtual FReply OnKeyChar(const FGeometry& MyGeometry, const FCharacterEvent& InCharacterEvent);
 
 	/**
 	 * Called after a key is pressed when this widget or a child of this widget has focus
@@ -370,7 +370,7 @@ public:
 	 * @param InKeyEvent  Key event
 	 * @return Returns whether the event was handled, along with other possible actions
 	 */
-	virtual FReply OnPreviewKeyDown(const FGeometry& MyGeometry, const FKeyEvent& InKeyEvent);
+	SLATECORE_API virtual FReply OnPreviewKeyDown(const FGeometry& MyGeometry, const FKeyEvent& InKeyEvent);
 
 	/**
 	 * Called after a key is pressed when this widget has focus (this event bubbles if not handled)
@@ -379,7 +379,7 @@ public:
 	 * @param InKeyEvent  Key event
 	 * @return Returns whether the event was handled, along with other possible actions
 	 */
-	virtual FReply OnKeyDown(const FGeometry& MyGeometry, const FKeyEvent& InKeyEvent);
+	SLATECORE_API virtual FReply OnKeyDown(const FGeometry& MyGeometry, const FKeyEvent& InKeyEvent);
 
 	/**
 	 * Called after a key is released when this widget has focus
@@ -388,7 +388,7 @@ public:
 	 * @param InKeyEvent  Key event
 	 * @return Returns whether the event was handled, along with other possible actions
 	 */
-	virtual FReply OnKeyUp(const FGeometry& MyGeometry, const FKeyEvent& InKeyEvent);
+	SLATECORE_API virtual FReply OnKeyUp(const FGeometry& MyGeometry, const FKeyEvent& InKeyEvent);
 
 	/**
 	 * Called when an analog value changes on a button that supports analog
@@ -397,7 +397,7 @@ public:
 	 * @param InAnalogInputEvent Analog input event
 	 * @return Returns whether the event was handled, along with other possible actions
 	 */
-	virtual FReply OnAnalogValueChanged(const FGeometry& MyGeometry, const FAnalogInputEvent& InAnalogInputEvent);
+	SLATECORE_API virtual FReply OnAnalogValueChanged(const FGeometry& MyGeometry, const FAnalogInputEvent& InAnalogInputEvent);
 
 	//
 	// MOUSE INPUT
@@ -410,7 +410,7 @@ public:
 	 * @param MouseEvent Information about the input event
 	 * @return Whether the event was handled along with possible requests for the system to take action.
 	 */
-	virtual FReply OnMouseButtonDown(const FGeometry& MyGeometry, const FPointerEvent& MouseEvent);
+	SLATECORE_API virtual FReply OnMouseButtonDown(const FGeometry& MyGeometry, const FPointerEvent& MouseEvent);
 
 	/**
 	 * Just like OnMouseButtonDown, but tunnels instead of bubbling.
@@ -419,7 +419,7 @@ public:
 	 * Use this event sparingly as preview events generally make UIs more
 	 * difficult to reason about.
 	 */
-	virtual FReply OnPreviewMouseButtonDown(const FGeometry& MyGeometry, const FPointerEvent& MouseEvent);
+	SLATECORE_API virtual FReply OnPreviewMouseButtonDown(const FGeometry& MyGeometry, const FPointerEvent& MouseEvent);
 
 	/**
 	 * The system calls this method to notify the widget that a mouse button was release within it. This event is bubbled.
@@ -428,7 +428,7 @@ public:
 	 * @param MouseEvent Information about the input event
 	 * @return Whether the event was handled along with possible requests for the system to take action.
 	 */
-	virtual FReply OnMouseButtonUp(const FGeometry& MyGeometry, const FPointerEvent& MouseEvent);
+	SLATECORE_API virtual FReply OnMouseButtonUp(const FGeometry& MyGeometry, const FPointerEvent& MouseEvent);
 
 	/**
 	 * The system calls this method to notify the widget that a mouse moved within it. This event is bubbled.
@@ -437,7 +437,7 @@ public:
 	 * @param MouseEvent Information about the input event
 	 * @return Whether the event was handled along with possible requests for the system to take action.
 	 */
-	virtual FReply OnMouseMove(const FGeometry& MyGeometry, const FPointerEvent& MouseEvent);
+	SLATECORE_API virtual FReply OnMouseMove(const FGeometry& MyGeometry, const FPointerEvent& MouseEvent);
 
 	/**
 	 * The system will use this event to notify a widget that the cursor has entered it. This event is uses a custom bubble strategy.
@@ -445,14 +445,14 @@ public:
 	 * @param MyGeometry The Geometry of the widget receiving the event
 	 * @param MouseEvent Information about the input event
 	 */
-	virtual void OnMouseEnter(const FGeometry& MyGeometry, const FPointerEvent& MouseEvent);
+	SLATECORE_API virtual void OnMouseEnter(const FGeometry& MyGeometry, const FPointerEvent& MouseEvent);
 
 	/**
 	 * The system will use this event to notify a widget that the cursor has left it. This event is uses a custom bubble strategy.
 	 *
 	 * @param MouseEvent Information about the input event
 	 */
-	virtual void OnMouseLeave(const FPointerEvent& MouseEvent);
+	SLATECORE_API virtual void OnMouseLeave(const FPointerEvent& MouseEvent);
 
 	/**
 	 * Called when the mouse wheel is spun. This event is bubbled.
@@ -460,21 +460,21 @@ public:
 	 * @param  MouseEvent  Mouse event
 	 * @return  Returns whether the event was handled, along with other possible actions
 	 */
-	virtual FReply OnMouseWheel(const FGeometry& MyGeometry, const FPointerEvent& MouseEvent);
+	SLATECORE_API virtual FReply OnMouseWheel(const FGeometry& MyGeometry, const FPointerEvent& MouseEvent);
 
 	/**
 	 * The system asks each widget under the mouse to provide a cursor. This event is bubbled.
 	 *
 	 * @return FCursorReply::Unhandled() if the event is not handled; return FCursorReply::Cursor() otherwise.
 	 */
-	virtual FCursorReply OnCursorQuery(const FGeometry& MyGeometry, const FPointerEvent& CursorEvent) const;
+	SLATECORE_API virtual FCursorReply OnCursorQuery(const FGeometry& MyGeometry, const FPointerEvent& CursorEvent) const;
 
 	/**
 	 * After OnCursorQuery has specified a cursor type the system asks each widget under the mouse to map that cursor to a widget. This event is bubbled.
 	 *
 	 * @return TOptional<TSharedRef<SWidget>>() if you don't have a mapping otherwise return the Widget to show.
 	 */
-	virtual TOptional<TSharedRef<SWidget>> OnMapCursor(const FCursorReply& CursorReply) const;
+	SLATECORE_API virtual TOptional<TSharedRef<SWidget>> OnMapCursor(const FCursorReply& CursorReply) const;
 
 	/**
 	 * Called when a mouse button is double clicked.  Override this in derived classes.
@@ -483,7 +483,7 @@ public:
 	 * @param  InMouseEvent  Mouse button event
 	 * @return  Returns whether the event was handled, along with other possible actions
 	 */
-	virtual FReply OnMouseButtonDoubleClick(const FGeometry& InMyGeometry, const FPointerEvent& InMouseEvent);
+	SLATECORE_API virtual FReply OnMouseButtonDoubleClick(const FGeometry& InMyGeometry, const FPointerEvent& InMouseEvent);
 
 	/**
 	 * Called when Slate wants to visualize tooltip.
@@ -493,7 +493,7 @@ public:
 	 * @param  TooltipContent    The TooltipContent that I may want to visualize.
 	 * @return true if this widget visualized the tooltip content; i.e., the event is handled.
 	 */
-	virtual bool OnVisualizeTooltip(const TSharedPtr<SWidget>& TooltipContent);
+	SLATECORE_API virtual bool OnVisualizeTooltip(const TSharedPtr<SWidget>& TooltipContent);
 
 	/**
 	 * Visualize a new pop-up if possible.  If it's not possible for this widget to host the pop-up
@@ -504,7 +504,7 @@ public:
 	 *
 	 * @return a valid FPopupLayer if this widget supported hosting it.  You can call Remove() on this to destroy the pop-up.
 	 */
-	virtual TSharedPtr<FPopupLayer> OnVisualizePopup(const TSharedRef<SWidget>& PopupContent);
+	SLATECORE_API virtual TSharedPtr<FPopupLayer> OnVisualizePopup(const TSharedRef<SWidget>& PopupContent);
 
 	/**
 	 * Called when Slate detects that a widget started to be dragged.
@@ -519,7 +519,7 @@ public:
 	 * @param  InMyGeometry  Widget geometry
 	 * @param  InMouseEvent  MouseMove that triggered the drag
 	 */
-	virtual FReply OnDragDetected(const FGeometry& MyGeometry, const FPointerEvent& MouseEvent);
+	SLATECORE_API virtual FReply OnDragDetected(const FGeometry& MyGeometry, const FPointerEvent& MouseEvent);
 
 	//
 	// DRAG AND DROP (DragDrop)
@@ -542,14 +542,14 @@ public:
 	 *
 	 * @return A reply that indicated whether the contents of the DragDropEvent can potentially be processed by this widget.
 	 */
-	virtual void OnDragEnter(const FGeometry& MyGeometry, const FDragDropEvent& DragDropEvent);
+	SLATECORE_API virtual void OnDragEnter(const FGeometry& MyGeometry, const FDragDropEvent& DragDropEvent);
 
 	/**
 	 * Called during drag and drop when the drag leaves a widget.
 	 *
 	 * @param DragDropEvent   The drag and drop event.
 	 */
-	virtual void OnDragLeave(const FDragDropEvent& DragDropEvent);
+	SLATECORE_API virtual void OnDragLeave(const FDragDropEvent& DragDropEvent);
 
 	/**
 	 * Called during drag and drop when the the mouse is being dragged over a widget.
@@ -558,7 +558,7 @@ public:
 	 * @param DragDropEvent   The drag and drop event.
 	 * @return A reply that indicated whether this event was handled.
 	 */
-	virtual FReply OnDragOver(const FGeometry& MyGeometry, const FDragDropEvent& DragDropEvent);
+	SLATECORE_API virtual FReply OnDragOver(const FGeometry& MyGeometry, const FDragDropEvent& DragDropEvent);
 
 	/**
 	 * Called when the user is dropping something onto a widget; terminates drag and drop.
@@ -567,7 +567,7 @@ public:
 	 * @param DragDropEvent   The drag and drop event.
 	 * @return A reply that indicated whether this event was handled.
 	 */
-	virtual FReply OnDrop(const FGeometry& MyGeometry, const FDragDropEvent& DragDropEvent);
+	SLATECORE_API virtual FReply OnDrop(const FGeometry& MyGeometry, const FDragDropEvent& DragDropEvent);
 
 	//
 	// TOUCH and GESTURES
@@ -579,42 +579,42 @@ public:
 	 * @param  GestureEvent  gesture event
 	 * @return  Returns whether the event was handled, along with other possible actions
 	 */
-	virtual FReply OnTouchGesture(const FGeometry& MyGeometry, const FPointerEvent& GestureEvent);
+	SLATECORE_API virtual FReply OnTouchGesture(const FGeometry& MyGeometry, const FPointerEvent& GestureEvent);
 
 	/**
 	 * Called when a touchpad touch is started (finger down)
 	 *
 	 * @param InTouchEvent	The touch event generated
 	 */
-	virtual FReply OnTouchStarted(const FGeometry& MyGeometry, const FPointerEvent& InTouchEvent);
+	SLATECORE_API virtual FReply OnTouchStarted(const FGeometry& MyGeometry, const FPointerEvent& InTouchEvent);
 
 	/**
 	 * Called when a touchpad touch is moved  (finger moved)
 	 *
 	 * @param InTouchEvent	The touch event generated
 	 */
-	virtual FReply OnTouchMoved(const FGeometry& MyGeometry, const FPointerEvent& InTouchEvent);
+	SLATECORE_API virtual FReply OnTouchMoved(const FGeometry& MyGeometry, const FPointerEvent& InTouchEvent);
 
 	/**
 	 * Called when a touchpad touch is ended (finger lifted)
 	 *
 	 * @param InTouchEvent	The touch event generated
 	 */
-	virtual FReply OnTouchEnded(const FGeometry& MyGeometry, const FPointerEvent& InTouchEvent);
+	SLATECORE_API virtual FReply OnTouchEnded(const FGeometry& MyGeometry, const FPointerEvent& InTouchEvent);
 
 	/**
 	 * Called when a touchpad touch force changes
 	 *
 	 * @param InTouchEvent	The touch event generated
 	 */
-	virtual FReply OnTouchForceChanged(const FGeometry& MyGeometry, const FPointerEvent& TouchEvent);
+	SLATECORE_API virtual FReply OnTouchForceChanged(const FGeometry& MyGeometry, const FPointerEvent& TouchEvent);
 
 	/**
 	 * Called when a touchpad touch first moves after TouchStarted
 	 *
 	 * @param InTouchEvent	The touch event generated
 	 */
-	virtual FReply OnTouchFirstMove(const FGeometry& MyGeometry, const FPointerEvent& TouchEvent);
+	SLATECORE_API virtual FReply OnTouchFirstMove(const FGeometry& MyGeometry, const FPointerEvent& TouchEvent);
 
 	/**
 	 * Called when motion is detected (controller or device)
@@ -622,14 +622,14 @@ public:
 	 *
 	 * @param InMotionEvent	The motion event generated
 	 */
-	virtual FReply OnMotionDetected(const FGeometry& MyGeometry, const FMotionEvent& InMotionEvent);
+	SLATECORE_API virtual FReply OnMotionDetected(const FGeometry& MyGeometry, const FMotionEvent& InMotionEvent);
 
 	/**
 	 * Called to determine if we should render the focus brush.
 	 *
 	 * @param InFocusCause	The cause of focus
 	 */
-	virtual TOptional<bool> OnQueryShowFocus(const EFocusCause InFocusCause) const;
+	SLATECORE_API virtual TOptional<bool> OnQueryShowFocus(const EFocusCause InFocusCause) const;
 
 	/**
 	 * Popups can manifest in a NEW OS WINDOW or via an OVERLAY in an existing window.
@@ -641,21 +641,21 @@ public:
 	 *      EPopupMethod::UserCurrentWindow. This makes all the menu anchors within them
 	 *      use the current window.
 	 */
-	virtual FPopupMethodReply OnQueryPopupMethod() const;
+	SLATECORE_API virtual FPopupMethodReply OnQueryPopupMethod() const;
 
-	virtual TOptional<FVirtualPointerPosition> TranslateMouseCoordinateForCustomHitTestChild(const SWidget& ChildWidget, const FGeometry& MyGeometry, const FVector2D ScreenSpaceMouseCoordinate, const FVector2D LastScreenSpaceMouseCoordinate) const;
+	SLATECORE_API virtual TOptional<FVirtualPointerPosition> TranslateMouseCoordinateForCustomHitTestChild(const SWidget& ChildWidget, const FGeometry& MyGeometry, const FVector2D ScreenSpaceMouseCoordinate, const FVector2D LastScreenSpaceMouseCoordinate) const;
 
 	/**
 	 * All the pointer (mouse, touch, stylus, etc.) events from this frame have been routed.
 	 * This is a widget's chance to act on any accumulated data.
 	 */
-	virtual void OnFinishedPointerInput();
+	SLATECORE_API virtual void OnFinishedPointerInput();
 
 	/**
 	 * All the key (keyboard, gamepay, joystick, etc.) input from this frame has been routed.
 	 * This is a widget's chance to act on any accumulated data.
 	 */
-	virtual void OnFinishedKeyInput();
+	SLATECORE_API virtual void OnFinishedKeyInput();
 
 	/**
 	 * Called when navigation is requested
@@ -663,7 +663,7 @@ public:
 	 *
 	 * @param InNavigationEvent	The navigation event generated
 	 */
-	virtual FNavigationReply OnNavigation(const FGeometry& MyGeometry, const FNavigationEvent& InNavigationEvent);
+	SLATECORE_API virtual FNavigationReply OnNavigation(const FGeometry& MyGeometry, const FNavigationEvent& InNavigationEvent);
 
 	/**
 	 * Called when the mouse is moved over the widget's window, to determine if we should report whether
@@ -672,10 +672,10 @@ public:
 	 *
 	 * @return	The window "zone" the cursor is over, or EWindowZone::Unspecified if no special behavior is needed
 	 */
-	virtual EWindowZone::Type GetWindowZoneOverride() const;
+	SLATECORE_API virtual EWindowZone::Type GetWindowZoneOverride() const;
 
 #if WITH_ACCESSIBILITY
-	virtual TSharedRef<class FSlateAccessibleWidget> CreateAccessibleWidget();
+	SLATECORE_API virtual TSharedRef<class FSlateAccessibleWidget> CreateAccessibleWidget();
 #endif
 
 public:
@@ -686,13 +686,13 @@ public:
 	bool NeedsPrepass() const { return bNeedsPrepass; }
 	/** DEPRECATED version of SlatePrepass that assumes no scaling beyond AppScale*/
 	//UE_DEPRECATED(4.20, "SlatePrepass requires a layout scale to be accurate.")
-	void SlatePrepass();
+	SLATECORE_API void SlatePrepass();
 
 	/**
 	 * Descends to leaf-most widgets in the hierarchy and gathers desired sizes on the way up.
 	 * i.e. Caches the desired size of all of this widget's children recursively, then caches desired size for itself.
 	 */
-	void SlatePrepass(float InLayoutScaleMultiplier);
+	SLATECORE_API void SlatePrepass(float InLayoutScaleMultiplier);
 
 	void SetCanTick(bool bInCanTick) { bInCanTick ? AddUpdateFlags(EWidgetUpdateFlags::NeedsTick) : RemoveUpdateFlags(EWidgetUpdateFlags::NeedsTick); }
 	bool GetCanTick() const { return HasAnyUpdateFlags(EWidgetUpdateFlags::NeedsTick); }
@@ -706,10 +706,10 @@ public:
 	const FWidgetProxyHandle GetProxyHandle() const { return FastPathProxyHandle; }
 
 	/** @return the DesiredSize that was computed the last time CacheDesiredSize() was called. */
-	UE::Slate::FDeprecateVector2DResult GetDesiredSize() const;
+	SLATECORE_API UE::Slate::FDeprecateVector2DResult GetDesiredSize() const;
 
-	void AssignParentWidget(TSharedPtr<SWidget> InParent);
-	bool ConditionallyDetatchParentWidget(SWidget* InExpectedParent);
+	SLATECORE_API void AssignParentWidget(TSharedPtr<SWidget> InParent);
+	SLATECORE_API bool ConditionallyDetatchParentWidget(SWidget* InExpectedParent);
 
 	/**  */
 	virtual bool ValidatePathToChild(SWidget* InChild) { return true; }
@@ -723,7 +723,7 @@ public:
 	 * Calculates what if any clipping state changes need to happen when drawing this widget.
 	 * @return the culling rect that should be used going forward.
 	 */
-	FSlateRect CalculateCullingAndClippingRules(const FGeometry& AllottedGeometry, const FSlateRect& MyCullingRect, bool& bClipToBounds, bool& bAlwaysClip, bool& bIntersectClipBounds) const;
+	SLATECORE_API FSlateRect CalculateCullingAndClippingRules(const FGeometry& AllottedGeometry, const FSlateRect& MyCullingRect, bool& bClipToBounds, bool& bAlwaysClip, bool& bIntersectClipBounds) const;
 
 	bool HasAnyUpdateFlags(EWidgetUpdateFlags FlagsToCheck) const
 	{
@@ -739,7 +739,7 @@ protected:
 	 * The system calls this method. It performs a breadth-first traversal of every visible widget and asks
 	 * each widget to cache how big it needs to be in order to present all of its content.
 	 */
-	virtual void CacheDesiredSize(float InLayoutScaleMultiplier);
+	SLATECORE_API virtual void CacheDesiredSize(float InLayoutScaleMultiplier);
 
 	/**
 	 * Compute the ideal size necessary to display this widget. For aggregate widgets (e.g. panels) this size should include the
@@ -759,12 +759,12 @@ protected:
 
 private:
 	void SetFastPathProxyHandle(const FWidgetProxyHandle& Handle) { FastPathProxyHandle = Handle; }
-	void SetFastPathProxyHandle(const FWidgetProxyHandle& Handle, FSlateInvalidationWidgetVisibility Visibility, bool bParentVolatile);
-	void SetFastPathSortOrder(const FSlateInvalidationWidgetSortOrder SortOrder);
+	SLATECORE_API void SetFastPathProxyHandle(const FWidgetProxyHandle& Handle, FSlateInvalidationWidgetVisibility Visibility, bool bParentVolatile);
+	SLATECORE_API void SetFastPathSortOrder(const FSlateInvalidationWidgetSortOrder SortOrder);
 
-	void UpdateFastPathVisibility(FSlateInvalidationWidgetVisibility ParentVisibility, FHittestGrid* ParentHittestGrid);
-	void UpdateFastPathWidgetRemoved(FHittestGrid* ParentHittestGrid);
-	void UpdateFastPathVolatility(bool bParentVolatile);
+	SLATECORE_API void UpdateFastPathVisibility(FSlateInvalidationWidgetVisibility ParentVisibility, FHittestGrid* ParentHittestGrid);
+	SLATECORE_API void UpdateFastPathWidgetRemoved(FHittestGrid* ParentHittestGrid);
+	SLATECORE_API void UpdateFastPathVolatility(bool bParentVolatile);
 
 	/**
 	 * Explicitly set the desired size. This is highly advanced functionality that is meant
@@ -776,7 +776,7 @@ private:
 	}
 
 #if STATS || ENABLE_STATNAMEDEVENTS
-	void CreateStatID() const;
+	SLATECORE_API void CreateStatID() const;
 #endif
 
 	void AddUpdateFlags(EWidgetUpdateFlags FlagsToAdd)
@@ -793,7 +793,7 @@ private:
 		FastPathProxyHandle.UpdateWidgetFlags(this, Previous, UpdateFlags);
 	}
 
-	void UpdateWidgetProxy(int32 NewLayerId, FSlateCachedElementsHandle& CacheHandle);
+	SLATECORE_API void UpdateWidgetProxy(int32 NewLayerId, FSlateCachedElementsHandle& CacheHandle);
 
 public:
 
@@ -831,7 +831,7 @@ public:
 	virtual float GetRelativeLayoutScale(const FSlotBase& Child, float LayoutScaleMultiplier) const { return 1.0f; }
 
 	/** What is the Child's scale relative to this widget. */
-	virtual float GetRelativeLayoutScale(const int32 ChildIndex, float LayoutScaleMultiplier) const;
+	SLATECORE_API virtual float GetRelativeLayoutScale(const int32 ChildIndex, float LayoutScaleMultiplier) const;
 
 	/**
 	 * Non-virtual entry point for arrange children. ensures common work is executed before calling the virtual
@@ -847,7 +847,7 @@ public:
 	 * @param ArrangedChildren    The array to which to add the WidgetGeometries that represent the arranged children.
 	 * @param bUpdateAttributes   Update the collapsed attributes.
 	 */
-	void ArrangeChildren(const FGeometry& AllottedGeometry, FArrangedChildren& ArrangedChildren, bool bUpdateAttributes = false) const;
+	SLATECORE_API void ArrangeChildren(const FGeometry& AllottedGeometry, FArrangedChildren& ArrangedChildren, bool bUpdateAttributes = false) const;
 
 	/**
 	 * Returns the useful children (if any) of this widget. Some widget type may hide widget if they are needed by the system.
@@ -877,59 +877,59 @@ public:
 	 *
 	 * @return  True if this widget can take keyboard focus
 	 */
-	virtual bool SupportsKeyboardFocus() const;
+	SLATECORE_API virtual bool SupportsKeyboardFocus() const;
 
 	/**
 	 * Checks to see if this widget currently has the keyboard focus
 	 *
 	 * @return  True if this widget has keyboard focus
 	 */
-	virtual bool HasKeyboardFocus() const;
+	SLATECORE_API virtual bool HasKeyboardFocus() const;
 
 	/**
 	 * Gets whether or not the specified users has this widget focused, and if so the type of focus.
 	 *
 	 * @return The optional will be set with the focus cause, if unset this widget doesn't have focus.
 	 */
-	TOptional<EFocusCause> HasUserFocus(int32 UserIndex) const;
+	SLATECORE_API TOptional<EFocusCause> HasUserFocus(int32 UserIndex) const;
 
 	/**
 	 * Gets whether or not any users have this widget focused, and if so the type of focus (first one found).
 	 *
 	 * @return The optional will be set with the focus cause, if unset this widget doesn't have focus.
 	 */
-	TOptional<EFocusCause> HasAnyUserFocus() const;
+	SLATECORE_API TOptional<EFocusCause> HasAnyUserFocus() const;
 
 	/**
 	 * Gets whether or not the specified users has this widget or any descendant focused.
 	 *
 	 * @return The optional will be set with the focus cause, if unset this widget doesn't have focus.
 	 */
-	bool HasUserFocusedDescendants(int32 UserIndex) const;
+	SLATECORE_API bool HasUserFocusedDescendants(int32 UserIndex) const;
 
 	/**
 	 * @return Whether this widget has any descendants with keyboard focus
 	 */
-	bool HasFocusedDescendants() const;
+	SLATECORE_API bool HasFocusedDescendants() const;
 
 	/**
 	 * @return whether or not any users have this widget focused, or any descendant focused.
 	 */
-	bool HasAnyUserFocusOrFocusedDescendants() const;
+	SLATECORE_API bool HasAnyUserFocusOrFocusedDescendants() const;
 
 	/**
 	 * Checks to see if this widget is the current mouse captor
 	 *
 	 * @return  True if this widget has captured the mouse
 	 */
-	bool HasMouseCapture() const;
+	SLATECORE_API bool HasMouseCapture() const;
 
 	/**
 	 * Checks to see if this widget has mouse capture from the provided user.
 	 *
 	 * @return  True if this widget has captured the mouse
 	 */
-	bool HasMouseCaptureByUser(int32 UserIndex, TOptional<int32> PointerIndex = TOptional<int32>()) const;
+	SLATECORE_API bool HasMouseCaptureByUser(int32 UserIndex, TOptional<int32> PointerIndex = TOptional<int32>()) const;
 
 protected:
 	/** Called when this widget had captured the mouse, but that capture has been revoked for some reason. */
@@ -938,7 +938,7 @@ protected:
 
 public:
 	/** Called when this widget had captured the mouse, but that capture has been revoked for some reason. */
-	virtual void OnMouseCaptureLost(const FCaptureLostEvent& CaptureLostEvent);
+	SLATECORE_API virtual void OnMouseCaptureLost(const FCaptureLostEvent& CaptureLostEvent);
 
 	/**
 	 * Sets the enabled state of this widget
@@ -963,10 +963,10 @@ public:
 	}
 
 	/** @return The tool tip associated with this widget; Invalid reference if there is not one */
-	virtual TSharedPtr<IToolTip> GetToolTip();
+	SLATECORE_API virtual TSharedPtr<IToolTip> GetToolTip();
 
 	/** Called when a tooltip displayed from this widget is being closed */
-	virtual void OnToolTipClosing();
+	SLATECORE_API virtual void OnToolTipClosing();
 
 	/**
 	 * Sets whether this widget is a "tool tip force field".  That is, tool-tips should never spawn over the area
@@ -974,7 +974,7 @@ public:
 	 *
 	 * @param	bEnableForceField	True to enable tool tip force field for this widget
 	 */
-	void EnableToolTipForceField(const bool bEnableForceField);
+	SLATECORE_API void EnableToolTipForceField(const bool bEnableForceField);
 
 	/** @return True if a tool tip force field is active on this widget */
 	bool HasToolTipForceField() const
@@ -992,7 +992,7 @@ public:
 	}
 
 	/** @return True if this widget is directly hovered */
-	bool IsDirectlyHovered() const;
+	SLATECORE_API bool IsDirectlyHovered() const;
 
 protected:
 	/**
@@ -1015,14 +1015,14 @@ public:
 	FORCEINLINE EVisibility GetVisibility() const { return VisibilityAttribute.Get(); }
 
 	/** @param InVisibility  should this widget be */
-	virtual void SetVisibility(TAttribute<EVisibility> InVisibility);
+	SLATECORE_API virtual void SetVisibility(TAttribute<EVisibility> InVisibility);
 
 	/**
 	 * @return is the widget visible and its parents also visible.
 	 * @note only valid if the widget is contained by an InvalidationRoot (the proxy is valid).
 	 */
 	UE_DEPRECATED(5.0, "IsFastPathVisible is deprecated and should not be used.")
-	bool IsFastPathVisible() const;
+	SLATECORE_API bool IsFastPathVisible() const;
 
 #if WITH_ACCESSIBILITY
 	/**
@@ -1031,14 +1031,14 @@ public:
 	 * @param AccessibleType Whether the widget is being accessed directly or through a summary query.
 	 * @return The text that should be conveyed to the user describing this widget.
 	 */
-	FText GetAccessibleText(EAccessibleType AccessibleType = EAccessibleType::Main) const;
+	SLATECORE_API FText GetAccessibleText(EAccessibleType AccessibleType = EAccessibleType::Main) const;
 
 	/**
 	 * Traverse all child widgets and concat their results of GetAccessibleText(Summary).
 	 *
 	 * @return The combined text of all child widget's summary text.
 	 */
-	FText GetAccessibleSummary() const;
+	SLATECORE_API FText GetAccessibleSummary() const;
 
 	/**
 	 * Whether this widget is considered accessible or not. A widget is accessible if its behavior
@@ -1046,7 +1046,7 @@ public:
 	 *
 	 * @return true if an accessible widget should be created for this widget.
 	 */
-	bool IsAccessible() const;
+	SLATECORE_API bool IsAccessible() const;
 
 	/**
 	 * Get the behavior describing how the accessible text of this widget should be retrieved.
@@ -1054,14 +1054,14 @@ public:
 	 * @param AccessibleType Whether the widget is being accessed directly or through a summary query.
 	 * @return The accessible behavior of the widget.
 	 */
-	EAccessibleBehavior GetAccessibleBehavior(EAccessibleType AccessibleType = EAccessibleType::Main) const;
+	SLATECORE_API EAccessibleBehavior GetAccessibleBehavior(EAccessibleType AccessibleType = EAccessibleType::Main) const;
 
 	/**
 	 * Checks whether this widget allows its children to be accessible or not.
 	 *
 	 * @return true if children can be accessible.
 	 */
-	bool CanChildrenBeAccessible() const;
+	SLATECORE_API bool CanChildrenBeAccessible() const;
 
 	/**
 	 * Set a new accessible behavior, and if the behavior is custom, new accessible text to go along with it.
@@ -1070,7 +1070,7 @@ public:
 	 * @param InText, If the new behavior is custom, this will be the custom text assigned to the widget.
 	 * @param AccessibleType Whether the widget is being accessed directly or through a summary query.
 	 */
-	void SetAccessibleBehavior(EAccessibleBehavior InBehavior, const TAttribute<FText>& InText = TAttribute<FText>(), EAccessibleType AccessibleType = EAccessibleType::Main);
+	SLATECORE_API void SetAccessibleBehavior(EAccessibleBehavior InBehavior, const TAttribute<FText>& InText = TAttribute<FText>(), EAccessibleType AccessibleType = EAccessibleType::Main);
 
 	/**
 	 * Sets whether children are allowed to be accessible or not.
@@ -1078,14 +1078,14 @@ public:
 	 *
 	 * @param InCanChildrenBeAccessible Whether children should be accessible or not.
 	 */
-	void SetCanChildrenBeAccessible(bool InCanChildrenBeAccessible);
+	SLATECORE_API void SetCanChildrenBeAccessible(bool InCanChildrenBeAccessible);
 
 	/**
 	 * Assign AccessibleText with a default value that can be used when AccessibleBehavior is set to Auto or Custom.
 	 *
 	 * @param AccessibleType Whether the widget is being accessed directly or through a summary query.
 	 */
-	virtual TOptional<FText> GetDefaultAccessibleText(EAccessibleType AccessibleType = EAccessibleType::Main) const;
+	SLATECORE_API virtual TOptional<FText> GetDefaultAccessibleText(EAccessibleType AccessibleType = EAccessibleType::Main) const;
 #endif
 
 	/** Whether or not a widget is volatile and will update every frame without being invalidated */
@@ -1117,7 +1117,7 @@ public:
 	 * Invalidates the widget from the view of a layout caching widget that may own this widget.
 	 * will force the owning widget to redraw and cache children on the next paint pass.
 	 */
-	void Invalidate(EInvalidateWidgetReason InvalidateReason);
+	SLATECORE_API void Invalidate(EInvalidateWidgetReason InvalidateReason);
 
 	/**
 	 * Recalculates volatility of the widget and caches the result.  Should be called any time 
@@ -1129,7 +1129,7 @@ public:
 	}
 
 	UE_DEPRECATED(5.0, "InvalidatePrepass is deprecated. Use the Invalidate(EInvalidateWidgetReason::Prepass) or use MarkPrepassAsDirty()")
-	void InvalidatePrepass();
+	SLATECORE_API void InvalidatePrepass();
 
 	/**
 	 * In fast path, if the widget is mark, do a full Prepass on its next update to calculate it's desired size.
@@ -1145,7 +1145,7 @@ protected:
 	 * @param MyCullingRect the culling rect of the widget currently doing the culling.
 	 * @param ArrangedChild the arranged widget in the widget currently attempting to cull children.
 	 */
-	bool IsChildWidgetCulled(const FSlateRect& MyCullingRect, const FArrangedWidget& ArrangedChild) const;
+	SLATECORE_API bool IsChildWidgetCulled(const FSlateRect& MyCullingRect, const FArrangedWidget& ArrangedChild) const;
 #else
 	FORCEINLINE bool IsChildWidgetCulled(const FSlateRect&, const FArrangedWidget&) const { return false; }
 #endif
@@ -1157,7 +1157,7 @@ protected:
 	 * Called when a child is removed from the tree parent's widget tree either by removing it from a slot. This can also be called manually if you've got some non-slot based what of no longer reporting children
 	 * An example of a widget that needs manual calling is SWidgetSwitcher.  It keeps all its children but only arranges and paints a single "active" one.  Once a child becomes inactive, its cached data should be removed.
 	 */
-	void InvalidateChildRemovedFromTree(SWidget& Child);
+	SLATECORE_API void InvalidateChildRemovedFromTree(SWidget& Child);
 
 	/**
 	 * Recalculates and caches volatility and returns 'true' if the volatility changed.
@@ -1253,7 +1253,7 @@ public:
 	/**
 	 * Sets the clipping to bounds rules for this widget.
 	 */
-	void SetClipping(EWidgetClipping InClipping);
+	SLATECORE_API void SetClipping(EWidgetClipping InClipping);
 
 	/** @return The current clipping rules for this widget. */
 	FORCEINLINE EWidgetClipping GetClipping() const
@@ -1264,7 +1264,7 @@ public:
 	/**
 	* Sets the pixel snapping method for this widget.
 	*/
-	void SetPixelSnapping(EWidgetPixelSnapping InPixelSnappingMethod);
+	SLATECORE_API void SetPixelSnapping(EWidgetPixelSnapping InPixelSnappingMethod);
 
 	/** @return The current pixel snapping rules for this widget. */
 	FORCEINLINE EWidgetPixelSnapping GetPixelSnapping() const
@@ -1309,24 +1309,24 @@ public:
 	EFlowDirectionPreference GetFlowDirectionPreference() const { return FlowDirectionPreference; }
 
 	/** Set the tool tip that should appear when this widget is hovered. */
-	void SetToolTipText(const TAttribute<FText>& ToolTipText);
+	SLATECORE_API void SetToolTipText(const TAttribute<FText>& ToolTipText);
 
 	/** Set the tool tip that should appear when this widget is hovered. */
-	void SetToolTipText( const FText& InToolTipText );
+	SLATECORE_API void SetToolTipText( const FText& InToolTipText );
 
 	/** Set the tool tip that should appear when this widget is hovered. */
-	void SetToolTip(const TAttribute<TSharedPtr<IToolTip>>& InToolTip);
+	SLATECORE_API void SetToolTip(const TAttribute<TSharedPtr<IToolTip>>& InToolTip);
 
 	/** Set the cursor that should appear when this widget is hovered  */
-	void SetCursor( const TAttribute< TOptional<EMouseCursor::Type> >& InCursor );
+	SLATECORE_API void SetCursor( const TAttribute< TOptional<EMouseCursor::Type> >& InCursor );
 
 protected:
 
 	/** Used by Slate to set the runtime debug info about this widget. */
-	void SetDebugInfo( const ANSICHAR* InType, const ANSICHAR* InFile, int32 OnLine, size_t InAllocSize );
+	SLATECORE_API void SetDebugInfo( const ANSICHAR* InType, const ANSICHAR* InFile, int32 OnLine, size_t InAllocSize );
 
 	/** The cursor to show when the mouse is hovering over this widget. */
-	virtual TOptional<EMouseCursor::Type> GetCursor() const;
+	SLATECORE_API virtual TOptional<EMouseCursor::Type> GetCursor() const;
 
 public:
 
@@ -1401,7 +1401,7 @@ public:
 
 private:
 
-	void AddMetadataInternal(const TSharedRef<ISlateMetaData>& AddMe);
+	SLATECORE_API void AddMetadataInternal(const TSharedRef<ISlateMetaData>& AddMe);
 
 	template<typename MetaDataType>
 	int32 RemoveAllMetaData()
@@ -1438,44 +1438,44 @@ private:
 public:
 
 	/** See OnMouseButtonDown event */
-	void SetOnMouseButtonDown(FPointerEventHandler EventHandler);
+	SLATECORE_API void SetOnMouseButtonDown(FPointerEventHandler EventHandler);
 
 	/** See OnMouseButtonUp event */
-	void SetOnMouseButtonUp(FPointerEventHandler EventHandler);
+	SLATECORE_API void SetOnMouseButtonUp(FPointerEventHandler EventHandler);
 
 	/** See OnMouseMove event */
-	void SetOnMouseMove(FPointerEventHandler EventHandler);
+	SLATECORE_API void SetOnMouseMove(FPointerEventHandler EventHandler);
 
 	/** See OnMouseDoubleClick event */
-	void SetOnMouseDoubleClick(FPointerEventHandler EventHandler);
+	SLATECORE_API void SetOnMouseDoubleClick(FPointerEventHandler EventHandler);
 
 	/** See OnMouseEnter event */
-	void SetOnMouseEnter(FNoReplyPointerEventHandler EventHandler);
+	SLATECORE_API void SetOnMouseEnter(FNoReplyPointerEventHandler EventHandler);
 
 	/** See OnMouseLeave event */
-	void SetOnMouseLeave(FSimpleNoReplyPointerEventHandler EventHandler);
+	SLATECORE_API void SetOnMouseLeave(FSimpleNoReplyPointerEventHandler EventHandler);
 
 public:
 
 	// Widget Inspector and debugging methods
 
 	/** @return A String representation of the widget */
-	virtual FString ToString() const;
+	SLATECORE_API virtual FString ToString() const;
 
 	/** @return A String of the widget's type */
-	FString GetTypeAsString() const;
+	SLATECORE_API FString GetTypeAsString() const;
 
 	/** @return The widget's type as an FName ID */
-	FName GetType() const;
+	SLATECORE_API FName GetType() const;
 
 	/** @return A String of the widget's code location in readable format "BaseFileName(LineNumber)" */
-	virtual FString GetReadableLocation() const;
+	SLATECORE_API virtual FString GetReadableLocation() const;
 
 	/** @return An FName of the widget's code location (full path with number == line number of the file) */
-	FName GetCreatedInLocation() const;
+	SLATECORE_API FName GetCreatedInLocation() const;
 
 	/** @return The name this widget was tagged with */
-	virtual FName GetTag() const;
+	SLATECORE_API virtual FName GetTag() const;
 
 #if STATS
 	size_t GetAllocSize() const { return AllocSize; }
@@ -1487,13 +1487,13 @@ public:
 #endif
 
 	/** @return the Foreground color that this widget sets; unset options if the widget does not set a foreground color */
-	virtual FSlateColor GetForegroundColor() const;
+	SLATECORE_API virtual FSlateColor GetForegroundColor() const;
 
 	/** @return the Foreground color that this widget sets when this widget or any of its ancestors are disabled; unset options if the widget does not set a foreground color */
-	virtual FSlateColor GetDisabledForegroundColor() const;
+	SLATECORE_API virtual FSlateColor GetDisabledForegroundColor() const;
 
 	//UE_DEPRECATED(4.23, "GetCachedGeometry has been deprecated, use GetTickSpaceGeometry instead")
-	const FGeometry& GetCachedGeometry() const;
+	SLATECORE_API const FGeometry& GetCachedGeometry() const;
 
 	/**
 	 * Gets the last geometry used to Tick the widget.  This data may not exist yet if this call happens prior to
@@ -1504,13 +1504,13 @@ public:
 	 * or what are referred to as hysteresis problems, both caused by depending on geometry from the previous frame
 	 * being used to advise how to layout a dependent object the current frame.
 	 */
-	const FGeometry& GetTickSpaceGeometry() const;
+	SLATECORE_API const FGeometry& GetTickSpaceGeometry() const;
 
 	/**
 	 * Gets the last geometry used to Tick the widget.  This data may not exist yet if this call happens prior to
 	 * the widget having been ticked/painted, or it may be out of date, or a frame behind.
  	 */
-	const FGeometry& GetPaintSpaceGeometry() const;
+	SLATECORE_API const FGeometry& GetPaintSpaceGeometry() const;
 
 	/** Returns the clipping state to clip this widget against its parent */
 	const TOptional<FSlateClippingState>& GetCurrentClippingState() const { return PersistentState.InitialClipState; }
@@ -1529,10 +1529,10 @@ protected:
 	 *
 	 * @see SNew
 	 */
-	SWidget();
+	SLATECORE_API SWidget();
 
 	/** Construct a SWidget based on initial parameters. */
-	void SWidgetConstruct(const FSlateBaseNamedArgs& Args);
+	SLATECORE_API void SWidgetConstruct(const FSlateBaseNamedArgs& Args);
 
 	/** Is the widget construction completed (did we called and returned from the Construct() function) */
 	bool IsConstructed() const { return bIsDeclarativeSyntaxConstructionCompleted; }
@@ -1546,7 +1546,7 @@ protected:
 	 * @param OutResult       A map of widget references to their respective geometries.
 	 * @return True if all the WidgetGeometries were found. False otherwise.
 	 */
-	bool FindChildGeometries( const FGeometry& MyGeometry, const TSet< TSharedRef<SWidget> >& WidgetsToFind, TMap<TSharedRef<SWidget>, FArrangedWidget>& OutResult ) const;
+	SLATECORE_API bool FindChildGeometries( const FGeometry& MyGeometry, const TSet< TSharedRef<SWidget> >& WidgetsToFind, TMap<TSharedRef<SWidget>, FArrangedWidget>& OutResult ) const;
 
 	/**
 	 * Actual implementation of FindChildGeometries.
@@ -1555,7 +1555,7 @@ protected:
 	 * @param WidgetsToFind   The widgets whose geometries we wish to discover.
 	 * @param OutResult       A map of widget references to their respective geometries.
 	 */
-	void FindChildGeometries_Helper( const FGeometry& MyGeometry, const TSet< TSharedRef<SWidget> >& WidgetsToFind, TMap<TSharedRef<SWidget>, FArrangedWidget>& OutResult ) const;
+	SLATECORE_API void FindChildGeometries_Helper( const FGeometry& MyGeometry, const TSet< TSharedRef<SWidget> >& WidgetsToFind, TMap<TSharedRef<SWidget>, FArrangedWidget>& OutResult ) const;
 
 	/** 
 	 * Find the geometry of a descendant widget. This method assumes that WidgetToFind is a descendant of this widget.
@@ -1564,13 +1564,13 @@ protected:
 	 * @param WidgetToFind The widget whose geometry we wish to discover.
 	 * @return the geometry of WidgetToFind.
 	 */
-	FGeometry FindChildGeometry( const FGeometry& MyGeometry, TSharedRef<SWidget> WidgetToFind ) const;
+	SLATECORE_API FGeometry FindChildGeometry( const FGeometry& MyGeometry, TSharedRef<SWidget> WidgetToFind ) const;
 
 	/** @return The index of the child that the mouse is currently hovering */
-	static int32 FindChildUnderMouse( const FArrangedChildren& Children, const FPointerEvent& MouseEvent );
+	static SLATECORE_API int32 FindChildUnderMouse( const FArrangedChildren& Children, const FPointerEvent& MouseEvent );
 
 	/** @return The index of the child that is under the specified position */
-	static int32 FindChildUnderPosition(const FArrangedChildren& Children, const UE::Slate::FDeprecateVector2DParameter& ArrangedSpacePosition);
+	static SLATECORE_API int32 FindChildUnderPosition(const FArrangedChildren& Children, const UE::Slate::FDeprecateVector2DParameter& ArrangedSpacePosition);
 
 	/** 
 	 * Determines if this widget should be enabled.
@@ -1585,7 +1585,7 @@ protected:
 	}
 
 	/** @return a brush to draw focus, nullptr if no focus drawing is desired */
-	virtual const FSlateBrush* GetFocusBrush() const;
+	SLATECORE_API virtual const FSlateBrush* GetFocusBrush() const;
 
 	/**
 	 * Recomputes the volatility of the widget.  If you have additional state you automatically want to make
@@ -1607,7 +1607,7 @@ protected:
 	 * Called when clipping is changed.  Should be used to forward clipping states onto potentially
 	 * hidden children that actually are responsible for clipping the content.
 	 */
-	virtual void OnClippingChanged();
+	SLATECORE_API virtual void OnClippingChanged();
 
 private:
 
@@ -1636,13 +1636,13 @@ private:
 	 */
 	virtual void OnArrangeChildren(const FGeometry& AllottedGeometry, FArrangedChildren& ArrangedChildren) const = 0;
 
-	void Prepass_Internal(float LayoutScaleMultiplier);
+	SLATECORE_API void Prepass_Internal(float LayoutScaleMultiplier);
 
 protected:
 
 	float GetPrepassLayoutScaleMultiplier() const { return PrepassLayoutScaleMultiplier.Get(1.0f); }
 	
-	void Prepass_ChildLoop(float InLayoutScaleMultiplier, FChildren* MyChildren);
+	SLATECORE_API void Prepass_ChildLoop(float InLayoutScaleMultiplier, FChildren* MyChildren);
 
 public:
 	/**
@@ -1668,12 +1668,12 @@ public:
 	 * @param TimerFunction The active timer delegate to call every Period seconds.
 	 * @return An active timer handle that can be used to UnRegister later.
 	 */
-	TSharedRef<FActiveTimerHandle> RegisterActiveTimer( float TickPeriod, FWidgetActiveTimerDelegate TickFunction );
+	SLATECORE_API TSharedRef<FActiveTimerHandle> RegisterActiveTimer( float TickPeriod, FWidgetActiveTimerDelegate TickFunction );
 
 	/**
 	 * Unregisters an active timer handle. This is optional, as the delegate can UnRegister itself by returning EActiveTimerReturnType::Stop.
 	 */
-	void UnRegisterActiveTimer( const TSharedRef<FActiveTimerHandle>& ActiveTimerHandle );
+	SLATECORE_API void UnRegisterActiveTimer( const TSharedRef<FActiveTimerHandle>& ActiveTimerHandle );
 	
 	/** Does this widget have any active timers? */
 	bool HasActiveTimers() const { return ActiveTimers.Num() > 0; }
@@ -1681,7 +1681,7 @@ public:
 private:
 
 	/** Iterates over the active timer handles on the widget and executes them if their interval has elapsed. */
-	void ExecuteActiveTimers(double CurrentTime, float DeltaTime);
+	SLATECORE_API void ExecuteActiveTimers(double CurrentTime, float DeltaTime);
 
 protected:
 	/**
@@ -1707,7 +1707,7 @@ protected:
 
 protected:
 	/** Dtor ensures that active timer handles are UnRegistered with the SlateApplication. */
-	virtual ~SWidget();
+	SLATECORE_API virtual ~SWidget();
 
 private:
 	/** Handle to the proxy when on the fast path */
