[SlateCore]

diff --git a/Engine/Source/Runtime/SlateCore/Public/Widgets/SWindow.h b/Engine/Source/Runtime/SlateCore/Public/Widgets/SWindow.h
--- a/Engine/Source/Runtime/SlateCore/Public/Widgets/SWindow.h
+++ b/Engine/Source/Runtime/SlateCore/Public/Widgets/SWindow.h
@@ -237,7 +237,7 @@ private:
 /**
  * SWindow is a platform-agnostic representation of a top-level window.
  */
-class SLATECORE_API SWindow
+class SWindow
 	: public SCompoundWidget
 	, public FSlateInvalidationRoot
 {
@@ -380,26 +380,26 @@ public:
 	/**
 	 * Default constructor. Use SNew(SWindow) instead.
 	 */
-	SWindow();
-	~SWindow();
+	SLATECORE_API SWindow();
+	SLATECORE_API ~SWindow();
 
 public:
 
-	void Construct(const FArguments& InArgs);
+	SLATECORE_API void Construct(const FArguments& InArgs);
 
 	/**
 	 * Make a tool tip window
 	 *
 	 * @return The new SWindow
 	 */
-	static TSharedRef<SWindow> MakeToolTipWindow();
+	static SLATECORE_API TSharedRef<SWindow> MakeToolTipWindow();
 
 	/**
 	 * Make cursor decorator window
 	 *
 	 * @return The new SWindow
 	 */
-	static TSharedRef<SWindow> MakeCursorDecorator();
+	static SLATECORE_API TSharedRef<SWindow> MakeCursorDecorator();
 
 	/**
 	 * Make cursor decorator window with a non-default style
@@ -407,20 +407,20 @@ public:
 	 * @param InStyle The style to use for the cursor decorator
 	 * @return The new SWindow
 	 */
-	static TSharedRef<SWindow> MakeStyledCursorDecorator(const FWindowStyle& InStyle);
+	static SLATECORE_API TSharedRef<SWindow> MakeStyledCursorDecorator(const FWindowStyle& InStyle);
 
 	/**
 	 * Make a notification window
 	 *
 	 * @return The new SWindow
 	 */
-	static TSharedRef<SWindow> MakeNotificationWindow();
+	static SLATECORE_API TSharedRef<SWindow> MakeNotificationWindow();
 
 	/**
 	 * @param ContentSize      The size of content that we want to accommodate 
 	 *
 	 * @return The size of the window necessary to accommodate the given content */
-	static UE::Slate::FDeprecateVector2DResult ComputeWindowSizeForContent( UE::Slate::FDeprecateVector2DParameter ContentSize );
+	static SLATECORE_API UE::Slate::FDeprecateVector2DResult ComputeWindowSizeForContent( UE::Slate::FDeprecateVector2DParameter ContentSize );
 
 	/**
 	 * Grabs the window type
@@ -457,97 +457,97 @@ public:
 	}
 
 	/** Paint the window and all of its contents. Not the same as Paint(). */
-	int32 PaintWindow( double CurrentTime, float DeltaTime, FSlateWindowElementList& OutDrawElements, const FWidgetStyle& InWidgetStyle, bool bParentEnabled );
+	SLATECORE_API int32 PaintWindow( double CurrentTime, float DeltaTime, FSlateWindowElementList& OutDrawElements, const FWidgetStyle& InWidgetStyle, bool bParentEnabled );
 
 	/**
 	 * Returns the size of the title bar as a Slate size parameter.  Does not take into account application scale!
 	 *
 	 * @return  Title bar size
 	 */
-	FOptionalSize GetTitleBarSize() const;
+	SLATECORE_API FOptionalSize GetTitleBarSize() const;
 
 	/** @return the desired size in desktop pixels */
-	UE::Slate::FDeprecateVector2DResult GetDesiredSizeDesktopPixels() const;
+	SLATECORE_API UE::Slate::FDeprecateVector2DResult GetDesiredSizeDesktopPixels() const;
 
 	/**	@return The initially desired screen position of the slate window */
-	UE::Slate::FDeprecateVector2DResult GetInitialDesiredSizeInScreen() const;
+	SLATECORE_API UE::Slate::FDeprecateVector2DResult GetInitialDesiredSizeInScreen() const;
 
 	/**	@return The initially desired size of the slate window */
-	UE::Slate::FDeprecateVector2DResult GetInitialDesiredPositionInScreen() const;
+	SLATECORE_API UE::Slate::FDeprecateVector2DResult GetInitialDesiredPositionInScreen() const;
 
 	/** Get the Geometry that describes this window. Windows in Slate are unique in that they know their own geometry. */
-	FGeometry GetWindowGeometryInScreen() const;
+	SLATECORE_API FGeometry GetWindowGeometryInScreen() const;
 
 	/** @return The geometry of the window in window space (i.e. position and AbsolutePosition are 0) */
-	FGeometry GetWindowGeometryInWindow() const;
+	SLATECORE_API FGeometry GetWindowGeometryInWindow() const;
 
 	/** @return the transform from local space to screen space (desktop space). */
-	FSlateLayoutTransform GetLocalToScreenTransform() const;
+	SLATECORE_API FSlateLayoutTransform GetLocalToScreenTransform() const;
 
 	/** @return the transform from local space to window space, which is basically desktop space without the offset. Essentially contains the DPI scale. */
-	FSlateLayoutTransform GetLocalToWindowTransform() const;
+	SLATECORE_API FSlateLayoutTransform GetLocalToWindowTransform() const;
 
 	/** @return The position of the window in screen space */
-	UE::Slate::FDeprecateVector2DResult GetPositionInScreen() const;
+	SLATECORE_API UE::Slate::FDeprecateVector2DResult GetPositionInScreen() const;
 
 	/** @return the size of the window in screen pixels */
-	UE::Slate::FDeprecateVector2DResult GetSizeInScreen() const;
+	SLATECORE_API UE::Slate::FDeprecateVector2DResult GetSizeInScreen() const;
 
 	/** @return the rectangle of the window for its non-maximized state */
-	FSlateRect GetNonMaximizedRectInScreen() const;
+	SLATECORE_API FSlateRect GetNonMaximizedRectInScreen() const;
 
 	/** @return Rectangle that this window occupies in screen space */
-	FSlateRect GetRectInScreen() const;
+	SLATECORE_API FSlateRect GetRectInScreen() const;
 
 	/** @return Rectangle of the window's usable client area in screen space. */
-	FSlateRect GetClientRectInScreen() const;
+	SLATECORE_API FSlateRect GetClientRectInScreen() const;
 
 	/** @return the size of the window's usable client area. */
-	UE::Slate::FDeprecateVector2DResult GetClientSizeInScreen() const;
+	SLATECORE_API UE::Slate::FDeprecateVector2DResult GetClientSizeInScreen() const;
 
 	/** @return a clipping rectangle that represents this window in Window Space (i.e. always starts at 0,0) */
-	FSlateRect GetClippingRectangleInWindow() const;
+	SLATECORE_API FSlateRect GetClippingRectangleInWindow() const;
 
 	/** Returns the margins used for the window border. This varies based on whether it's maximized or not. */
-	FMargin GetWindowBorderSize( bool bIncTitleBar = false ) const;
+	SLATECORE_API FMargin GetWindowBorderSize( bool bIncTitleBar = false ) const;
 
 	/** Returns the margins used for the window border if it's not maximized */
-	FMargin GetNonMaximizedWindowBorderSize() const;
+	SLATECORE_API FMargin GetNonMaximizedWindowBorderSize() const;
 
 	/** Relocate the window to a screenspace position specified by NewPosition */
-	void MoveWindowTo( UE::Slate::FDeprecateVector2DParameter NewPosition );
+	SLATECORE_API void MoveWindowTo( UE::Slate::FDeprecateVector2DParameter NewPosition );
 	/** Relocate the window to a screenspace position specified by NewPosition and resize it to NewSize */
-	void ReshapeWindow( UE::Slate::FDeprecateVector2DParameter NewPosition, UE::Slate::FDeprecateVector2DParameter NewSize );
-	void ReshapeWindow( const FSlateRect& InNewShape );
+	SLATECORE_API void ReshapeWindow( UE::Slate::FDeprecateVector2DParameter NewPosition, UE::Slate::FDeprecateVector2DParameter NewSize );
+	SLATECORE_API void ReshapeWindow( const FSlateRect& InNewShape );
 	/**
 	 * Resize the window to be dpi scaled NewClientSize immediately
 	 *
 	 * @param NewClientSize: Client size with DPI scaling already applied that does not include border or title bars.
 	 */
-	void Resize( UE::Slate::FDeprecateVector2DParameter NewClientSize );
+	SLATECORE_API void Resize( UE::Slate::FDeprecateVector2DParameter NewClientSize );
 
 	/** Returns the rectangle of the screen the window is associated with */
-	FSlateRect GetFullScreenInfo() const;
+	SLATECORE_API FSlateRect GetFullScreenInfo() const;
 
 	/** @return Returns true if the window is currently morphing to a new position, shape and/or opacity */
-	bool IsMorphing() const;
+	SLATECORE_API bool IsMorphing() const;
 	/** @return Returns true if the window is currently morphing and is morphing by size */
-	bool IsMorphingSize() const;
+	SLATECORE_API bool IsMorphingSize() const;
 	/** Animate the window to TargetOpacity and TargetPosition over a short period of time */
-	void MorphToPosition( const FCurveSequence& Sequence, const float TargetOpacity, const UE::Slate::FDeprecateVector2DParameter& TargetPosition );
+	SLATECORE_API void MorphToPosition( const FCurveSequence& Sequence, const float TargetOpacity, const UE::Slate::FDeprecateVector2DParameter& TargetPosition );
 	/** Animate the window to TargetOpacity and TargetShape over a short period of time */
-	void MorphToShape( const FCurveSequence& Sequence, const float TargetOpacity, const FSlateRect& TargetShape );
+	SLATECORE_API void MorphToShape( const FCurveSequence& Sequence, const float TargetOpacity, const FSlateRect& TargetShape );
 	/** Set a new morph shape and force the morph to run for at least one frame in order to reach that target */
-	void UpdateMorphTargetShape( const FSlateRect& TargetShape );
+	SLATECORE_API void UpdateMorphTargetShape( const FSlateRect& TargetShape );
 	/** Set a new morph position and force the morph to run for at least one frame in order to reach that target */
-	void UpdateMorphTargetPosition( const UE::Slate::FDeprecateVector2DParameter& TargetPosition );
+	SLATECORE_API void UpdateMorphTargetPosition( const UE::Slate::FDeprecateVector2DParameter& TargetPosition );
 	/** @return Returns the currently set morph target position */
-	UE::Slate::FDeprecateVector2DResult GetMorphTargetPosition() const;
+	SLATECORE_API UE::Slate::FDeprecateVector2DResult GetMorphTargetPosition() const;
 	/** @return Returns the currently set morph target shape */
-	FSlateRect GetMorphTargetShape() const;
+	SLATECORE_API FSlateRect GetMorphTargetShape() const;
 
 	/** Flashed the window, used for drawing attention to modal dialogs */
-	void FlashWindow();
+	SLATECORE_API void FlashWindow();
 
 	/**
 	 * Attempts to draw the user's attention to this window in whatever way is
@@ -556,7 +556,7 @@ public:
 	 * @param Parameters The parameters for bouncing. Depending on the
 	 *        platform, not all parameters may be supported.
 	 */
-	void DrawAttention(const FWindowDrawAttentionParameters& Parameters);
+	SLATECORE_API void DrawAttention(const FWindowDrawAttentionParameters& Parameters);
 
 	/** 
 	 * Bring the window to the front 
@@ -565,28 +565,28 @@ public:
 	 *					In general do not pass force in.  It can be useful for some window types, like game windows where not forcing it to the front
 	 *					would cause mouse capture and mouse lock to happen but without the window visible
 	 */
-	void BringToFront( bool bForce = false );
+	SLATECORE_API void BringToFront( bool bForce = false );
 
 	/** @hack Force a window to front even if a different application is in front. */
-	void HACK_ForceToFront();
+	SLATECORE_API void HACK_ForceToFront();
 
 	/** Sets the actual screen position of the window. THIS SHOULD ONLY BE CALLED BY THE OS */
-	void SetCachedScreenPosition(UE::Slate::FDeprecateVector2DParameter NewPosition);
+	SLATECORE_API void SetCachedScreenPosition(UE::Slate::FDeprecateVector2DParameter NewPosition);
 
 	/**	Sets the actual size of the window. THIS SHOULD ONLY BE CALLED BY THE OS */
-	void SetCachedSize(UE::Slate::FDeprecateVector2DParameter NewSize);
+	SLATECORE_API void SetCachedSize(UE::Slate::FDeprecateVector2DParameter NewSize);
 
-	TSharedPtr<FGenericWindow> GetNativeWindow();
-	TSharedPtr<const FGenericWindow> GetNativeWindow() const ;
+	SLATECORE_API TSharedPtr<FGenericWindow> GetNativeWindow();
+	SLATECORE_API TSharedPtr<const FGenericWindow> GetNativeWindow() const ;
 
 	/** Returns the DPI scale factor of the native window */
-	float GetDPIScaleFactor() const;
+	SLATECORE_API float GetDPIScaleFactor() const;
 
 	/** Overrides the DPI scale factor of the native window */
-	void SetDPIScaleFactor(const float Factor);
+	SLATECORE_API void SetDPIScaleFactor(const float Factor);
 
 	/** Will inform the native window that we need to handle any DPI changes from within the application */
-	void SetManualManageDPIChanges(const bool bManualDPI);
+	SLATECORE_API void SetManualManageDPIChanges(const bool bManualDPI);
 
 	bool IsManualManageDPIChanges() const
 	{
@@ -599,35 +599,35 @@ public:
 	 * @param ParentWindow	The window to check
 	 * @return true if the window is a child of ParentWindow, false otherwise.
 	 */
-	bool IsDescendantOf( const TSharedPtr<SWindow>& ParentWindow ) const;
+	SLATECORE_API bool IsDescendantOf( const TSharedPtr<SWindow>& ParentWindow ) const;
 
 	/**
 	 * Sets the native OS window associated with this SWindow
 	 *
 	 * @param InNativeWindow	The native window
 	 */
-	void SetNativeWindow( TSharedRef<FGenericWindow> InNativeWindow );
+	SLATECORE_API void SetNativeWindow( TSharedRef<FGenericWindow> InNativeWindow );
 	
 	/**
 	 * Sets the widget content for this window
 	 *
 	 * @param	InContent	The widget to use as content for this window
 	 */
-	void SetContent( TSharedRef<SWidget> InContent );
+	SLATECORE_API void SetContent( TSharedRef<SWidget> InContent );
 
 	/**
 	 * Gets the widget content for this window
 	 *
 	 * @return	The widget content for this window
 	 */
-	TSharedRef<SWidget> GetContent();
+	SLATECORE_API TSharedRef<SWidget> GetContent();
 
 	/**
 	 * Check whether we have a full window overlay, used to draw content over the entire window.
 	 * 
 	 * @return true if the window has an overlay
 	 */
-	bool HasOverlay() const;
+	SLATECORE_API bool HasOverlay() const;
 
 	/**
 	 * Adds content to draw on top of the entire window
@@ -635,7 +635,7 @@ public:
 	 * @param	InZOrder	Z-order to use for this widget
 	 * @return The added overlay slot so that it can be configured and populated
 	 */
-	SOverlay::FScopedWidgetSlotArguments AddOverlaySlot( const int32 ZOrder = INDEX_NONE );
+	SLATECORE_API SOverlay::FScopedWidgetSlotArguments AddOverlaySlot( const int32 ZOrder = INDEX_NONE );
 
 	/**
 	 * Removes a widget that is being drawn over the entire window
@@ -643,7 +643,7 @@ public:
 	 * @param	InContent	The widget to remove
 	 * @return	true if successful
 	 */
-	bool RemoveOverlaySlot(const TSharedRef<SWidget>& InContent);
+	SLATECORE_API bool RemoveOverlaySlot(const TSharedRef<SWidget>& InContent);
 
 	/**
 	 * Visualize a new pop-up if possible.  If it's not possible for this widget to host the pop-up
@@ -654,46 +654,46 @@ public:
 	 *
 	 * @return a valid FPopupLayer if this widget supported hosting it.  You can call Remove() on this to destroy the pop-up.
 	 */
-	virtual TSharedPtr<FPopupLayer> OnVisualizePopup(const TSharedRef<SWidget>& PopupContent) override;
+	SLATECORE_API virtual TSharedPtr<FPopupLayer> OnVisualizePopup(const TSharedRef<SWidget>& PopupContent) override;
 
 	using FScopedWidgetSlotArguments = TPanelChildren<FPopupLayerSlot>::FScopedWidgetSlotArguments;
 	/** Return a new slot in the popup layer. Assumes that the window has a popup layer. */
-	FScopedWidgetSlotArguments AddPopupLayerSlot();
+	SLATECORE_API FScopedWidgetSlotArguments AddPopupLayerSlot();
 
 	/** Counterpart to AddPopupLayerSlot */
-	void RemovePopupLayerSlot( const TSharedRef<SWidget>& WidgetToRemove );
+	SLATECORE_API void RemovePopupLayerSlot( const TSharedRef<SWidget>& WidgetToRemove );
 
 	/**
 	 * Sets a widget to use as a full window overlay, or clears an existing widget if set.  When set, this widget will be drawn on top of all other window content.
 	 *
 	 * @param	InContent	The widget to use for full window overlay content, or nullptr for no overlay
 	 */
-	void SetFullWindowOverlayContent( TSharedPtr<SWidget> InContent );
+	SLATECORE_API void SetFullWindowOverlayContent( TSharedPtr<SWidget> InContent );
 
 	/**
 	 * Begins a transition from showing regular window content to overlay content
 	 * During the transition we show both sets of content
 	 */
-	void BeginFullWindowOverlayTransition();
+	SLATECORE_API void BeginFullWindowOverlayTransition();
 
 	/**
 	 * Ends a transition from showing regular window content to overlay content
 	 * When this is called content occluded by the full window overlay(if there is one) will be physically hidden
 	 */
-	void EndFullWindowOverlayTransition();
+	SLATECORE_API void EndFullWindowOverlayTransition();
 
 	/**
 	 * Checks to see if there is content assigned as a full window overlay
 	 *
 	 * @return	True if there is an overlay widget assigned
 	 */
-	bool HasFullWindowOverlayContent() const;
+	SLATECORE_API bool HasFullWindowOverlayContent() const;
 
 	/** Shows or hides native window buttons on platforms that use them */
-	void SetNativeWindowButtonsVisibility(bool bVisible);
+	SLATECORE_API void SetNativeWindowButtonsVisibility(bool bVisible);
 
 	/** @return should this window show up in the taskbar */
-	bool AppearsInTaskbar() const;
+	SLATECORE_API bool AppearsInTaskbar() const;
 
 	/** Gets the multicast delegate executed when the window is deactivated */
 	FOnWindowActivatedEvent& GetOnWindowActivatedEvent() { return WindowActivatedEvent; }
@@ -702,75 +702,75 @@ public:
 	FOnWindowDeactivatedEvent& GetOnWindowDeactivatedEvent() { return WindowDeactivatedEvent; }
 
 	/** Sets the delegate to execute right before the window is closed */
-	void SetOnWindowClosed( const FOnWindowClosed& InDelegate );
+	SLATECORE_API void SetOnWindowClosed( const FOnWindowClosed& InDelegate );
 
 	/** Gets the multicast delegate to execute right before the window is closed */
 	FOnWindowClosedEvent& GetOnWindowClosedEvent() { return WindowClosedEvent; }
 
 	/** Sets the delegate to execute right after the window has been moved */
-	void SetOnWindowMoved( const FOnWindowMoved& InDelegate);
+	SLATECORE_API void SetOnWindowMoved( const FOnWindowMoved& InDelegate);
 
 	/** Sets the delegate to override RequestDestroyWindow */
-	void SetRequestDestroyWindowOverride( const FRequestDestroyWindowOverride& InDelegate );
+	SLATECORE_API void SetRequestDestroyWindowOverride( const FRequestDestroyWindowOverride& InDelegate );
 
 	/** Request that this window be destroyed. The window is not destroyed immediately. Instead it is placed in a queue for destruction on next Tick */
-	void RequestDestroyWindow();
+	SLATECORE_API void RequestDestroyWindow();
 
 	/** Warning: use Request Destroy Window whenever possible!  This method destroys the window immediately! */
-	void DestroyWindowImmediately();
+	SLATECORE_API void DestroyWindowImmediately();
  
 	/** Calls the OnWindowClosed delegate when this window is about to be closed */
-	void NotifyWindowBeingDestroyed();
+	SLATECORE_API void NotifyWindowBeingDestroyed();
 
 	/** Make the window visible */
-	void ShowWindow();
+	SLATECORE_API void ShowWindow();
 
 	/** Make the window invisible */
-	void HideWindow();
+	SLATECORE_API void HideWindow();
 
 	/**
 	 * Enables or disables this window and all of its children
 	 *
 	 * @param bEnable	true to enable this window and its children false to diable this window and its children
 	 */
-	void EnableWindow( bool bEnable );
+	SLATECORE_API void EnableWindow( bool bEnable );
 
 	/** Toggle window between window modes (fullscreen, windowed, etc) */
-	void SetWindowMode( EWindowMode::Type WindowMode );
+	SLATECORE_API void SetWindowMode( EWindowMode::Type WindowMode );
 
 	/** @return The current window mode (fullscreen, windowed, etc) */
 	EWindowMode::Type GetWindowMode() const { return NativeWindow->GetWindowMode(); }
 
 	/** @return true if the window is visible, false otherwise*/
-	bool IsVisible() const;
+	SLATECORE_API bool IsVisible() const;
 
 	/** @return true if the window is maximized, false otherwise*/
-	bool IsWindowMaximized() const;
+	SLATECORE_API bool IsWindowMaximized() const;
 
 	/** @return true of the window is minimized (iconic), false otherwise */
-	bool IsWindowMinimized() const;
+	SLATECORE_API bool IsWindowMinimized() const;
 
 	/** Maximize the window if bInitiallyMaximized is set */
-	void InitialMaximize();
+	SLATECORE_API void InitialMaximize();
 
 	/** Maximize the window if bInitiallyMinimized is set */
-	void InitialMinimize();
+	SLATECORE_API void InitialMinimize();
 
 	/**
 	 * Sets the opacity of this window
 	 *
 	 * @param	InOpacity	The new window opacity represented as a floating point scalar
 	 */
-	void SetOpacity( const float InOpacity );
+	SLATECORE_API void SetOpacity( const float InOpacity );
 
 	/** @return the window's current opacity */
-	float GetOpacity() const;
+	SLATECORE_API float GetOpacity() const;
 
 	/** @return the level of transparency supported by this window */
-	EWindowTransparency GetTransparencySupport() const;
+	SLATECORE_API EWindowTransparency GetTransparencySupport() const;
 
 	/** @return A String representation of the widget */
-	virtual FString ToString() const override;
+	SLATECORE_API virtual FString ToString() const override;
 
 	/**
 	 * Sets a widget that should become focused when this window is next activated
@@ -791,25 +791,25 @@ public:
 	}
 
 	/** @return the window activation policy used when showing the window */
-	EWindowActivationPolicy ActivationPolicy() const;
+	SLATECORE_API EWindowActivationPolicy ActivationPolicy() const;
 
 	/** @return true if the window accepts input; false if the window is non-interactive */
-	bool AcceptsInput() const;
+	SLATECORE_API bool AcceptsInput() const;
 
 	/** @return true if the user decides the size of the window */
-	bool IsUserSized() const;
+	SLATECORE_API bool IsUserSized() const;
 
 	/** @return true if the window is sized by the windows content */
-	bool IsAutosized() const;
+	SLATECORE_API bool IsAutosized() const;
 
 	/** Should this window automatically derive its size based on its content or be user-drive? */
-	void SetSizingRule( ESizingRule InSizingRule );
+	SLATECORE_API void SetSizingRule( ESizingRule InSizingRule );
 
 	/** @return true if this is a vanilla window, or one being used for some special purpose: e.g. tooltip or menu */
-	bool IsRegularWindow() const;
+	SLATECORE_API bool IsRegularWindow() const;
 
 	/** @return true if the window should be on top of all other windows; false otherwise */
-	bool IsTopmostWindow() const;
+	SLATECORE_API bool IsTopmostWindow() const;
 
 	/** @return True if we expect the window size to change frequently. See description of bSizeWillChangeOften member variable. */
 	bool SizeWillChangeOften() const
@@ -838,19 +838,19 @@ public:
 	bool HasOSWindowBorder() const { return bHasOSWindowBorder; }
 
 	/** @return true if mouse coordinates is within this window */
-	bool IsScreenspaceMouseWithin(UE::Slate::FDeprecateVector2DParameter ScreenspaceMouseCoordinate) const;
+	SLATECORE_API bool IsScreenspaceMouseWithin(UE::Slate::FDeprecateVector2DParameter ScreenspaceMouseCoordinate) const;
 
 	/** @return true if this is a user-sized window with a thick edge */
-	bool HasSizingFrame() const;
+	SLATECORE_API bool HasSizingFrame() const;
 
 	/** @return true if this window has a close button/box on the titlebar area */
-	bool HasCloseBox() const;
+	SLATECORE_API bool HasCloseBox() const;
 
 	/** @return true if this window has a maximize button/box on the titlebar area */
-	bool HasMaximizeBox() const;
+	SLATECORE_API bool HasMaximizeBox() const;
 
 	/** @return true if this window has a minimize button/box on the titlebar area */
-	bool HasMinimizeBox() const;
+	SLATECORE_API bool HasMinimizeBox() const;
 
 	/** Set modal window related flags - called by Slate app code during FSlateApplication::AddModalWindow() */
 	void SetAsModalWindow()
@@ -896,18 +896,18 @@ public:
 	}
 
 	// Events
-	virtual FCursorReply OnCursorQuery( const FGeometry& MyGeometry, const FPointerEvent& CursorEvent ) const override;
+	SLATECORE_API virtual FCursorReply OnCursorQuery( const FGeometry& MyGeometry, const FPointerEvent& CursorEvent ) const override;
 
 	/** The system will call this method to notify the window that it has been places in the foreground or background. */
-	virtual bool OnIsActiveChanged( const FWindowActivateEvent& ActivateEvent );
+	SLATECORE_API virtual bool OnIsActiveChanged( const FWindowActivateEvent& ActivateEvent );
 	
 	/** Windows functions */
-	void Maximize();
-	void Restore();
-	void Minimize();
+	SLATECORE_API void Maximize();
+	SLATECORE_API void Restore();
+	SLATECORE_API void Minimize();
 
 	/** Gets the current Window Zone that mouse position is over. */
-	EWindowZone::Type GetCurrentWindowZone(UE::Slate::FDeprecateVector2DParameter LocalMousePosition);
+	SLATECORE_API EWindowZone::Type GetCurrentWindowZone(UE::Slate::FDeprecateVector2DParameter LocalMousePosition);
 
 	/** Used to store the zone where the mouse down event occurred during move / drag */
 	EWindowZone::Type MoveResizeZone;
@@ -915,36 +915,36 @@ public:
 	FSlateRect MoveResizeRect;
 
 	/** @return Gets the radius of the corner rounding of the window. */
-	int32 GetCornerRadius();
+	SLATECORE_API int32 GetCornerRadius();
 
-	virtual bool SupportsKeyboardFocus() const override;
+	SLATECORE_API virtual bool SupportsKeyboardFocus() const override;
 
 	bool IsDrawingEnabled() const { return bIsDrawingEnabled; }
 
 	virtual bool Advanced_IsWindow() const override { return true; }
-	virtual bool Advanced_IsInvalidationRoot() const override;
-	virtual const FSlateInvalidationRoot* Advanced_AsInvalidationRoot() const override;
+	SLATECORE_API virtual bool Advanced_IsInvalidationRoot() const override;
+	SLATECORE_API virtual const FSlateInvalidationRoot* Advanced_AsInvalidationRoot() const override;
 
 #if WITH_ACCESSIBILITY
-	virtual TSharedRef<FSlateAccessibleWidget> CreateAccessibleWidget() override;
-	virtual TOptional<FText> GetDefaultAccessibleText(EAccessibleType AccessibleType = EAccessibleType::Main) const override;
+	SLATECORE_API virtual TSharedRef<FSlateAccessibleWidget> CreateAccessibleWidget() override;
+	SLATECORE_API virtual TOptional<FText> GetDefaultAccessibleText(EAccessibleType AccessibleType = EAccessibleType::Main) const override;
 #endif
 private:
-	virtual FReply OnFocusReceived( const FGeometry& MyGeometry, const FFocusEvent& InFocusEvent ) override;
-	virtual FReply OnMouseButtonDown( const FGeometry& MyGeometry, const FPointerEvent& MouseEvent ) override;
-	virtual FReply OnMouseButtonUp( const FGeometry& MyGeometry, const FPointerEvent& MouseEvent ) override;
-	virtual FReply OnMouseMove( const FGeometry& MyGeometry, const FPointerEvent& MouseEvent ) override;
+	SLATECORE_API virtual FReply OnFocusReceived( const FGeometry& MyGeometry, const FFocusEvent& InFocusEvent ) override;
+	SLATECORE_API virtual FReply OnMouseButtonDown( const FGeometry& MyGeometry, const FPointerEvent& MouseEvent ) override;
+	SLATECORE_API virtual FReply OnMouseButtonUp( const FGeometry& MyGeometry, const FPointerEvent& MouseEvent ) override;
+	SLATECORE_API virtual FReply OnMouseMove( const FGeometry& MyGeometry, const FPointerEvent& MouseEvent ) override;
 
-	virtual int32 OnPaint(const FPaintArgs& Args, const FGeometry& AllottedGeometry, const FSlateRect& MyCullingRect, FSlateWindowElementList& OutDrawElements, int32 LayerId, const FWidgetStyle& InWidgetStyle, bool bParentEnabled) const override;
+	SLATECORE_API virtual int32 OnPaint(const FPaintArgs& Args, const FGeometry& AllottedGeometry, const FSlateRect& MyCullingRect, FSlateWindowElementList& OutDrawElements, int32 LayerId, const FWidgetStyle& InWidgetStyle, bool bParentEnabled) const override;
 
 	/** The window's desired size takes into account the ratio between the slate units and the pixel size */
-	virtual FVector2D ComputeDesiredSize(float) const override;
-	virtual bool ComputeVolatility() const override;
+	SLATECORE_API virtual FVector2D ComputeDesiredSize(float) const override;
+	SLATECORE_API virtual bool ComputeVolatility() const override;
 
 	/** Resize using already dpi scaled window size including borders/title bar */
-	void ResizeWindowSize( FVector2f NewWindowSize );
+	SLATECORE_API void ResizeWindowSize( FVector2f NewWindowSize );
 
-	void OnGlobalInvalidationToggled(bool bGlobalInvalidationEnabled);
+	SLATECORE_API void OnGlobalInvalidationToggled(bool bGlobalInvalidationEnabled);
 public:
 	/**
 	 * For a given client size, calculate the window size required to accommodate any potential non-OS borders and title bars
@@ -952,7 +952,7 @@ public:
 	 * @param InClientSize: Client size with DPI scaling already applied
 	 * @param DPIScale: Scale that will be applied for border and title. When not supplied detects DPIScale using native or initial position.
 	 */
-	UE::Slate::FDeprecateVector2DResult GetWindowSizeFromClientSize(UE::Slate::FDeprecateVector2DParameter InClientSize, TOptional<float> DPIScale = TOptional<float>());
+	SLATECORE_API UE::Slate::FDeprecateVector2DResult GetWindowSizeFromClientSize(UE::Slate::FDeprecateVector2DParameter InClientSize, TOptional<float> DPIScale = TOptional<float>());
 
 	/** @return true if this window will be focused when it is first shown */
 	inline bool IsFocusedInitially() const
@@ -961,25 +961,25 @@ public:
 	}
 
 	/** @return the list of this window's child windows */
-	const TArray< TSharedRef<SWindow> >& GetChildWindows() const;
+	SLATECORE_API const TArray< TSharedRef<SWindow> >& GetChildWindows() const;
 	
 	/** @return the list of this window's child windows */
-	TArray< TSharedRef<SWindow> >& GetChildWindows();
+	SLATECORE_API TArray< TSharedRef<SWindow> >& GetChildWindows();
 
 	/** Add ChildWindow as this window's child */
-	void AddChildWindow( const TSharedRef<SWindow>& ChildWindow );
+	SLATECORE_API void AddChildWindow( const TSharedRef<SWindow>& ChildWindow );
 
 	/** @return the parent of this window; Invalid shared pointer if this window is not a child */
-	TSharedPtr<SWindow> GetParentWindow() const;
+	SLATECORE_API TSharedPtr<SWindow> GetParentWindow() const;
 
 	/** Look up the parent chain until we find the top-level window that owns this window */
-	TSharedPtr<SWindow> GetTopmostAncestor();
+	SLATECORE_API TSharedPtr<SWindow> GetTopmostAncestor();
 
 	/** Remove DescendantToRemove from this window's children or their children. */
-	bool RemoveDescendantWindow( const TSharedRef<SWindow>& DescendantToRemove );
+	SLATECORE_API bool RemoveDescendantWindow( const TSharedRef<SWindow>& DescendantToRemove );
 
 	/** Sets the delegate to call when switching worlds in before ticking,drawing, or sending messages to widgets in this window */
-	void SetOnWorldSwitchHack( FOnSwitchWorldHack& InOnWorldSwitchHack );
+	SLATECORE_API void SetOnWorldSwitchHack( FOnSwitchWorldHack& InOnWorldSwitchHack );
 
 	/**
 	 * Hack to switch worlds
@@ -987,16 +987,16 @@ public:
 	 * @param WorldId: User ID for a world that should be restored or -1 if no restore
 	 * @param The ID of the world restore later
 	 */
-	int32 SwitchWorlds( int32 WorldId ) const;
+	SLATECORE_API int32 SwitchWorlds( int32 WorldId ) const;
 
 	/** Is this window active? */
-	bool IsActive() const;
+	SLATECORE_API bool IsActive() const;
 
 	/** Are any of our child windows active? */
-	bool HasActiveChildren() const;
+	SLATECORE_API bool HasActiveChildren() const;
 
 	/** Are any of our parent windows active? */
-	bool HasActiveParent() const;
+	SLATECORE_API bool HasActiveParent() const;
 
 	/**
 	 * Sets whether or not the viewport size should be driven by the window's size.  If true, the two will be the same.  If false, an independent viewport size can be specified with SetIndependentViewportSize
@@ -1053,33 +1053,33 @@ public:
 	 *
 	 * @see FHittestGrid for more details.
 	 */
-	FHittestGrid& GetHittestGrid();
+	SLATECORE_API FHittestGrid& GetHittestGrid();
 
 	/** Optional constraints on min and max sizes that this window can be. */
-	FWindowSizeLimits GetSizeLimits() const;
+	SLATECORE_API FWindowSizeLimits GetSizeLimits() const;
 
 	/** Set optional constraints on min and max sizes that this window can be. */
-	void SetSizeLimits(const FWindowSizeLimits& InSizeLimits);
+	SLATECORE_API void SetSizeLimits(const FWindowSizeLimits& InSizeLimits);
 
-	void SetAllowFastUpdate(bool bInAllowFastUpdate);
+	SLATECORE_API void SetAllowFastUpdate(bool bInAllowFastUpdate);
 public:
 
 	//~ SWidget overrides
-	virtual void Tick( const FGeometry& AllottedGeometry, const double InCurrentTime, const float InDeltaTime ) override;
+	SLATECORE_API virtual void Tick( const FGeometry& AllottedGeometry, const double InCurrentTime, const float InDeltaTime ) override;
 
 	/** Windows that are not hittestable should not show up in the hittest grid. */
-	EVisibility GetWindowVisibility() const;
+	SLATECORE_API EVisibility GetWindowVisibility() const;
 
 protected:
 	/**Returns swindow title bar widgets. */
-	virtual TSharedRef<SWidget> MakeWindowTitleBar(const TSharedRef<SWindow>& Window, const TSharedPtr<SWidget>& CenterContent, EHorizontalAlignment CenterContentAlignment);
+	SLATECORE_API virtual TSharedRef<SWidget> MakeWindowTitleBar(const TSharedRef<SWindow>& Window, const TSharedPtr<SWidget>& CenterContent, EHorizontalAlignment CenterContentAlignment);
 	/**Returns the alignment type for the titlebar's title text. */
-	virtual EHorizontalAlignment GetTitleAlignment();
+	SLATECORE_API virtual EHorizontalAlignment GetTitleAlignment();
 
 	/** Kick off a morph to whatever the target shape happens to be. */
-	void StartMorph();
+	SLATECORE_API void StartMorph();
 
-	virtual bool CustomPrepass(float LayoutScaleMultiplier) override;
+	SLATECORE_API virtual bool CustomPrepass(float LayoutScaleMultiplier) override;
 protected:
 
 	/** Type of the window */
@@ -1324,23 +1324,23 @@ protected:
 
 protected:
 	
-	void ConstructWindowInternals();
+	SLATECORE_API void ConstructWindowInternals();
 
 	/** One-off active timer to trigger a the morph sequence to play */
-	EActiveTimerReturnType TriggerPlayMorphSequence( double InCurrentTime, float InDeltaTime );
+	SLATECORE_API EActiveTimerReturnType TriggerPlayMorphSequence( double InCurrentTime, float InDeltaTime );
 
-	void SetWindowBackground(const FSlateBrush* InWindowBackground);
+	SLATECORE_API void SetWindowBackground(const FSlateBrush* InWindowBackground);
 
-	void UpdateWindowContentVisibility();
+	SLATECORE_API void UpdateWindowContentVisibility();
 
 	//~ FSlateInvalidationRoot overrides
-	virtual TSharedRef<SWidget> GetRootWidget() override;
-	virtual int32 PaintSlowPath(const FSlateInvalidationContext& InvalidationContext) override;
+	SLATECORE_API virtual TSharedRef<SWidget> GetRootWidget() override;
+	SLATECORE_API virtual int32 PaintSlowPath(const FSlateInvalidationContext& InvalidationContext) override;
 
 public:
 
 	/** Process the invalidation of the widget contained by the window in GlobalInvalidation. */
-	void ProcessWindowInvalidation();
+	SLATECORE_API void ProcessWindowInvalidation();
 
 private:
 
