Fix copy/paste error in code comment documentation.

#rnx

diff --git a/Engine/Source/Runtime/SlateCore/Public/Widgets/SWindow.h b/Engine/Source/Runtime/SlateCore/Public/Widgets/SWindow.h
--- a/Engine/Source/Runtime/SlateCore/Public/Widgets/SWindow.h
+++ b/Engine/Source/Runtime/SlateCore/Public/Widgets/SWindow.h
@@ -695,7 +695,7 @@ public:
 	/** @return should this window show up in the taskbar */
 	SLATECORE_API bool AppearsInTaskbar() const;
 
-	/** Gets the multicast delegate executed when the window is deactivated */
+	/** Gets the multicast delegate executed when the window is activated */
 	FOnWindowActivatedEvent& GetOnWindowActivatedEvent() { return WindowActivatedEvent; }
 
 	/** Gets the multicast delegate executed when the window is deactivated */
