- Added a per project Font display DPI setting, in order to allow using (but not being limited to) the standard 72 DPI used by most creation software (Photoshop, Figma,...). Slate will still work at 96DPI.
- Converted the font size to float, to allow specifying more precise size, and support the sizes that would be wanted at (for instance) 72 DPI, without losing information (ie rounding to the closest 96DPI int).
- Change various font caches behavior (glyph, advance, kerning...) to match behavior of Freetype, and ensure that a given size in Freetype will be stored only once in the font caches.
- Added cvar commands:
    - Slate.DumpFontCacheStats: will display the list of all font combinations used, with the number of glyphs for each combination.
    - Slate.FlushFontCache: will clear the font caches as well as the font atlases, then will invalidate all the widgets to force the cache to be refilled.
- Fixed the combo box for font typeface selection not disabled when necessary.

#tests:
- changed the Font display settings between various values and customDPI, validatet that everything was save correctly, and that the display was reacting to the value change.
- tested all the tooltips, simple, extended (ctrl + alt), and dynamic (display of the current DPI setting next to the font size).
- validated that the font atlas/cache was containing not more than one time each glyph rendered by Freetype.
- validated that the edition of the font size was still fully functionnal (spinbox, undo/redo, support of disable state,...).
- tested default value for a new project (72DPI).
- compiled Slate/SlateCore/UMG with UnsafeTypeCastWarningLevel = WarningLevel.Warning to fix warnings related to the change.
- tested in a game with latest changes from BP team (CL 26034633), ensuring that there are no regression related to the change of type (int to float) from the Font size.
- running with it in a game and UE branch for months.

diff --git a/Engine/Documentation/Source/Shared/Enums/EFontDPI/EFontDPI.INT.udn b/Engine/Documentation/Source/Shared/Enums/EFontDPI/EFontDPI.INT.udn
new file mode 100644
--- /dev/null
+++ b/Engine/Documentation/Source/Shared/Enums/EFontDPI/EFontDPI.INT.udn
@@ -0,0 +1,29 @@
+Availability:NoPublish
+Title:EFontDPI
+Crumbs:
+
+[EXCERPT:Standard]
+When you create visual designs in an external application and realize them later in UMG, your text elements might appear to be the wrong size. This is typically because the applications use different scaling factors to determine the pixel heights of letters based on the point size of the font.
+
+You may be able to fix this discrepancy by adjusting the scaling factor in your project’s Font Resolution setting to match the one used by your design application.
+
+- 72 DPI matches common web-based design tools like Figma. It is the default value for UMG in Unreal Engine 5.2 and later.
+
+- 96 DPI matches the scaling factor used internally by Unreal Engine and Slate. It was also the default value for UMG up until Unreal Engine 5.2. Use it if you need to remain consistent with older projects and workflows, or if you need Unreal Engine artists and programmers to refer to font sizes using consistent values.
+
+- If you use a design application that has a different scaling factor, you can set your own custom Font Resolution value.
+
+Note:
+- This setting only affects the way UMG interprets and displays font sizes. Internally, font sizes are always converted to 96 DPI before saving. This may occasionally result in minor rounding errors for specific combinations of font sizes and resolutions.
+
+- Always use the fewest possible number of different font sizes in your project. This avoids wasting memory in font atlas textures. If you need to change font resolution or font size after you have already created UMG assets with text elements, make sure the font sizes used in those existing assets match the sizes you use for new assets you create.
+[/EXCERPT:Standard]
+
+[EXCERPT:Unreal]
+[VAR:ExcerptAlias]
+Standard
+[/VAR]
+[/EXCERPT:Unreal]
+
+
+
diff --git a/Engine/Documentation/Source/Shared/Types/FSlateFontInfo/FSlateFontInfo.INT.udn b/Engine/Documentation/Source/Shared/Types/FSlateFontInfo/FSlateFontInfo.INT.udn
new file mode 100644
--- /dev/null
+++ b/Engine/Documentation/Source/Shared/Types/FSlateFontInfo/FSlateFontInfo.INT.udn
@@ -0,0 +1,11 @@
+Availability:NoPublish
+Title:FSlateFontInfo
+Crumbs:
+
+[EXCERPT:Size]
+Tips:
+
+- Always use the fewest possible number of different font sizes in your project. This avoids wasting memory in font atlas textures.
+
+- If text elements do not appear at the pixel size you expect, you may need to change your project’s Font Resolution setting. To find this setting, choose Edit > Project Settings from the main menu and search for Font Resolution.
+[/EXCERPT:Size]
diff --git a/Engine/Documentation/Source/Shared/Types/UUserInterfaceSettings/UUserInterfaceSettings.INT.udn b/Engine/Documentation/Source/Shared/Types/UUserInterfaceSettings/UUserInterfaceSettings.INT.udn
new file mode 100644
--- /dev/null
+++ b/Engine/Documentation/Source/Shared/Types/UUserInterfaceSettings/UUserInterfaceSettings.INT.udn
@@ -0,0 +1,32 @@
+Availability:NoPublish
+Title:UUserInterfaceSettings
+Crumbs:
+
+[EXCERPT:FontDPIPreset]
+When you create visual designs in an external application and realize them later in UMG, your text elements might appear to be the wrong size. This is typically because the applications use different scaling factors to determine the pixel heights of letters based on the point size of the font.
+
+You may be able to fix this discrepancy by adjusting the scaling factor in your project’s Font Resolution setting to match the one used by your design application.
+
+- 72 DPI matches common web-based design tools like Figma. It is the default value for UMG in Unreal Engine 5.2 and later.
+
+- 96 DPI matches the scaling factor used internally by Unreal Engine and Slate. It was also the default value for UMG up until Unreal Engine 5.2. Use it if you need to remain consistent with older projects and workflows, or if you need Unreal Engine artists and programmers to refer to font sizes using consistent values.
+
+- If you use a design application that has a different scaling factor, you can set your own custom Font Resolution value.
+
+Note:
+- This setting only affects the way UMG interprets and displays font sizes. Internally, font sizes are always converted to 96 DPI before saving. This may occasionally result in minor rounding errors for specific combinations of font sizes and resolutions.
+
+- Always use the fewest possible number of different font sizes in your project. This avoids wasting memory in font atlas textures. If you need to change font resolution or font size after you have already created UMG assets with text elements, make sure the font sizes used in those existing assets match the sizes you use for new assets you create.
+[/EXCERPT:FontDPIPreset]
+
+[EXCERPT:bUseCustomFontDPI]
+[VAR:ExcerptAlias]
+FontDPIPreset
+[/VAR]
+[/EXCERPT:bUseCustomFontDPI]
+
+[EXCERPT:EFontDPI]
+[VAR:ExcerptAlias]
+FontDPIPreset
+[/VAR]
+[/EXCERPT:EFontDPI]
diff --git a/Engine/Plugins/Runtime/CommonUI/Source/CommonUI/Private/CommonRichTextBlock.cpp b/Engine/Plugins/Runtime/CommonUI/Source/CommonUI/Private/CommonRichTextBlock.cpp
--- a/Engine/Plugins/Runtime/CommonUI/Source/CommonUI/Private/CommonRichTextBlock.cpp
+++ b/Engine/Plugins/Runtime/CommonUI/Source/CommonUI/Private/CommonRichTextBlock.cpp
@@ -101,7 +101,8 @@ public:
 			}
 			else if (Tag == SupportedMarkupKeys::Size)
 			{
-				int32 FontSize = FCString::Atoi(*Value);
+				float FontSize;
+				LexFromString(FontSize, *Value);
 				TextStyle.SetFontSize(FontSize);
 			}
 			else if (Tag == SupportedMarkupKeys::Case)
diff --git a/Engine/Source/Editor/DetailCustomizations/Private/SlateFontInfoCustomization.cpp b/Engine/Source/Editor/DetailCustomizations/Private/SlateFontInfoCustomization.cpp
--- a/Engine/Source/Editor/DetailCustomizations/Private/SlateFontInfoCustomization.cpp
+++ b/Engine/Source/Editor/DetailCustomizations/Private/SlateFontInfoCustomization.cpp
@@ -6,12 +6,15 @@
 #include "Containers/UnrealString.h"
 #include "DetailLayoutBuilder.h"
 #include "DetailWidgetRow.h"
+#include "Editor.h"
 #include "Engine/Font.h"
+#include "Engine/UserInterfaceSettings.h"
 #include "Fonts/CompositeFont.h"
 #include "Fonts/SlateFontInfo.h"
 #include "HAL/Platform.h"
 #include "IDetailChildrenBuilder.h"
 #include "IDetailPropertyRow.h"
+#include "IDocumentation.h"
 #include "Logging/LogCategory.h"
 #include "Logging/LogMacros.h"
 #include "Misc/AssertionMacros.h"
@@ -25,8 +28,12 @@
 #include "UObject/Object.h"
 #include "UObject/UnrealNames.h"
 #include "Widgets/DeclarativeSyntaxSupport.h"
+#include "Widgets/Input/SNumericEntryBox.h"
+#include "Widgets/SToolTip.h"
 #include "Widgets/Text/STextBlock.h"
 
+#include "Internationalization/FastDecimalFormat.h"
+
 class SWidget;
 
 #define LOCTEXT_NAMESPACE "SlateFontInfo"
@@ -103,6 +110,7 @@ void FSlateFontInfoStructCustomization::CustomizeChildren(TSharedRef<IPropertyHa
 		.OnComboBoxOpening(this, &FSlateFontInfoStructCustomization::OnFontEntryComboOpening)
 		.OnSelectionChanged(this, &FSlateFontInfoStructCustomization::OnFontEntrySelectionChanged)
 		.OnGenerateWidget(this, &FSlateFontInfoStructCustomization::MakeFontEntryWidget)
+		.IsEnabled(!TypefaceFontNameProperty->IsEditConst())
 		[
 			SNew(STextBlock)
 			.Text(this, &FSlateFontInfoStructCustomization::GetFontEntryComboText)
@@ -110,7 +118,7 @@ void FSlateFontInfoStructCustomization::CustomizeChildren(TSharedRef<IPropertyHa
 		]
 	];
 
-	InStructBuilder.AddProperty(FontSizeProperty.ToSharedRef());
+	AddFontSizeProperty(InStructBuilder);
 
 	InStructBuilder.AddProperty(InStructPropertyHandle->GetChildHandle(GET_MEMBER_NAME_CHECKED(FSlateFontInfo, LetterSpacing)).ToSharedRef());
 	
@@ -121,6 +129,167 @@ void FSlateFontInfoStructCustomization::CustomizeChildren(TSharedRef<IPropertyHa
 	InStructBuilder.AddProperty(InStructPropertyHandle->GetChildHandle(GET_MEMBER_NAME_CHECKED(FSlateFontInfo, OutlineSettings)).ToSharedRef());
 }
 
+void FSlateFontInfoStructCustomization::AddFontSizeProperty(IDetailChildrenBuilder& InStructBuilder)
+{
+	TSharedRef<IPropertyHandle> FontSizePropertyRef = FontSizeProperty.ToSharedRef();
+
+	auto GetFloatMetaDataFromKey = [FontSizePropertyRef](const FName& Key) -> const TOptional<float>
+	{
+		const FString* InstanceValue = FontSizePropertyRef->GetInstanceMetaData(Key);
+		const FString MetaDataValueString = (InstanceValue != nullptr) ? *InstanceValue : FontSizePropertyRef->GetMetaData(Key);
+		if (MetaDataValueString.Len())
+		{
+			float FloatValue;
+			LexFromString(FloatValue, *MetaDataValueString);
+			return FloatValue;
+		}
+
+		return TOptional<float>();
+	};
+
+	TOptional<float> MinValue = GetFloatMetaDataFromKey("ClampMin");
+	TOptional<float> MaxValue = GetFloatMetaDataFromKey("ClampMax");
+
+	IDetailPropertyRow& FontSizeRow = InStructBuilder.AddProperty(FontSizePropertyRef);
+	FontSizeRow.CustomWidget()
+	.NameContent()
+	[
+		FontSizePropertyRef->CreatePropertyNameWidget()
+	]
+	.ValueContent()
+	[
+		SNew(SNumericEntryBox<float>)
+		.Value(this, &FSlateFontInfoStructCustomization::OnFontSizeGetValue)
+		.OnValueChanged(this, &FSlateFontInfoStructCustomization::OnFontSizeValueChanged)
+		.OnValueCommitted(this, &FSlateFontInfoStructCustomization::OnFontSizeValueCommitted)
+		.UndeterminedString(LOCTEXT("MultipleValues", "Multiple Values"))
+		.OnBeginSliderMovement(this, &FSlateFontInfoStructCustomization::OnFontSizeBeginSliderMovement)
+		.OnEndSliderMovement(this, &FSlateFontInfoStructCustomization::OnFontSizeEndSliderMovement)
+		.Font(FAppStyle::GetFontStyle("PropertyWindow.NormalFont"))
+		.MinFractionalDigits(0)
+		.MaxFractionalDigits(2)
+		.MinValue(MinValue)
+		.MaxValue(MaxValue)
+		.MinSliderValue(MinValue)
+		.MaxSliderValue(MaxValue)
+		.Delta(1.0f)
+		.AllowWheel(true)
+		.WheelStep(1.0f)
+		.AllowSpin(FontSizePropertyRef->GetNumPerObjectValues() == 1) //Don't allow spin for multiple value select. Allowing it would result in the widget background not being displayed.
+		.IsEnabled(!FontSizePropertyRef->IsEditConst())
+		.ToolTip(IDocumentation::Get()->CreateToolTip(TAttribute<FText>(this, &FSlateFontInfoStructCustomization::GetFontSizeTooltipText),
+													  nullptr,
+													  TEXT("Shared/Types/FSlateFontInfo"),
+													  TEXT("Size")))
+	];
+}
+
+FText FSlateFontInfoStructCustomization::GetFontSizeTooltipText() const
+{
+	const UUserInterfaceSettings* UISettings = GetDefault<UUserInterfaceSettings>();
+
+	FFormatNamedArguments Args;
+	Args.Add(TEXT("DPI"), UISettings->GetFontDPIDisplayString());
+	return FText::Format(LOCTEXT("FontSizeToolTip", "Size of the font in points.\nCurrent font resolution : {DPI}"), Args);
+}
+
+float FSlateFontInfoStructCustomization::ConvertFontSizeFromNativeToDisplay(float NativeFontSize)
+{
+	const UUserInterfaceSettings* UISettings = GetDefault<UUserInterfaceSettings>();
+	const float FontDisplayDPI = UISettings->GetFontDisplayDPI();
+	const float DisplayedSize = NativeFontSize * static_cast<float>(FontConstants::RenderDPI) / FontDisplayDPI;
+	return DisplayedSize;
+}
+
+float FSlateFontInfoStructCustomization::ConvertFontSizeFromDisplayToNative(float DisplayFontSize)
+{
+	const UUserInterfaceSettings* UISettings = GetDefault<UUserInterfaceSettings>();
+	const float FontDisplayDPI = UISettings->GetFontDisplayDPI();
+	const float NativeSize = DisplayFontSize * FontDisplayDPI / static_cast<float>(FontConstants::RenderDPI);
+	const float RoundedSize = FMath::GridSnap(NativeSize, 0.01f);
+	return RoundedSize;
+}
+
+TOptional<float> FSlateFontInfoStructCustomization::OnFontSizeGetValue() const
+{
+	float value(0.0f);
+	const TSharedRef< IPropertyHandle > PropertyHandle = FontSizeProperty.ToSharedRef();
+
+	if (PropertyHandle->GetValue(value) == FPropertyAccess::Success)
+	{
+		return ConvertFontSizeFromNativeToDisplay(value);
+	}
+
+	// Return an unset value so it displays the "multiple values" indicator instead
+	return TOptional<float>();
+}
+
+void FSlateFontInfoStructCustomization::OnFontSizeValueChanged(float NewDisplayValue)
+{
+	if (!bIsUsingSlider)
+		return;
+
+	const TSharedRef< IPropertyHandle > PropertyHandle = FontSizeProperty.ToSharedRef();
+
+	float OrgValue(0.0f);
+	const float NativeFontSize = ConvertFontSizeFromDisplayToNative(NewDisplayValue);
+
+	if (PropertyHandle->GetValue(OrgValue) != FPropertyAccess::Fail)
+	{
+		// Value hasn't changed, so lets return now
+		if (OrgValue == NativeFontSize)
+		{
+			return;
+		}
+	}
+
+	// We don't create a transaction for each property change when using the slider. Only once when the slider first is moved
+	EPropertyValueSetFlags::Type Flags = (EPropertyValueSetFlags::InteractiveChange | EPropertyValueSetFlags::NotTransactable);
+	PropertyHandle->SetValue(NativeFontSize, Flags);
+}
+
+void FSlateFontInfoStructCustomization::OnFontSizeValueCommitted(float NewDisplayValue, ETextCommit::Type CommitInfo)
+{
+	const TSharedRef< IPropertyHandle > PropertyHandle = FontSizeProperty.ToSharedRef();
+	const float NativeFontSize = ConvertFontSizeFromDisplayToNative(NewDisplayValue);
+
+	float OrgValue(0.0f);
+	if (bIsUsingSlider || (PropertyHandle->GetValue(OrgValue) == FPropertyAccess::Fail || OrgValue != NativeFontSize))
+	{
+
+		PropertyHandle->SetValue(NativeFontSize);
+		LastSliderFontSizeCommittedValue = NativeFontSize;
+	}
+}
+
+void FSlateFontInfoStructCustomization::OnFontSizeBeginSliderMovement()
+{
+	bIsUsingSlider = true;
+
+	const TSharedRef< IPropertyHandle > PropertyHandle = FontSizeProperty.ToSharedRef();
+	PropertyHandle->GetValue(LastSliderFontSizeCommittedValue);
+
+	GEditor->BeginTransaction(LOCTEXT("UpdateFontSizeTransaction", "Edit font size"));
+}
+
+void FSlateFontInfoStructCustomization::OnFontSizeEndSliderMovement(float NewDisplayValue)
+{
+	const float NativeFontSize = ConvertFontSizeFromDisplayToNative(NewDisplayValue);
+	bIsUsingSlider = false;
+
+	// When the slider end, we may have not called SetValue(NewValue) without the InteractiveChange|NotTransactable flags.
+	//That prevents some transaction and callback to be triggered like the NotifyHook.
+	if (LastSliderFontSizeCommittedValue != NativeFontSize)
+	{
+		const TSharedRef< IPropertyHandle > PropertyHandle = FontSizeProperty.ToSharedRef();
+		PropertyHandle->SetValue(NativeFontSize);
+	}
+	else
+	{
+		GEditor->EndTransaction();
+	}
+}
+
 bool FSlateFontInfoStructCustomization::OnFilterFontAsset(const FAssetData& InAssetData)
 {
 	// We want to filter font assets that aren't valid to use with Slate/UMG
diff --git a/Engine/Source/Editor/DetailCustomizations/Public/Customizations/SlateFontInfoCustomization.h b/Engine/Source/Editor/DetailCustomizations/Public/Customizations/SlateFontInfoCustomization.h
--- a/Engine/Source/Editor/DetailCustomizations/Public/Customizations/SlateFontInfoCustomization.h
+++ b/Engine/Source/Editor/DetailCustomizations/Public/Customizations/SlateFontInfoCustomization.h
@@ -37,6 +37,29 @@ public:
 	virtual void CustomizeChildren(TSharedRef<class IPropertyHandle> InStructPropertyHandle, class IDetailChildrenBuilder& InStructBuilder, IPropertyTypeCustomizationUtils& InStructCustomizationUtils) override;
 
 protected:
+	void AddFontSizeProperty(IDetailChildrenBuilder& InStructBuilder);
+
+	/** @return The value or unset if properties with multiple values are viewed */
+	TOptional<float> OnFontSizeGetValue() const;
+	void OnFontSizeValueChanged(float NewDisplayValue);
+	void OnFontSizeValueCommitted(float NewDisplayValue, ETextCommit::Type CommitInfo);
+
+	/**
+	 * Called when the slider begins to move.  We create a transaction here to undo the property
+	 */
+	void OnFontSizeBeginSliderMovement();
+
+	/**
+	 * Called when the slider stops moving.  We end the previously created transaction
+	 */
+	void OnFontSizeEndSliderMovement(float NewDisplayValue);
+
+	/** @return a dynamic text explaining what the size does and showing the current Font DPI setting */
+	FText GetFontSizeTooltipText() const;
+
+	static float ConvertFontSizeFromNativeToDisplay(float FontSize);
+	static float ConvertFontSizeFromDisplayToNative(float FontSize);
+
 	/** Called to filter out invalid font assets */
 	static bool OnFilterFontAsset(const FAssetData& InAssetData);
 
@@ -82,4 +105,11 @@ protected:
 
 	/** Source data for the font entry combo widget */
 	TArray<TSharedPtr<FName>> FontEntryComboData;
+
+	/** True if the slider is being used to change the value of the property */
+	bool bIsUsingSlider = false;
+
+	/** When using the slider, what was the last committed value */
+	float LastSliderFontSizeCommittedValue = 0.0f;
+
 };
diff --git a/Engine/Source/Editor/GameProjectGeneration/Private/GameProjectUtils.cpp b/Engine/Source/Editor/GameProjectGeneration/Private/GameProjectUtils.cpp
--- a/Engine/Source/Editor/GameProjectGeneration/Private/GameProjectUtils.cpp
+++ b/Engine/Source/Editor/GameProjectGeneration/Private/GameProjectUtils.cpp
@@ -211,6 +211,18 @@ namespace
 				TEXT("bAuthorizeAutomaticWidgetVariableCreation"),
 				TEXT("False"),
 				true /* ShouldReplaceExistingValue */);
+
+			ConfigValues.Emplace(TEXT("DefaultEngine.ini"),
+				TEXT("/Script/Engine.UserInterfaceSettings"),
+				TEXT("FontDPIPreset"),
+				TEXT("Standard"),
+				true /* ShouldReplaceExistingValue */);
+
+			ConfigValues.Emplace(TEXT("DefaultEngine.ini"),
+				TEXT("/Script/Engine.UserInterfaceSettings"),
+				TEXT("FontDPI"),
+				TEXT("72"),
+				true /* ShouldReplaceExistingValue */);
 		}
 	}
 } // namespace <>
diff --git a/Engine/Source/Editor/TranslationEditor/Private/TranslationEditor.cpp b/Engine/Source/Editor/TranslationEditor/Private/TranslationEditor.cpp
--- a/Engine/Source/Editor/TranslationEditor/Private/TranslationEditor.cpp
+++ b/Engine/Source/Editor/TranslationEditor/Private/TranslationEditor.cpp
@@ -846,7 +846,7 @@ void FTranslationEditor::ChangeSourceFont()
 
 	if ( bOpened && NewFontFilename != "")
 	{
-		SourceFont = FSlateFontInfo(NewFontFilename, IntCastChecked<uint16>(SourceFont.Size));
+		SourceFont = FSlateFontInfo(NewFontFilename, SourceFont.Size);
 		RefreshUI();
 	}
 }
@@ -861,7 +861,7 @@ void FTranslationEditor::ChangeTranslationTargetFont()
 
 	if ( bOpened && NewFontFilename != "")
 	{
-		TranslationTargetFont = FSlateFontInfo(NewFontFilename, IntCastChecked<uint16>(TranslationTargetFont.Size));
+		TranslationTargetFont = FSlateFontInfo(NewFontFilename, TranslationTargetFont.Size);
 		RefreshUI();
 	}
 }
diff --git a/Engine/Source/Runtime/Engine/Classes/Engine/Font.h b/Engine/Source/Runtime/Engine/Classes/Engine/Font.h
--- a/Engine/Source/Runtime/Engine/Classes/Engine/Font.h
+++ b/Engine/Source/Runtime/Engine/Classes/Engine/Font.h
@@ -161,7 +161,7 @@ class UFont : public UObject, public IFontProviderInterface
 	float ScalingFactor;
 
 	/** The default size of the font used for legacy Canvas APIs that don't specify a font size */
-	UPROPERTY(EditAnywhere, Category=RuntimeFont)
+	UPROPERTY(EditAnywhere, Category=RuntimeFont, meta = (ClampMin = 1, ClampMax = 1000))
 	int32 LegacyFontSize;
 
 	/** The default font name to use for legacy Canvas APIs that don't specify a font name */
@@ -189,7 +189,7 @@ public:
 	/** Get the info needed to use this UFont with Slate, using the fallback data for legacy Canvas APIs */
 	FORCEINLINE FSlateFontInfo GetLegacySlateFontInfo() const
 	{
-		return FSlateFontInfo(this, LegacyFontSize, LegacyFontName);
+		return FSlateFontInfo(this, static_cast<float>(LegacyFontSize), LegacyFontName);
 	}
 
 	/**
diff --git a/Engine/Source/Runtime/Engine/Classes/Engine/UserInterfaceSettings.h b/Engine/Source/Runtime/Engine/Classes/Engine/UserInterfaceSettings.h
--- a/Engine/Source/Runtime/Engine/Classes/Engine/UserInterfaceSettings.h
+++ b/Engine/Source/Runtime/Engine/Classes/Engine/UserInterfaceSettings.h
@@ -45,6 +45,18 @@ enum class EUIScalingRule : uint8
 	Custom
 };
 
+/** The most used DPI value. */
+UENUM()
+enum class EFontDPI : uint8
+{
+	/** Best for working with with web-based design tools like Figma. */
+	Standard UMETA(DisplayName = "72 DPI (Default)"),
+	/** Resolution used internally by Unreal Engine. */
+	Unreal UMETA(DisplayName = "96 DPI (Unreal Engine)"),
+
+	Custom UMETA(Hidden)
+};
+
 class UWidget;
 class UDPICustomScalingRule;
 
@@ -212,12 +224,48 @@ public:
 	ENGINE_API float GetDPIScaleBasedOnSize(FIntPoint Size) const;
 
 #if WITH_EDITOR
+	/**
+	 * Gets the string to use when we need detailed info about the DPI.
+	 * It will take into account the defined presets to give more context.
+	 */
+	ENGINE_API FText GetFontDPIDisplayString() const;
+
+	/** Gets the Font DPI that should be used for display */
+	ENGINE_API uint32 GetFontDisplayDPI() const;
+
 	ENGINE_API virtual void PostEditChangeProperty(struct FPropertyChangedEvent& PropertyChangedEvent) override;
 #endif
 
 private:
 	float CalculateScale(FIntPoint Size, bool& bError) const;
 
+#if WITH_EDITOR
+	/** Convert from an int to an entry in the EFontDPI enum. Will return the 'Unset' entry if no match is found */
+	static constexpr EFontDPI ConvertToEFontDPI(uint32 inFontDPI);
+	/** Convert from an entry in the EFontDPI enum to an int */
+	static constexpr uint32 ConvertToFontDPI(EFontDPI inFontDPIEntry);
+#endif
+
+#if WITH_EDITORONLY_DATA
+	/**
+	 * Controls the relationship between UMG font size and pixel height.
+	 */
+	UPROPERTY(config, EditAnywhere, Category = "UMG Fonts", meta = (DisplayName = "Font Resolution", EditCondition = "bUseCustomFontDPI", EditConditionHides, ClampMin = "1", ClampMax = "1000"))
+	uint32 CustomFontDPI;
+
+	/**
+	 * Controls the relationship between UMG font size and pixel height.
+	 */
+	UPROPERTY(config, EditAnywhere, Category = "UMG Fonts", meta = (DisplayName = "Font Resolution", EditCondition = "!bUseCustomFontDPI", EditConditionHides))
+	EFontDPI FontDPIPreset;
+
+	/**
+	 * To set your own custom value, check this box, then enter the value in the text box.
+	 */
+	UPROPERTY(config, EditAnywhere, Category = "UMG Fonts", meta = (DisplayName = "Use Custom DPI"))
+	bool bUseCustomFontDPI;
+#endif
+
 	UPROPERTY(Transient)
 	TArray<TObjectPtr<UObject>> CursorClasses;
 
diff --git a/Engine/Source/Runtime/Engine/Classes/GameFramework/HUD.h b/Engine/Source/Runtime/Engine/Classes/GameFramework/HUD.h
--- a/Engine/Source/Runtime/Engine/Classes/GameFramework/HUD.h
+++ b/Engine/Source/Runtime/Engine/Classes/GameFramework/HUD.h
@@ -507,7 +507,7 @@ public:
 	//=============================================================================
 
 	/** @return UFont* from given FontSize index */
-	ENGINE_API virtual  UFont* GetFontFromSizeIndex(int32 FontSize) const;
+	ENGINE_API virtual  UFont* GetFontFromSizeIndex(int32 FontSizeIndex) const;
 
 	/**
 	 *	Pauses or unpauses the game due to main window's focus being lost.
diff --git a/Engine/Source/Runtime/Engine/Private/UserInterfaceSettings.cpp b/Engine/Source/Runtime/Engine/Private/UserInterfaceSettings.cpp
--- a/Engine/Source/Runtime/Engine/Private/UserInterfaceSettings.cpp
+++ b/Engine/Source/Runtime/Engine/Private/UserInterfaceSettings.cpp
@@ -16,6 +16,11 @@ UUserInterfaceSettings::UUserInterfaceSettings(const FObjectInitializer& ObjectI
 	, ApplicationScale(1)
 	, bLoadWidgetsOnDedicatedServer(true)
 	, bAuthorizeAutomaticWidgetVariableCreation(true)
+#if WITH_EDITORONLY_DATA
+	, CustomFontDPI(FontConstants::RenderDPI)
+	, FontDPIPreset(ConvertToEFontDPI(CustomFontDPI))
+	, bUseCustomFontDPI(false)
+#endif
 {
 	SectionName = TEXT("UI");
 }
@@ -155,13 +160,61 @@ float UUserInterfaceSettings::CalculateScale(FIntPoint Size, bool& bError) const
 }
 
 #if WITH_EDITOR
+FText UUserInterfaceSettings::GetFontDPIDisplayString() const
+{
+	const uint32 DisplayDPI = GetFontDisplayDPI();
+
+	//Wether the Display DPI is from a custom value or preset, try to find the matching preset.
+	const EFontDPI Preset = ConvertToEFontDPI(GetFontDisplayDPI());
+	if (Preset != EFontDPI::Custom)
+	{
+		const UEnum* Enum = StaticEnum<EFontDPI>();
+		return Enum->GetDisplayNameTextByValue(static_cast<int64>(Preset));
+	}
+
+	//If no preset found, just use the value.
+	FFormatNamedArguments Args;
+	Args.Add(TEXT("DPI"), FText::AsNumber(DisplayDPI));
+	return FText::Format(LOCTEXT("DPI","{DPI} DPI"), Args);
+}
+
+uint32 UUserInterfaceSettings::GetFontDisplayDPI() const
+{
+	if (bUseCustomFontDPI)
+	{
+		return CustomFontDPI;
+	}
+	return ConvertToFontDPI(FontDPIPreset);
+}
+
+constexpr EFontDPI UUserInterfaceSettings::ConvertToEFontDPI(uint32 inFontDPI)
+{
+	switch (inFontDPI)
+	{
+		case 72: return EFontDPI::Standard;
+		case 96: return EFontDPI::Unreal;
+		default: return EFontDPI::Custom;
+	}
+}
+
+constexpr uint32 UUserInterfaceSettings::ConvertToFontDPI(EFontDPI inFontDPIEntry)
+{
+	switch (inFontDPIEntry)
+	{
+		case EFontDPI::Unreal:
+			return 96;
+		case EFontDPI::Standard:
+		case EFontDPI::Custom:
+		default:
+			return 72;
+	}
+}
 
 void UUserInterfaceSettings::PostEditChangeProperty(FPropertyChangedEvent& PropertyChangedEvent)
 {
 	Super::PostEditChangeProperty(PropertyChangedEvent);
 
 	// Look for changed properties
-	const FName PropertyName = (PropertyChangedEvent.Property != nullptr) ? PropertyChangedEvent.Property->GetFName() : NAME_None;
 	const FName MemberPropertyName = (PropertyChangedEvent.MemberProperty != nullptr) ? PropertyChangedEvent.MemberProperty->GetFName() : NAME_None;
 
 	if (MemberPropertyName == GET_MEMBER_NAME_CHECKED(UUserInterfaceSettings, CustomScalingRuleClass))
diff --git a/Engine/Source/Runtime/SlateCore/Private/Fonts/FontCache.cpp b/Engine/Source/Runtime/SlateCore/Private/Fonts/FontCache.cpp
--- a/Engine/Source/Runtime/SlateCore/Private/Fonts/FontCache.cpp
+++ b/Engine/Source/Runtime/SlateCore/Private/Fonts/FontCache.cpp
@@ -55,6 +55,20 @@ static TAutoConsoleVariable<int32> CVarDefaultTextShapingMethod(
 	ECVF_Default
 	);
 
+#if !UE_BUILD_SHIPPING
+static bool bDumpFontCacheStats = false;
+FAutoConsoleVariableRef CVarDumpFontCacheStats(
+	TEXT("Slate.DumpFontCacheStats"),
+	bDumpFontCacheStats,
+	TEXT("Dump statistics about font cache usage."));
+
+static bool bFlushFontCache = false;
+FAutoConsoleVariableRef CVarFlushFontCache(
+	TEXT("Slate.FlushFontCache"),
+	bFlushFontCache,
+	TEXT("Flush the font cache."));
+#endif
+
 ETextShapingMethod GetDefaultTextShapingMethod()
 {
 	const int32 DefaultTextShapingMethodAsInt = CVarDefaultTextShapingMethod.AsVariable()->GetInt();
@@ -88,19 +102,22 @@ float FShapedGlyphEntry::GetBitmapRenderScale() const
 
 FShapedGlyphEntryKey::FShapedGlyphEntryKey(const FShapedGlyphFaceData& InFontFaceData, uint32 InGlyphIndex, const FFontOutlineSettings& InOutlineSettings)
 	: FontFace(InFontFaceData.FontFace)
-	, FontSize(InFontFaceData.FontSize)
-	, OutlineSize((float)InOutlineSettings.OutlineSize)
 	, OutlineSeparateFillAlpha(InOutlineSettings.bSeparateFillAlpha)
-	, FontScale(InFontFaceData.FontScale)
 	, GlyphIndex(InGlyphIndex)
 	, KeyHash(0)
 	, FontSkew(InFontFaceData.FontSkew)
 {
+#if WITH_FREETYPE
+	FontRenderSize = FreeTypeUtils::ComputeFontPixelSize(InFontFaceData.FontSize, InFontFaceData.FontScale);
+#else
+	FontRenderSize = 0;
+#endif
+	OutlineRenderSize = FMath::TruncToInt(FMath::RoundToFloat(InOutlineSettings.OutlineSize * InFontFaceData.FontScale));
+
 	KeyHash = HashCombine(KeyHash, GetTypeHash(FontFace));
-	KeyHash = HashCombine(KeyHash, GetTypeHash(FontSize));
-	KeyHash = HashCombine(KeyHash, GetTypeHash(OutlineSize));
+	KeyHash = HashCombine(KeyHash, GetTypeHash(FontRenderSize));
+	KeyHash = HashCombine(KeyHash, GetTypeHash(OutlineRenderSize));
 	KeyHash = HashCombine(KeyHash, GetTypeHash(OutlineSeparateFillAlpha));
-	KeyHash = HashCombine(KeyHash, GetTypeHash(FontScale));
 	KeyHash = HashCombine(KeyHash, GetTypeHash(GlyphIndex));
 	KeyHash = HashCombine(KeyHash, GetTypeHash(FontSkew));
 }
@@ -135,7 +152,7 @@ FShapedGlyphSequence::FShapedGlyphSequence(TArray<FShapedGlyphEntry> InGlyphsToR
 		if (!SourceIndexToGlyphData)
 		{
 			// Track reverse look-up data
-			UE_LOG(LogSlate, Warning, TEXT("No Gylph! Index %i. Valid %i, %i, Valid Glyph %i, Visible %i, Num %i, Grapheme %i, Dir %u"),
+			UE_LOG(LogSlate, Warning, TEXT("No Glyph! Index %i. Valid %i, %i, Valid Glyph %i, Visible %i, Num %i, Grapheme %i, Dir %u"),
 				CurrentGlyph.SourceIndex,
 				SourceIndicesToGlyphData.GetSourceTextStartIndex(),
 				SourceIndicesToGlyphData.GetSourceTextEndIndex(),
@@ -1041,10 +1058,9 @@ FShapedGlyphFontAtlasData FSlateFontCache::GetShapedGlyphFontAtlasData( const FS
 	}
 
 	{
-
 		QUICK_SCOPE_CYCLE_COUNTER(STAT_SlateFontCacheAddNewShapedEntry)
 
-			// Not cached at all... create a new entry
+		// Not cached at all... create a new entry
 		TSharedRef<FShapedGlyphFontAtlasData> NewAtlasData = MakeShareable(new FShapedGlyphFontAtlasData());
 		AddNewEntry(InShapedGlyph, InOutlineSettings, *NewAtlasData);
 
@@ -1098,7 +1114,7 @@ void FSlateFontCache::GetStrikeMetrics( const FSlateFontInfo& InFontInfo, const
 	FontRenderer->GetStrikeMetrics(InFontInfo, FontScale, OutStrikeLinePos, OutStrikeLineThickness);
 }
 
-int8 FSlateFontCache::GetKerning( const FFontData& InFontData, const int32 InSize, TCHAR First, TCHAR Second, float Scale ) const
+int8 FSlateFontCache::GetKerning( const FFontData& InFontData, const float InSize, TCHAR First, TCHAR Second, float Scale ) const
 {
 	return FontRenderer->GetKerning(InFontData, InSize, First, Second, Scale);
 }
@@ -1157,14 +1173,151 @@ void FSlateFontCache::FlushCompositeFont(const FCompositeFont& InCompositeFont)
 	CompositeFontCache->FlushCompositeFont(InCompositeFont);
 }
 
+#if (!UE_BUILD_SHIPPING && WITH_FREETYPE)
+struct FFontCacheStatsKey
+{
+	friend inline uint32 GetTypeHash(const FFontCacheStatsKey& Key)
+	{
+		return Key.KeyHash;
+	}
+
+	FFontCacheStatsKey(const FShapedGlyphEntryKey& ShapedGlyphKey)
+		: FontFace(ShapedGlyphKey.FontFace)
+		, FontRenderSize(ShapedGlyphKey.FontRenderSize)
+		, OutlineRenderSize(ShapedGlyphKey.OutlineRenderSize)
+		, OutlineSeparateFillAlpha(ShapedGlyphKey.OutlineSeparateFillAlpha)
+		, KeyHash(0)
+		, FontSkew(ShapedGlyphKey.FontSkew)
+	{
+
+		KeyHash = HashCombine(KeyHash, GetTypeHash(FontFace));
+		KeyHash = HashCombine(KeyHash, GetTypeHash(FontRenderSize));
+		KeyHash = HashCombine(KeyHash, GetTypeHash(OutlineRenderSize));
+		KeyHash = HashCombine(KeyHash, GetTypeHash(OutlineSeparateFillAlpha));
+		KeyHash = HashCombine(KeyHash, GetTypeHash(FontSkew));
+	}
+
+	bool operator==(const FFontCacheStatsKey& Other) const
+	{
+		return FontFace == Other.FontFace
+			&& FontRenderSize == Other.FontRenderSize
+			&& OutlineRenderSize == Other.OutlineRenderSize
+			&& OutlineSeparateFillAlpha == Other.OutlineSeparateFillAlpha
+			&& FontSkew == Other.FontSkew;
+	}
+
+	bool operator!=(const FFontCacheStatsKey& Other) const
+	{
+		return !(*this == Other);
+	}
+
+	/** Weak pointer to the FreeType face to render with */
+	TWeakPtr<FFreeTypeFace> FontFace;
+	/** The size in pixel that Freetype will actually use to render the font, scale included */
+	uint32 FontRenderSize;
+	/** The size in pixels of the outline to render, scale included */
+	uint32 OutlineRenderSize;
+	/** If checked, the outline will be completely translucent where the filled area will be. @see FFontOutlineSettings */
+	bool OutlineSeparateFillAlpha;
+	/** Cached hash value used for map lookups */
+	uint32 KeyHash;
+	/** The skew transform amount for the rendered font */
+	float FontSkew;
+};
+
+void FSlateFontCache::ConditionalDumpFontCacheStats() const
+{
+	if (bDumpFontCacheStats)
+	{
+		TMap<FFontCacheStatsKey, int> Stats;
+
+		// Prepare statistics.
+		int GlyphCount = 0;
+		int PixelsUsedInAtlasCount = 0;
+		for (const TPair<FShapedGlyphEntryKey, TSharedRef<FShapedGlyphFontAtlasData>>& Kvp : ShapedGlyphToAtlasData)
+		{
+			const FShapedGlyphFontAtlasData& FontAtlasData = Kvp.Value.Get();
+			PixelsUsedInAtlasCount += (FontAtlasData.USize * FontAtlasData.VSize);
+			GlyphCount++;
+
+			FFontCacheStatsKey Key(Kvp.Key);
+			FFreeTypeFace* Face = Key.FontFace.Pin().Get();
+			const bool bIsEntryValid = Face != nullptr && Face->GetFace() != nullptr;
+			if (bIsEntryValid) //Only add entries that we'll be able to sort and display.
+			{
+				int& GlyphByFontCount = Stats.FindOrAdd(Key);
+				GlyphByFontCount++;
+			}
+		}
+
+		int FontAtlasPixelTotalCount = 0;
+		for (const TSharedRef<ISlateFontTexture>& FontTexture : AllFontTextures)
+		{
+			ISlateFontTexture& SlateFontTexture = FontTexture.Get();
+			const FSlateShaderResource* SlateShaderResource = SlateFontTexture.GetSlateTexture();
+			check(SlateShaderResource);
+ 
+			FontAtlasPixelTotalCount += SlateShaderResource->GetWidth() * SlateShaderResource->GetHeight();
+		}
+
+		//Display statistics
+		UE_LOG(LogSlate, Log, TEXT("%d glyphs in cache, using %d pixels on %d total. Using %f%% of all font textures."), GlyphCount, PixelsUsedInAtlasCount, FontAtlasPixelTotalCount, static_cast<float>(PixelsUsedInAtlasCount) * 100.0f / static_cast<float>(FontAtlasPixelTotalCount));
+		UE_LOG(LogSlate, Log, TEXT("FontFace & Style | Size (72DPI) | OutlineSize | Sep. Fill α | Skew || Glyphs Count |"));
+
+		Stats.KeySort([](const FFontCacheStatsKey& A, const FFontCacheStatsKey& B)
+			{
+				if (A.FontFace.Pin().Get() == B.FontFace.Pin().Get())
+				{
+					if (A.FontRenderSize == B.FontRenderSize)
+					{
+						if (A.OutlineRenderSize == B.OutlineRenderSize)
+						{
+							if (A.OutlineSeparateFillAlpha == B.OutlineSeparateFillAlpha)
+							{
+								return A.FontSkew < B.FontSkew;
+							}
+							return A.OutlineSeparateFillAlpha < B.OutlineSeparateFillAlpha;
+						}
+						return A.OutlineRenderSize < B.OutlineRenderSize;
+					}
+					return A.FontRenderSize < B.FontRenderSize;
+				}
+				return A.FontFace.Pin().Get()->GetFace()->family_name < B.FontFace.Pin().Get()->GetFace()->family_name;
+			});
+
+		for (const TPair<FFontCacheStatsKey, int>& Kvp : Stats)
+		{
+			FT_Face FT_face = Kvp.Key.FontFace.Pin().Get()->GetFace();
+			UE_LOG(LogSlate, Log, TEXT("%s %s | %d | %d | %s | %.2f || %d |"),
+				*FString(FT_face->family_name),
+				*FString(FT_face->style_name),
+				Kvp.Key.FontRenderSize,
+				Kvp.Key.OutlineRenderSize,
+				Kvp.Key.OutlineSeparateFillAlpha ? *FString("true") : *FString("false"),
+				Kvp.Key.FontSkew,
+				Kvp.Value);
+		}
+
+		bDumpFontCacheStats = false;
+	}
+}
+#endif
+
 bool FSlateFontCache::ConditionalFlushCache()
 {
 	bool bFlushed = false;
-	if (bFlushRequested)
+	if (bFlushRequested
+#if !UE_BUILD_SHIPPING
+		|| bFlushFontCache
+#endif
+		)
 	{
 		if (FlushCache())
 		{
 			bFlushRequested = false;
+#if !UE_BUILD_SHIPPING
+			bFlushFontCache = false;
+#endif
 			bFlushed = true;
 		}
 	}
@@ -1193,6 +1346,10 @@ void FSlateFontCache::UpdateCache()
 	UpdateFontAtlasTextures(ColorFontAtlasIndices);
 
 	CompositeFontCache->Update();
+
+#if (!UE_BUILD_SHIPPING && WITH_FREETYPE)
+	ConditionalDumpFontCacheStats();
+#endif
 }
 
 void FSlateFontCache::ReleaseResources()
diff --git a/Engine/Source/Runtime/SlateCore/Private/Fonts/FontCacheFreeType.cpp b/Engine/Source/Runtime/SlateCore/Private/Fonts/FontCacheFreeType.cpp
--- a/Engine/Source/Runtime/SlateCore/Private/Fonts/FontCacheFreeType.cpp
+++ b/Engine/Source/Runtime/SlateCore/Private/Fonts/FontCacheFreeType.cpp
@@ -72,35 +72,41 @@ namespace FreeTypeUtils
 
 #if WITH_FREETYPE
 
-void ApplySizeAndScale(FT_Face InFace, const int32 InFontSize, const float InFontScale)
+uint32 ComputeFontPixelSize(float InFontSize, float InFontScale)
 {
 	// Convert to fixed scale for maximum precision
-	const FT_F26Dot6 FixedFontSize = ConvertPixelTo26Dot6<FT_F26Dot6>(FMath::Max<int32>(0, InFontSize));
+	const FT_F26Dot6 FixedFontSize = ConvertPixelTo26Dot6<FT_F26Dot6>(FMath::Max<float>(0.0f, InFontSize));
 	const FT_Long FixedFontScale = ConvertPixelTo16Dot16<FT_Long>(FMath::Max<float>(0.0f, InFontScale));
 
 	// Convert the requested font size to a pixel size based on our render DPI and the requested scale
-	// This result is left in 26.6 space as it is the common format used by both FT_IS_SCALABLE and FT_HAS_FIXED_SIZES fonts
-	FT_F26Dot6 RequiredFixedFontPixelSize = 0;
-	{
-		// Convert the 26.6 character size to the unscaled 26.6 pixel size
-		// Note: The logic here mirrors what FT_REQUEST_WIDTH and FT_REQUEST_HEIGHT do internally when using FT_Set_Char_Size
-		RequiredFixedFontPixelSize = ((FixedFontSize * (FT_Pos)FreeTypeConstants::RenderDPI) + 36) / 72;
 
-		// Scale the 26.6 pixel size by the desired 16.16 fractional scaling value
-		RequiredFixedFontPixelSize = FT_MulFix(RequiredFixedFontPixelSize, FixedFontScale);
-	}
+	// Convert the 26.6 character size to the unscaled 26.6 pixel size
+	// Note: The logic here mirrors what FT_REQUEST_WIDTH and FT_REQUEST_HEIGHT do internally when using FT_Set_Char_Size
+	FT_F26Dot6 RequiredFixedFontPixelSize = ((FixedFontSize * (FT_Pos)FontConstants::RenderDPI) + 36) / 72;
+
+	// Scale the 26.6 pixel size by the desired 16.16 fractional scaling value
+	RequiredFixedFontPixelSize =  FT_MulFix(RequiredFixedFontPixelSize, FixedFontScale);
+
+	return Convert26Dot6ToRoundedPixel<uint32>(RequiredFixedFontPixelSize);
+}
+
+void ApplySizeAndScale(FT_Face InFace, const float InFontSize, const float InFontScale)
+{
+	ApplySizeAndScale(InFace, ComputeFontPixelSize(InFontSize, InFontScale));
+}
 
+void ApplySizeAndScale(FT_Face InFace, const uint32 RequiredFontPixelSize)
+{
 	if (FT_IS_SCALABLE(InFace))
 	{
-		// Convert the 26.6 scaled pixel size back into pixel space
-		const uint32 RequiredFontPixelSize = Convert26Dot6ToRoundedPixel<uint32>(RequiredFixedFontPixelSize);
-
 		// Set the pixel size
 		FT_Error Error = FT_Set_Pixel_Sizes(InFace, RequiredFontPixelSize, RequiredFontPixelSize);
 		check(Error == 0);
 	}
 	else if (FT_HAS_FIXED_SIZES(InFace))
 	{
+		FT_F26Dot6 RequiredFixedFontPixelSize = ConvertPixelTo26Dot6<uint32>(RequiredFontPixelSize);
+
 		// Find the best bitmap strike for our required pixel size
 		// Height is the most important metric, so that is the one we search for
 		int32 BestStrikeIndex = INDEX_NONE;
@@ -174,12 +180,17 @@ void ApplySizeAndScale(FT_Face InFace, const int32 InFontSize, const float InFon
 	}
 }
 
-FT_Error LoadGlyph(FT_Face InFace, const uint32 InGlyphIndex, const int32 InLoadFlags, const int32 InFontSize, const float InFontScale)
+FT_Error LoadGlyph(FT_Face InFace, const uint32 InGlyphIndex, const int32 InLoadFlags, const float InFontSize, const float InFontScale)
+{
+	return LoadGlyph(InFace, InGlyphIndex, InLoadFlags, ComputeFontPixelSize(InFontSize, InFontScale));
+}
+
+FT_Error LoadGlyph(FT_Face InFace, const uint32 InGlyphIndex, const int32 InLoadFlags, const uint32 RequiredFontPixelSize)
 {
 #if WITH_VERY_VERBOSE_SLATE_STATS
 	QUICK_SCOPE_CYCLE_COUNTER(STAT_FreetypeLoadGlyph);
 #endif
-	ApplySizeAndScale(InFace, InFontSize, InFontScale);
+	ApplySizeAndScale(InFace, RequiredFontPixelSize);
 	return FT_Load_Glyph(InFace, InGlyphIndex, InLoadFlags);
 }
 
@@ -598,11 +609,10 @@ unsigned long FFreeTypeFace::FFTStreamHandler::ReadData(FT_Stream InStream, unsi
 
 #if WITH_FREETYPE
 
-FFreeTypeGlyphCache::FFreeTypeGlyphCache(FT_Face InFace, const int32 InLoadFlags, const int32 InFontSize, const float InFontScale)
+FFreeTypeGlyphCache::FFreeTypeGlyphCache(FT_Face InFace, const int32 InLoadFlags, const float InFontSize, const float InFontScale)
 	: Face(InFace)
 	, LoadFlags(InLoadFlags)
-	, FontSize(InFontSize)
-	, FontScale(InFontScale)
+	, FontRenderSize(FreeTypeUtils::ComputeFontPixelSize(InFontSize,InFontScale))
 	, GlyphDataMap()
 {
 }
@@ -620,7 +630,7 @@ bool FFreeTypeGlyphCache::FindOrCache(const uint32 InGlyphIndex, FCachedGlyphDat
 	}
 
 	// No cached data, go ahead and add an entry for it...
-	FT_Error Error = FreeTypeUtils::LoadGlyph(Face, InGlyphIndex, LoadFlags, FontSize, FontScale);
+	FT_Error Error = FreeTypeUtils::LoadGlyph(Face, InGlyphIndex, LoadFlags, FontRenderSize);
 	if (Error == 0)
 	{
 		OutCachedGlyphData = FCachedGlyphData();
@@ -646,11 +656,10 @@ bool FFreeTypeGlyphCache::FindOrCache(const uint32 InGlyphIndex, FCachedGlyphDat
 }
 
 
-FFreeTypeAdvanceCache::FFreeTypeAdvanceCache(FT_Face InFace, const int32 InLoadFlags, const int32 InFontSize, const float InFontScale)
+FFreeTypeAdvanceCache::FFreeTypeAdvanceCache(FT_Face InFace, const int32 InLoadFlags, const float InFontSize, const float InFontScale)
 	: Face(InFace)
 	, LoadFlags(InLoadFlags)
-	, FontSize(InFontSize)
-	, FontScale(InFontScale)
+	, FontRenderSize(FreeTypeUtils::ComputeFontPixelSize(InFontSize, InFontScale))
 	, AdvanceMap()
 {
 }
@@ -667,7 +676,7 @@ bool FFreeTypeAdvanceCache::FindOrCache(const uint32 InGlyphIndex, FT_Fixed& Out
 		}
 	}
 
-	FreeTypeUtils::ApplySizeAndScale(Face, FontSize, FontScale);
+	FreeTypeUtils::ApplySizeAndScale(Face, FontRenderSize);
 
 	// No cached data, go ahead and add an entry for it...
 	const FT_Error Error = FT_Get_Advance(Face, InGlyphIndex, LoadFlags, &OutCachedAdvance);
@@ -687,11 +696,10 @@ bool FFreeTypeAdvanceCache::FindOrCache(const uint32 InGlyphIndex, FT_Fixed& Out
 }
 
 
-FFreeTypeKerningCache::FFreeTypeKerningCache(FT_Face InFace, const int32 InKerningFlags, const int32 InFontSize, const float InFontScale)
+FFreeTypeKerningCache::FFreeTypeKerningCache(FT_Face InFace, const int32 InKerningFlags, const float InFontSize, const float InFontScale)
 	: Face(InFace)
 	, KerningFlags(InKerningFlags)
-	, FontSize(InFontSize)
-	, FontScale(InFontScale)
+	, FontRenderSize(FreeTypeUtils::ComputeFontPixelSize(InFontSize, InFontScale))
 	, KerningMap()
 {
 	check(Face);
@@ -711,7 +719,7 @@ bool FFreeTypeKerningCache::FindOrCache(const uint32 InFirstGlyphIndex, const ui
 		}
 	}
 
-	FreeTypeUtils::ApplySizeAndScale(Face, FontSize, FontScale);
+	FreeTypeUtils::ApplySizeAndScale(Face, FontRenderSize);
 
 	// No cached data, go ahead and add an entry for it...
 	const FT_Error Error = FT_Get_Kerning(Face, InFirstGlyphIndex, InSecondGlyphIndex, KerningFlags, &OutKerning);
@@ -734,7 +742,7 @@ bool FFreeTypeKerningCache::FindOrCache(const uint32 InFirstGlyphIndex, const ui
 	return false;
 }
 
-TSharedRef<FFreeTypeGlyphCache> FFreeTypeCacheDirectory::GetGlyphCache(FT_Face InFace, const int32 InLoadFlags, const int32 InFontSize, const float InFontScale)
+TSharedRef<FFreeTypeGlyphCache> FFreeTypeCacheDirectory::GetGlyphCache(FT_Face InFace, const int32 InLoadFlags, const float InFontSize, const float InFontScale)
 {
 	const FFontKey Key(InFace, InLoadFlags, InFontSize, InFontScale);
 	TSharedPtr<FFreeTypeGlyphCache>& Entry = GlyphCacheMap.FindOrAdd(Key);
@@ -745,7 +753,7 @@ TSharedRef<FFreeTypeGlyphCache> FFreeTypeCacheDirectory::GetGlyphCache(FT_Face I
 	return Entry.ToSharedRef();
 }
 
-TSharedRef<FFreeTypeAdvanceCache> FFreeTypeCacheDirectory::GetAdvanceCache(FT_Face InFace, const int32 InLoadFlags, const int32 InFontSize, const float InFontScale)
+TSharedRef<FFreeTypeAdvanceCache> FFreeTypeCacheDirectory::GetAdvanceCache(FT_Face InFace, const int32 InLoadFlags, const float InFontSize, const float InFontScale)
 {
 	const FFontKey Key(InFace, InLoadFlags, InFontSize, InFontScale);
 	TSharedPtr<FFreeTypeAdvanceCache>& Entry = AdvanceCacheMap.FindOrAdd(Key);
@@ -756,7 +764,7 @@ TSharedRef<FFreeTypeAdvanceCache> FFreeTypeCacheDirectory::GetAdvanceCache(FT_Fa
 	return Entry.ToSharedRef();
 }
 
-TSharedPtr<FFreeTypeKerningCache> FFreeTypeCacheDirectory::GetKerningCache(FT_Face InFace, const int32 InKerningFlags, const int32 InFontSize, const float InFontScale)
+TSharedPtr<FFreeTypeKerningCache> FFreeTypeCacheDirectory::GetKerningCache(FT_Face InFace, const int32 InKerningFlags, const float InFontSize, const float InFontScale)
 {
 	// Check if this font has kerning as not all fonts do.
 	// We also can't perform kerning between two separate font faces
diff --git a/Engine/Source/Runtime/SlateCore/Private/Fonts/FontCacheFreeType.h b/Engine/Source/Runtime/SlateCore/Private/Fonts/FontCacheFreeType.h
--- a/Engine/Source/Runtime/SlateCore/Private/Fonts/FontCacheFreeType.h
+++ b/Engine/Source/Runtime/SlateCore/Private/Fonts/FontCacheFreeType.h
@@ -44,28 +44,31 @@
 	#undef generic
 #endif	//PLATFORM_COMPILER_HAS_GENERIC_KEYWORD
 
-
-namespace FreeTypeConstants
-{
-	/** The horizontal DPI we render at (horizontal and vertical) */
-	inline const uint32 RenderDPI = 96;
-} // namespace FreeTypeConstants
-
-
 namespace FreeTypeUtils
 {
 
 #if WITH_FREETYPE
 
+/**
+ * Compute the actual size that will be used by Freetype to render or do any process on glyphs.
+ */
+uint32 ComputeFontPixelSize(float InFontSize, float InFontScale);
+
 /**
  * Apply the given point size and scale to the face.
  */
-void ApplySizeAndScale(FT_Face InFace, const int32 InFontSize, const float InFontScale);
+void ApplySizeAndScale(FT_Face InFace, const float InFontSize, const float InFontScale);
+
+/**
+ * Apply the given size in pixel to the face.
+ */
+void ApplySizeAndScale(FT_Face InFace, const uint32 RequiredFontPixelSize);
 
 /**
  * Load the given glyph into the active slot of the given face.
  */
-FT_Error LoadGlyph(FT_Face InFace, const uint32 InGlyphIndex, const int32 InLoadFlags, const int32 InFontSize, const float InFontScale);
+FT_Error LoadGlyph(FT_Face InFace, const uint32 InGlyphIndex, const int32 InLoadFlags, const float InFontSize, const float InFontScale);
+FT_Error LoadGlyph(FT_Face InFace, const uint32 InGlyphIndex, const int32 InLoadFlags, const uint32 RequiredFontPixelSize);
 
 /**
  * Get the height of the given face under the given layout method.
@@ -362,7 +365,7 @@ class FFreeTypeGlyphCache
 {
 public:
 #if WITH_FREETYPE
-	FFreeTypeGlyphCache(FT_Face InFace, const int32 InLoadFlags, const int32 InFontSize, const float InFontScale);
+	FFreeTypeGlyphCache(FT_Face InFace, const int32 InLoadFlags, const float InFontSize, const float InFontScale);
 
 	struct FCachedGlyphData
 	{
@@ -378,9 +381,8 @@ public:
 private:
 #if WITH_FREETYPE
 	FT_Face Face;
-	int32 LoadFlags;
-	int32 FontSize;
-	float FontScale;
+	const int32 LoadFlags;
+	const uint32 FontRenderSize;
 	TMap<uint32, FCachedGlyphData> GlyphDataMap;
 #endif // WITH_FREETYPE
 };
@@ -393,7 +395,7 @@ class FFreeTypeAdvanceCache
 {
 public:
 #if WITH_FREETYPE
-	FFreeTypeAdvanceCache(FT_Face InFace, const int32 InLoadFlags, const int32 InFontSize, const float InFontScale);
+	FFreeTypeAdvanceCache(FT_Face InFace, const int32 InLoadFlags, const float InFontSize, const float InFontScale);
 
 	bool FindOrCache(const uint32 InGlyphIndex, FT_Fixed& OutCachedAdvance);
 #endif // WITH_FREETYPE
@@ -404,8 +406,7 @@ private:
 #if WITH_FREETYPE
 	FT_Face Face;
 	const int32 LoadFlags;
-	const int32 FontSize;
-	const float FontScale;
+	const uint32 FontRenderSize;
 	TMap<uint32, FT_Fixed> AdvanceMap;
 #endif // WITH_FREETYPE
 };
@@ -418,7 +419,7 @@ class FFreeTypeKerningCache
 {
 public:
 #if WITH_FREETYPE
-	FFreeTypeKerningCache(FT_Face InFace, const int32 InKerningFlags, const int32 InFontSize, const float InFontScale);
+	FFreeTypeKerningCache(FT_Face InFace, const int32 InKerningFlags, const float InFontSize, const float InFontScale);
 
 	/**
 	 * Retrieve the kerning vector for a given pair of glyphs.
@@ -465,8 +466,7 @@ private:
 
 	FT_Face Face;
 	const int32 KerningFlags;
-	const int32 FontSize;
-	const float FontScale;
+	const int32 FontRenderSize;
 	TMap<FKerningPair, FT_Vector> KerningMap;
 #endif // WITH_FREETYPE
 };
@@ -483,19 +483,19 @@ public:
 	 * Retrieve the glyph cache for a given set of font parameters.
 	 * @return A reference to the font glyph cache.
 	 */
-	TSharedRef<FFreeTypeGlyphCache> GetGlyphCache(FT_Face InFace, const int32 InLoadFlags, const int32 InFontSize, const float InFontScale);
+	TSharedRef<FFreeTypeGlyphCache> GetGlyphCache(FT_Face InFace, const int32 InLoadFlags, const float InFontSize, const float InFontScale);
 
 	/**
 	 * Retrieve the advance cache for a given set of font parameters.
 	 * @return A reference to the font advance cache.
 	 */
-	TSharedRef<FFreeTypeAdvanceCache> GetAdvanceCache(FT_Face InFace, const int32 InLoadFlags, const int32 InFontSize, const float InFontScale);
+	TSharedRef<FFreeTypeAdvanceCache> GetAdvanceCache(FT_Face InFace, const int32 InLoadFlags, const float InFontSize, const float InFontScale);
 
 	/**
 	 * Retrieve the kerning cache for a given set of font parameters.
 	 * @return A pointer to the font kerning cache, invalid if the font does not perform kerning.
 	 */
-	TSharedPtr<FFreeTypeKerningCache> GetKerningCache(FT_Face InFace, const int32 InKerningFlags, const int32 InFontSize, const float InFontScale);
+	TSharedPtr<FFreeTypeKerningCache> GetKerningCache(FT_Face InFace, const int32 InKerningFlags, const float InFontSize, const float InFontScale);
 #endif // WITH_FREETYPE
 
 	void FlushCache();
@@ -506,25 +506,22 @@ private:
 	class FFontKey
 	{
 	public:
-		FFontKey(FT_Face InFace, const int32 InFlags, const int32 InFontSize, const float InFontScale)
+		FFontKey(FT_Face InFace, const int32 InFlags, const float InFontSize, const float InFontScale)
 			: Face(InFace)
 			, Flags(InFlags)
-			, FontSize(InFontSize)
-			, FontScale(InFontScale)
+			, FontRenderSize(FreeTypeUtils::ComputeFontPixelSize(InFontSize, InFontScale))
 			, KeyHash(0)
 		{
 			KeyHash = GetTypeHash(Face);
 			KeyHash = HashCombine(KeyHash, GetTypeHash(Flags));
-			KeyHash = HashCombine(KeyHash, GetTypeHash(FontSize));
-			KeyHash = HashCombine(KeyHash, GetTypeHash(FontScale));
+			KeyHash = HashCombine(KeyHash, GetTypeHash(FontRenderSize));
 		}
 
 		FORCEINLINE bool operator==(const FFontKey& Other) const
 		{
 			return Face == Other.Face
 				&& Flags == Other.Flags
-				&& FontSize == Other.FontSize
-				&& FontScale == Other.FontScale;
+				&& FontRenderSize == Other.FontRenderSize;
 		}
 
 		FORCEINLINE bool operator!=(const FFontKey& Other) const
@@ -539,9 +536,8 @@ private:
 
 	private:
 		FT_Face Face;
-		int32 Flags;
-		int32 FontSize;
-		float FontScale;
+		const int32 Flags;
+		const int32 FontRenderSize;
 		uint32 KeyHash;
 	};
 
diff --git a/Engine/Source/Runtime/SlateCore/Private/Fonts/FontCacheHarfBuzz.cpp b/Engine/Source/Runtime/SlateCore/Private/Fonts/FontCacheHarfBuzz.cpp
--- a/Engine/Source/Runtime/SlateCore/Private/Fonts/FontCacheHarfBuzz.cpp
+++ b/Engine/Source/Runtime/SlateCore/Private/Fonts/FontCacheHarfBuzz.cpp
@@ -103,7 +103,7 @@ hb_user_data_key_t UserDataKey;
 
 struct FUserData
 {
-	FUserData(const int32 InFontSize, const float InFontScale, FFreeTypeCacheDirectory* InFTCacheDirectory, const hb_font_extents_t& InHarfBuzzFontExtents)
+	FUserData(const float InFontSize, const float InFontScale, FFreeTypeCacheDirectory* InFTCacheDirectory, const hb_font_extents_t& InHarfBuzzFontExtents)
 		: FontSize(InFontSize)
 		, FontScale(InFontScale)
 		, FTCacheDirectory(InFTCacheDirectory)
@@ -111,13 +111,13 @@ struct FUserData
 	{
 	}
 
-	int32 FontSize;
+	float FontSize;
 	float FontScale;
 	FFreeTypeCacheDirectory* FTCacheDirectory;
 	hb_font_extents_t HarfBuzzFontExtents;
 };
 
-void* CreateUserData(const int32 InFontSize, const float InFontScale, FFreeTypeCacheDirectory* InFTCacheDirectory, const hb_font_extents_t& InHarfBuzzFontExtents)
+void* CreateUserData(const float InFontSize, const float InFontScale, FFreeTypeCacheDirectory* InFTCacheDirectory, const hb_font_extents_t& InHarfBuzzFontExtents)
 {
 	return new FUserData(InFontSize, InFontScale, InFTCacheDirectory, InHarfBuzzFontExtents);
 }
@@ -449,7 +449,7 @@ FHarfBuzzFontFactory::~FHarfBuzzFontFactory()
 
 #if WITH_HARFBUZZ
 
-hb_font_t* FHarfBuzzFontFactory::CreateFont(const FFreeTypeFace& InFace, const uint32 InGlyphFlags, const int32 InFontSize, const float InFontScale) const
+hb_font_t* FHarfBuzzFontFactory::CreateFont(const FFreeTypeFace& InFace, const uint32 InGlyphFlags, const float InFontSize, const float InFontScale) const
 {
 	hb_font_t* HarfBuzzFont = nullptr;
 
diff --git a/Engine/Source/Runtime/SlateCore/Private/Fonts/FontCacheHarfBuzz.h b/Engine/Source/Runtime/SlateCore/Private/Fonts/FontCacheHarfBuzz.h
--- a/Engine/Source/Runtime/SlateCore/Private/Fonts/FontCacheHarfBuzz.h
+++ b/Engine/Source/Runtime/SlateCore/Private/Fonts/FontCacheHarfBuzz.h
@@ -54,7 +54,7 @@ public:
 
 #if WITH_HARFBUZZ
 	/** Create a HarfBuzz font from the given face - must be destroyed with hb_font_destroy when done */
-	hb_font_t* CreateFont(const FFreeTypeFace& InFace, const uint32 InGlyphFlags, const int32 InFontSize, const float InFontScale) const;
+	hb_font_t* CreateFont(const FFreeTypeFace& InFace, const uint32 InGlyphFlags, const float InFontSize, const float InFontScale) const;
 #endif // WITH_HARFBUZZ
 
 private:
diff --git a/Engine/Source/Runtime/SlateCore/Private/Fonts/SlateFontInfo.cpp b/Engine/Source/Runtime/SlateCore/Private/Fonts/SlateFontInfo.cpp
--- a/Engine/Source/Runtime/SlateCore/Private/Fonts/SlateFontInfo.cpp
+++ b/Engine/Source/Runtime/SlateCore/Private/Fonts/SlateFontInfo.cpp
@@ -50,19 +50,19 @@ FSlateFontInfo::FSlateFontInfo( )
 }
 
 
-FSlateFontInfo::FSlateFontInfo( TSharedPtr<const FCompositeFont> InCompositeFont, const int32 InSize, const FName& InTypefaceFontName, const FFontOutlineSettings& InOutlineSettings )
+FSlateFontInfo::FSlateFontInfo( TSharedPtr<const FCompositeFont> InCompositeFont, const float InSize, const FName& InTypefaceFontName, const FFontOutlineSettings& InOutlineSettings )
 	: FontObject(nullptr)
 	, FontMaterial(nullptr)
 	, OutlineSettings(InOutlineSettings)
 	, CompositeFont(InCompositeFont)
 	, TypefaceFontName(InTypefaceFontName)
-	, Size(FMath::Clamp<int32>(InSize, 0, std::numeric_limits<uint16>::max()))
+	, Size(FMath::Clamp<float>(InSize, 0.f, std::numeric_limits<uint16>::max()))
 	, FontFallback(EFontFallback::FF_Max)
 #if WITH_EDITORONLY_DATA
 	, Hinting_DEPRECATED(EFontHinting::Default)
 #endif
 {
-	ensureMsgf(InSize >= 0 && InSize <= std::numeric_limits<uint16>::max(), TEXT("The size provided is not supported by the renderer."));
+	ensureMsgf(InSize >= 0.f && InSize <= std::numeric_limits<uint16>::max(), TEXT("The size provided is not supported by the renderer."));
 	if (!InCompositeFont.IsValid())
 	{
 		UE_LOG(LogSlate, Warning, TEXT("FSlateFontInfo was constructed with a null FCompositeFont. Slate will be forced to use the fallback font path which may be slower."));
@@ -70,19 +70,19 @@ FSlateFontInfo::FSlateFontInfo( TSharedPtr<const FCompositeFont> InCompositeFont
 }
 
 
-FSlateFontInfo::FSlateFontInfo( const UObject* InFontObject, const int32 InSize, const FName& InTypefaceFontName, const FFontOutlineSettings& InOutlineSettings )
+FSlateFontInfo::FSlateFontInfo( const UObject* InFontObject, const float InSize, const FName& InTypefaceFontName, const FFontOutlineSettings& InOutlineSettings )
 	: FontObject(InFontObject)
 	, FontMaterial(nullptr)
 	, OutlineSettings(InOutlineSettings)
 	, CompositeFont()
 	, TypefaceFontName(InTypefaceFontName)
-	, Size(FMath::Clamp<int32>(InSize, 0, std::numeric_limits<uint16>::max()))
+	, Size(FMath::Clamp<float>(InSize, 0.f, std::numeric_limits<uint16>::max()))
 	, FontFallback(EFontFallback::FF_Max)
 #if WITH_EDITORONLY_DATA
 	, Hinting_DEPRECATED(EFontHinting::Default)
 #endif
 {
-	ensureMsgf(InSize >= 0 && InSize <= std::numeric_limits<uint16>::max(), TEXT("The size provided is not supported by the renderer."));
+	ensureMsgf(InSize >= 0.f && InSize <= std::numeric_limits<uint16>::max(), TEXT("The size provided is not supported by the renderer."));
 	if (InFontObject)
 	{
 		const IFontProviderInterface* FontProvider = Cast<const IFontProviderInterface>(InFontObject);
@@ -98,19 +98,21 @@ FSlateFontInfo::FSlateFontInfo( const UObject* InFontObject, const int32 InSize,
 }
 
 
-FSlateFontInfo::FSlateFontInfo( const FString& InFontName, uint16 InSize, EFontHinting InHinting, const FFontOutlineSettings& InOutlineSettings)
+FSlateFontInfo::FSlateFontInfo( const FString& InFontName, float InSize, EFontHinting InHinting, const FFontOutlineSettings& InOutlineSettings)
 	: FontObject(nullptr)
 	, FontMaterial(nullptr)
 	, OutlineSettings(InOutlineSettings)
 	, CompositeFont()
 	, TypefaceFontName()
-	, Size(InSize)
+	, Size(FMath::Clamp<float>(InSize, 0.f, std::numeric_limits<uint16>::max()))
 	, FontFallback(EFontFallback::FF_Max)
 #if WITH_EDITORONLY_DATA
 	, FontName_DEPRECATED(*InFontName)
 	, Hinting_DEPRECATED(InHinting)
 #endif
 {
+	ensureMsgf(InSize >= 0.f && InSize <= std::numeric_limits<uint16>::max(), TEXT("The size provided is not supported by the renderer."));
+
 	//Useful for debugging style breakages
 	//check( FPaths::FileExists( FontName.ToString() ) );
 
@@ -118,18 +120,20 @@ FSlateFontInfo::FSlateFontInfo( const FString& InFontName, uint16 InSize, EFontH
 }
 
 
-FSlateFontInfo::FSlateFontInfo( const FName& InFontName, uint16 InSize, EFontHinting InHinting )
+FSlateFontInfo::FSlateFontInfo( const FName& InFontName, float InSize, EFontHinting InHinting )
 	: FontObject(nullptr)
 	, FontMaterial(nullptr)
 	, CompositeFont()
 	, TypefaceFontName()
-	, Size(InSize)
+	, Size(FMath::Clamp<float>(InSize, 0.f, std::numeric_limits<uint16>::max()))
 	, FontFallback(EFontFallback::FF_Max)
 #if WITH_EDITORONLY_DATA
 	, FontName_DEPRECATED(InFontName)
 	, Hinting_DEPRECATED(InHinting)
 #endif
 {
+	ensureMsgf(InSize >= 0.f && InSize <= std::numeric_limits<uint16>::max(), TEXT("The size provided is not supported by the renderer."));
+
 	//Useful for debugging style breakages
 	//check( FPaths::FileExists( FontName.ToString() ) );
 
@@ -137,18 +141,20 @@ FSlateFontInfo::FSlateFontInfo( const FName& InFontName, uint16 InSize, EFontHin
 }
 
 
-FSlateFontInfo::FSlateFontInfo( const ANSICHAR* InFontName, uint16 InSize, EFontHinting InHinting )
+FSlateFontInfo::FSlateFontInfo( const ANSICHAR* InFontName, float InSize, EFontHinting InHinting )
 	: FontObject(nullptr)
 	, FontMaterial(nullptr)
 	, CompositeFont()
 	, TypefaceFontName()
-	, Size(InSize)
+	, Size(FMath::Clamp<float>(InSize, 0.f, std::numeric_limits<uint16>::max()))
 	, FontFallback(EFontFallback::FF_Max)
 #if WITH_EDITORONLY_DATA
 	, FontName_DEPRECATED(InFontName)
 	, Hinting_DEPRECATED(InHinting)
 #endif
 {
+	ensureMsgf(InSize >= 0.f && InSize <= std::numeric_limits<uint16>::max(), TEXT("The size provided is not supported by the renderer."));
+
 	//Useful for debugging style breakages
 	//check( FPaths::FileExists( FontName.ToString() ) );
 
@@ -156,18 +162,20 @@ FSlateFontInfo::FSlateFontInfo( const ANSICHAR* InFontName, uint16 InSize, EFont
 }
 
 
-FSlateFontInfo::FSlateFontInfo( const WIDECHAR* InFontName, uint16 InSize, EFontHinting InHinting )
+FSlateFontInfo::FSlateFontInfo( const WIDECHAR* InFontName, float InSize, EFontHinting InHinting )
 	: FontObject(nullptr)
 	, FontMaterial(nullptr)
 	, CompositeFont()
 	, TypefaceFontName()
-	, Size(InSize)
+	, Size(FMath::Clamp<float>(InSize, 0.f, std::numeric_limits<uint16>::max()))
 	, FontFallback(EFontFallback::FF_Max)
 #if WITH_EDITORONLY_DATA
 	, FontName_DEPRECATED(InFontName)
 	, Hinting_DEPRECATED(InHinting)
 #endif
 {
+	ensureMsgf(InSize >= 0.f && InSize <= std::numeric_limits<uint16>::max(), TEXT("The size provided is not supported by the renderer."));
+
 	//Useful for debugging style breakages
 	//check( FPaths::FileExists( FontName.ToString() ) );
 
@@ -226,9 +234,9 @@ void FSlateFontInfo::UpgradeLegacyFontInfo(FName LegacyFontName, EFontHinting Le
 		: FLegacySlateFontInfoCache::Get().GetCompositeFont(LegacyFontName, LegacyHinting);
 }
 
-uint16 FSlateFontInfo::GetClampSize() const
+float FSlateFontInfo::GetClampSize() const
 {
-	return (uint16)FMath::Clamp<int32>(Size, 0, std::numeric_limits<uint16>::max());
+	return FMath::Clamp<float>(Size, 0.f, std::numeric_limits<uint16>::max());
 }
 
 float FSlateFontInfo::GetClampSkew() const
diff --git a/Engine/Source/Runtime/SlateCore/Private/Fonts/SlateFontRenderer.cpp b/Engine/Source/Runtime/SlateCore/Private/Fonts/SlateFontRenderer.cpp
--- a/Engine/Source/Runtime/SlateCore/Private/Fonts/SlateFontRenderer.cpp
+++ b/Engine/Source/Runtime/SlateCore/Private/Fonts/SlateFontRenderer.cpp
@@ -286,7 +286,7 @@ bool FSlateFontRenderer::HasKerning(const FFontData& InFontData) const
 #endif // WITH_FREETYPE
 }
 
-int8 FSlateFontRenderer::GetKerning(const FFontData& InFontData, const int32 InSize, TCHAR First, TCHAR Second, const float InScale) const
+int8 FSlateFontRenderer::GetKerning(const FFontData& InFontData, const float InSize, TCHAR First, TCHAR Second, const float InScale) const
 {
 #if WITH_FREETYPE
 	FT_Face FontFace = GetFontFace(InFontData);
@@ -312,7 +312,7 @@ int8 FSlateFontRenderer::GetKerning(const FFontData& InFontData, const int32 InS
 	return 0;
 }
 
-TSharedPtr<FFreeTypeKerningCache> FSlateFontRenderer::GetKerningCache(const FFontData& InFontData, const int32 InSize, const float InScale) const
+TSharedPtr<FFreeTypeKerningCache> FSlateFontRenderer::GetKerningCache(const FFontData& InFontData, const float InSize, const float InScale) const
 {
 #if WITH_FREETYPE
 	FT_Face FontFace = GetFontFace(InFontData);
diff --git a/Engine/Source/Runtime/SlateCore/Private/Fonts/SlateFontRenderer.h b/Engine/Source/Runtime/SlateCore/Private/Fonts/SlateFontRenderer.h
--- a/Engine/Source/Runtime/SlateCore/Private/Fonts/SlateFontRenderer.h
+++ b/Engine/Source/Runtime/SlateCore/Private/Fonts/SlateFontRenderer.h
@@ -100,7 +100,7 @@ public:
 	 * @param InScale		The scale at which the font is being drawn
 	 * @return The kerning amount, 0 if no kerning
 	 */
-	int8 GetKerning(const FFontData& InFontData, const int32 InSize, TCHAR First, TCHAR Second, const float InScale) const;
+	int8 GetKerning(const FFontData& InFontData, const float InSize, TCHAR First, TCHAR Second, const float InScale) const;
 
 	/**
 	 * Retrieves the kerning cache object used to retrieve kerning amounts between glyphs
@@ -110,7 +110,7 @@ public:
 	 * @param InScale		The scale at which the font is being drawn
 	 * @return A pointer object referring to the kerning cache matching the given parameters, invalid/null if no kerning is performed
 	 */
-	TSharedPtr<FFreeTypeKerningCache> GetKerningCache(const FFontData& InFontData, const int32 InSize, const float InScale) const;
+	TSharedPtr<FFreeTypeKerningCache> GetKerningCache(const FFontData& InFontData, const float InSize, const float InScale) const;
 
 	/**
 	 * Whether or not the specified character, within the specified font, can be loaded with the specified maximum font fallback level
diff --git a/Engine/Source/Runtime/SlateCore/Private/Fonts/SlateTextShaper.h b/Engine/Source/Runtime/SlateCore/Private/Fonts/SlateTextShaper.h
--- a/Engine/Source/Runtime/SlateCore/Private/Fonts/SlateTextShaper.h
+++ b/Engine/Source/Runtime/SlateCore/Private/Fonts/SlateTextShaper.h
@@ -17,7 +17,7 @@ class IBreakIterator;
 class FShapedGlyphFaceData
 {
 public:
-	FShapedGlyphFaceData(TWeakPtr<FFreeTypeFace> InFontFace, const uint32 InGlyphFlags, const int32 InFontSize, const float InFontScale, const float InFontSkew)
+	FShapedGlyphFaceData(TWeakPtr<FFreeTypeFace> InFontFace, const uint32 InGlyphFlags, const float InFontSize, const float InFontScale, const float InFontSkew)
 		: FontFace(MoveTemp(InFontFace))
 		, GlyphFlags(InGlyphFlags)
 		, FontSize(InFontSize)
@@ -38,7 +38,7 @@ public:
 	/** Provides the glyph flags used to render the font */
 	uint32 GlyphFlags;
 	/** Provides the point size used to render the font */
-	int32 FontSize;
+	float FontSize;
 	/** Provides the final scale used to render to the font */
 	float FontScale;
 	/** Any additional scale that should be applied when rendering this glyph */
diff --git a/Engine/Source/Runtime/SlateCore/Private/Styling/CoreStyle.cpp b/Engine/Source/Runtime/SlateCore/Private/Styling/CoreStyle.cpp
--- a/Engine/Source/Runtime/SlateCore/Private/Styling/CoreStyle.cpp
+++ b/Engine/Source/Runtime/SlateCore/Private/Styling/CoreStyle.cpp
@@ -69,7 +69,7 @@ TSharedRef<const FCompositeFont> FCoreStyle::GetDefaultFont()
 }
 
 
-FSlateFontInfo FCoreStyle::GetDefaultFontStyle(const FName InTypefaceFontName, const int32 InSize, const FFontOutlineSettings& InOutlineSettings)
+FSlateFontInfo FCoreStyle::GetDefaultFontStyle(const FName InTypefaceFontName, const float InSize, const FFontOutlineSettings& InOutlineSettings)
 {
 	return FSlateFontInfo(GetDefaultFont(), InSize, InTypefaceFontName, InOutlineSettings);
 }
diff --git a/Engine/Source/Runtime/SlateCore/Private/Styling/StarshipCoreStyle.cpp b/Engine/Source/Runtime/SlateCore/Private/Styling/StarshipCoreStyle.cpp
--- a/Engine/Source/Runtime/SlateCore/Private/Styling/StarshipCoreStyle.cpp
+++ b/Engine/Source/Runtime/SlateCore/Private/Styling/StarshipCoreStyle.cpp
@@ -94,7 +94,7 @@ TSharedRef<const FCompositeFont> FStarshipCoreStyle::GetDefaultFont()
 }
 
 
-FSlateFontInfo FStarshipCoreStyle::GetDefaultFontStyle(const FName InTypefaceFontName, const int32 InSize, const FFontOutlineSettings& InOutlineSettings)
+FSlateFontInfo FStarshipCoreStyle::GetDefaultFontStyle(const FName InTypefaceFontName, const float InSize, const FFontOutlineSettings& InOutlineSettings)
 {
 	return FSlateFontInfo(GetDefaultFont(), InSize, InTypefaceFontName, InOutlineSettings);
 }
diff --git a/Engine/Source/Runtime/SlateCore/Private/Styling/UMGCoreStyle.cpp b/Engine/Source/Runtime/SlateCore/Private/Styling/UMGCoreStyle.cpp
--- a/Engine/Source/Runtime/SlateCore/Private/Styling/UMGCoreStyle.cpp
+++ b/Engine/Source/Runtime/SlateCore/Private/Styling/UMGCoreStyle.cpp
@@ -126,7 +126,7 @@ TSharedRef<const FCompositeFont> FUMGCoreStyle::GetDefaultFont()
 }
 
 
-FSlateFontInfo FUMGCoreStyle::GetDefaultFontStyle(const FName InTypefaceFontName, const int32 InSize, const FFontOutlineSettings& InOutlineSettings)
+FSlateFontInfo FUMGCoreStyle::GetDefaultFontStyle(const FName InTypefaceFontName, const float InSize, const FFontOutlineSettings& InOutlineSettings)
 {
 	return FSlateFontInfo(GetDefaultFont(), InSize, InTypefaceFontName, InOutlineSettings);
 }
diff --git a/Engine/Source/Runtime/SlateCore/Public/Fonts/FontCache.h b/Engine/Source/Runtime/SlateCore/Public/Fonts/FontCache.h
--- a/Engine/Source/Runtime/SlateCore/Public/Fonts/FontCache.h
+++ b/Engine/Source/Runtime/SlateCore/Public/Fonts/FontCache.h
@@ -154,16 +154,17 @@ private:
 /** Minimal FShapedGlyphEntry key information used for map lookups */
 struct FShapedGlyphEntryKey
 {
+	friend struct FFontCacheStatsKey;
+
 public:
 	FShapedGlyphEntryKey(const FShapedGlyphFaceData& InFontFaceData, uint32 InGlyphIndex, const FFontOutlineSettings& InOutlineSettings);
 
 	FORCEINLINE bool operator==(const FShapedGlyphEntryKey& Other) const
 	{
 		return FontFace == Other.FontFace 
-			&& FontSize == Other.FontSize
-			&& OutlineSize == Other.OutlineSize
+			&& FontRenderSize == Other.FontRenderSize
+			&& OutlineRenderSize == Other.OutlineRenderSize
 			&& OutlineSeparateFillAlpha == Other.OutlineSeparateFillAlpha
-			&& FontScale == Other.FontScale
 			&& GlyphIndex == Other.GlyphIndex
 			&& FontSkew == Other.FontSkew;
 	}
@@ -181,14 +182,12 @@ public:
 private:
 	/** Weak pointer to the FreeType face to render with */
 	TWeakPtr<FFreeTypeFace> FontFace;
-	/** Provides the point size used to render the font */
-	int32 FontSize;
-	/** The size in pixels of the outline to render for the font */
-	float OutlineSize;
+	/** The size in pixel that Freetype will actually use to render the font, scale included */
+	uint32 FontRenderSize;
+	/** The size in pixels of the outline to render, scale included */
+	int32 OutlineRenderSize;
 	/** If checked, the outline will be completely translucent where the filled area will be. @see FFontOutlineSettings */
 	bool OutlineSeparateFillAlpha;
-	/** Provides the final scale used to render to the font */
-	float FontScale;
 	/** The index of this glyph in the FreeType face */
 	uint32 GlyphIndex;
 	/** Cached hash value used for map lookups */
@@ -782,7 +781,7 @@ public:
 	 * @param Second		The second character in the pair
 	 * @return The kerning amount, 0 if no kerning
 	 */
-	SLATECORE_API int8 GetKerning( const FFontData& InFontData, const int32 InSize, TCHAR First, TCHAR Second, float Scale ) const;
+	SLATECORE_API int8 GetKerning( const FFontData& InFontData, const float InSize, TCHAR First, TCHAR Second, float Scale ) const;
 
 	/**
 	 * @return Whether or not the font used has kerning information
@@ -860,6 +859,12 @@ private:
 	SLATECORE_API bool AddNewEntry(const FShapedGlyphEntry& InShapedGlyph, const FFontOutlineSettings& InOutlineSettings, FShapedGlyphFontAtlasData& OutAtlasData);
 
 	SLATECORE_API bool AddNewEntry(const FCharacterRenderData InRenderData, uint8& OutTextureIndex, uint16& OutGlyphX, uint16& OutGlyphY, uint16& OutGlyphWidth, uint16& OutGlyphHeight);
+
+#if !UE_BUILD_SHIPPING
+	/** Dump statistics about Font cache usage if needed */
+	SLATECORE_API void ConditionalDumpFontCacheStats() const;
+#endif
+
 private:
 
 	/** FreeType library instance (owned by this font cache) */
diff --git a/Engine/Source/Runtime/SlateCore/Public/Fonts/SlateFontInfo.h b/Engine/Source/Runtime/SlateCore/Public/Fonts/SlateFontInfo.h
--- a/Engine/Source/Runtime/SlateCore/Public/Fonts/SlateFontInfo.h
+++ b/Engine/Source/Runtime/SlateCore/Public/Fonts/SlateFontInfo.h
@@ -8,6 +8,12 @@
 #include "Fonts/CompositeFont.h"
 #include "SlateFontInfo.generated.h"
 
+namespace FontConstants
+{
+	/** The  DPI we render at (horizontal and vertical) */
+	inline const uint32 RenderDPI = 96;
+}
+
 /**
  * Sets the maximum font fallback level, for when a character can't be found in the selected font set.
  * UI code that renders strings from a third party (e.g. player chat in a multiplayer game) may want to restrict font fallback to prevent potential performance problems.
@@ -153,12 +159,12 @@ struct FSlateFontInfo
 	FName TypefaceFontName;
 
 	/**
-	 * The font size is a measure in point values. The conversion of points to Slate Units is done at 96 dpi.  So if 
-	 * you're using a tool like Photoshop to prototype layouts and UI mock ups, be sure to change the default dpi 
-	 * measurements from 72 dpi to 96 dpi.
+	 * The font size is a measure in point values. The conversion of points to Slate Units is done at 96 DPI.
+	 * So if you're using a tool like Photoshop to prototype layouts and UI mock ups, you can change the UMG Font settings
+	 * to ensure that UMG font size is displayed in its 72 DPI equivalent, even if Slate will still use 96 DPI internally.
 	 */
 	UPROPERTY(EditAnywhere, BlueprintReadWrite, Category=SlateStyleRules, meta=(ClampMin=1, ClampMax=1000))
-	int32 Size;
+	float Size;
 
 	/** The uniform spacing (or tracking) between all characters in the text. */
 	UPROPERTY(EditAnywhere, BlueprintReadWrite, Category=SlateStyleRules, meta=(ClampMin=-1000, ClampMax=10000))
@@ -198,7 +204,7 @@ public:
 	 * @param InSize The size of the font.
 	 * @param InTypefaceFontName The name of the font to use from the default typeface (None will use the first entry)
 	 */
-	SLATECORE_API FSlateFontInfo( TSharedPtr<const FCompositeFont> InCompositeFont, const int32 InSize, const FName& InTypefaceFontName = NAME_None, const FFontOutlineSettings& InOutlineSettings = FFontOutlineSettings());
+	SLATECORE_API FSlateFontInfo( TSharedPtr<const FCompositeFont> InCompositeFont, const float InSize, const FName& InTypefaceFontName = NAME_None, const FFontOutlineSettings& InOutlineSettings = FFontOutlineSettings());
 
 	/**
 	 * Creates and initializes a new instance with the specified font, size, and emphasis.
@@ -207,7 +213,7 @@ public:
 	 * @param InSize The size of the font.
 	 * @param InFamilyFontName The name of the font to use from the default typeface (None will use the first entry)
 	 */
-	SLATECORE_API FSlateFontInfo( const UObject* InFontObject, const int32 InSize, const FName& InTypefaceFontName = NAME_None, const FFontOutlineSettings& InOutlineSettings = FFontOutlineSettings());
+	SLATECORE_API FSlateFontInfo( const UObject* InFontObject, const float InSize, const FName& InTypefaceFontName = NAME_None, const FFontOutlineSettings& InOutlineSettings = FFontOutlineSettings());
 
 	/**
 	 * DEPRECATED - Creates and initializes a new instance with the specified font name and size.
@@ -216,7 +222,7 @@ public:
 	 * @param InSize The size of the font.
 	 * @param InHinting The type of hinting to use for the font.
 	 */
-	SLATECORE_API FSlateFontInfo( const FString& InFontName, uint16 InSize, EFontHinting InHinting = EFontHinting::Default, const FFontOutlineSettings& InOutlineSettings = FFontOutlineSettings());
+	SLATECORE_API FSlateFontInfo( const FString& InFontName, float InSize, EFontHinting InHinting = EFontHinting::Default, const FFontOutlineSettings& InOutlineSettings = FFontOutlineSettings());
 
 	/**
 	 * DEPRECATED - Creates and initializes a new instance with the specified font name and size.
@@ -225,7 +231,7 @@ public:
 	 * @param InSize The size of the font.
 	 * @param InHinting The type of hinting to use for the font.
 	 */
-	SLATECORE_API FSlateFontInfo( const FName& InFontName, uint16 InSize, EFontHinting InHinting = EFontHinting::Default );
+	SLATECORE_API FSlateFontInfo( const FName& InFontName, float InSize, EFontHinting InHinting = EFontHinting::Default );
 
 	/**
 	 * DEPRECATED - Creates and initializes a new instance with the specified font name and size.
@@ -234,7 +240,7 @@ public:
 	 * @param InSize The size of the font.
 	 * @param InHinting The type of hinting to use for the font.
 	 */
-	SLATECORE_API FSlateFontInfo( const ANSICHAR* InFontName, uint16 InSize, EFontHinting InHinting = EFontHinting::Default );
+	SLATECORE_API FSlateFontInfo( const ANSICHAR* InFontName, float InSize, EFontHinting InHinting = EFontHinting::Default );
 
 	/**
 	 * DEPRECATED - Creates and initializes a new instance with the specified font name and size.
@@ -243,7 +249,7 @@ public:
 	 * @param InSize The size of the font.
 	 * @param InHinting The type of hinting to use for the font.
 	 */
-	SLATECORE_API FSlateFontInfo( const WIDECHAR* InFontName, uint16 InSize, EFontHinting InHinting = EFontHinting::Default );
+	SLATECORE_API FSlateFontInfo( const WIDECHAR* InFontName, float InSize, EFontHinting InHinting = EFontHinting::Default );
 
 public:
 	inline bool IsIdentialToForCaching(const FSlateFontInfo& Other) const
@@ -286,7 +292,7 @@ public:
 	SLATECORE_API const FCompositeFont* GetCompositeFont() const;
 
 	/** Get the font size clamp for the font renderer (on 16bits) */
-	SLATECORE_API uint16 GetClampSize() const;
+	SLATECORE_API float GetClampSize() const;
 
 	/** Get the skew amount clamp for the text shaper */
 	SLATECORE_API float GetClampSkew() const;
diff --git a/Engine/Source/Runtime/SlateCore/Public/Styling/CoreStyle.h b/Engine/Source/Runtime/SlateCore/Public/Styling/CoreStyle.h
--- a/Engine/Source/Runtime/SlateCore/Public/Styling/CoreStyle.h
+++ b/Engine/Source/Runtime/SlateCore/Public/Styling/CoreStyle.h
@@ -45,7 +45,7 @@ public:
 	static SLATECORE_API TSharedRef<const FCompositeFont> GetDefaultFont();
 
 	/** Get a font style using the default for for Slate */
-	static SLATECORE_API FSlateFontInfo GetDefaultFontStyle(const FName InTypefaceFontName, const int32 InSize, const FFontOutlineSettings& InOutlineSettings = FFontOutlineSettings());
+	static SLATECORE_API FSlateFontInfo GetDefaultFontStyle(const FName InTypefaceFontName, const float InSize, const FFontOutlineSettings& InOutlineSettings = FFontOutlineSettings());
 
 	static SLATECORE_API void ResetToDefault( );
 
diff --git a/Engine/Source/Runtime/SlateCore/Public/Styling/SlateTypes.h b/Engine/Source/Runtime/SlateCore/Public/Styling/SlateTypes.h
--- a/Engine/Source/Runtime/SlateCore/Public/Styling/SlateTypes.h
+++ b/Engine/Source/Runtime/SlateCore/Public/Styling/SlateTypes.h
@@ -329,17 +329,17 @@ struct FTextBlockStyle : public FSlateWidgetStyle
 	UPROPERTY(EditAnywhere, BlueprintReadWrite, Category=Appearance)
 	FSlateFontInfo Font;
 	FTextBlockStyle& SetFont(const FSlateFontInfo& InFont) { Font = InFont; return *this; }
-	FTextBlockStyle& SetFont(TSharedPtr<const FCompositeFont> InCompositeFont, const int32 InSize, const FName& InTypefaceFontName = NAME_None) { Font = FSlateFontInfo(InCompositeFont, InSize, InTypefaceFontName); return *this; }
-	FTextBlockStyle& SetFont(const UObject* InFontObject, const int32 InSize, const FName& InTypefaceFontName = NAME_None) { Font = FSlateFontInfo(InFontObject, InSize, InTypefaceFontName); return *this; }
-	FTextBlockStyle& SetFont(const FName& InFontName, uint16 InSize) { Font = FSlateFontInfo(InFontName, InSize); return *this; }
-	FTextBlockStyle& SetFont(const FString& InFontName, uint16 InSize) { Font = FSlateFontInfo(*InFontName, InSize); return *this; }
-	FTextBlockStyle& SetFont(const WIDECHAR* InFontName, uint16 InSize) { Font = FSlateFontInfo(InFontName, InSize); return *this; }
-	FTextBlockStyle& SetFont(const ANSICHAR* InFontName, uint16 InSize) { Font = FSlateFontInfo(InFontName, InSize); return *this; }
+	FTextBlockStyle& SetFont(TSharedPtr<const FCompositeFont> InCompositeFont, const float InSize, const FName& InTypefaceFontName = NAME_None) { Font = FSlateFontInfo(InCompositeFont, InSize, InTypefaceFontName); return *this; }
+	FTextBlockStyle& SetFont(const UObject* InFontObject, const float InSize, const FName& InTypefaceFontName = NAME_None) { Font = FSlateFontInfo(InFontObject, InSize, InTypefaceFontName); return *this; }
+	FTextBlockStyle& SetFont(const FName& InFontName, float InSize) { Font = FSlateFontInfo(InFontName, InSize); return *this; }
+	FTextBlockStyle& SetFont(const FString& InFontName, float InSize) { Font = FSlateFontInfo(*InFontName, InSize); return *this; }
+	FTextBlockStyle& SetFont(const WIDECHAR* InFontName, float InSize) { Font = FSlateFontInfo(InFontName, InSize); return *this; }
+	FTextBlockStyle& SetFont(const ANSICHAR* InFontName, float InSize) { Font = FSlateFontInfo(InFontName, InSize); return *this; }
 	FTextBlockStyle& SetFontName(const FName& InFontName) { Font = FSlateFontInfo(InFontName, Font.GetClampSize()); return *this; }
 	FTextBlockStyle& SetFontName(const FString& InFontName) { Font = FSlateFontInfo(InFontName, Font.GetClampSize()); return *this; }
 	FTextBlockStyle& SetFontName(const WIDECHAR* InFontName) { Font = FSlateFontInfo(InFontName, Font.GetClampSize()); return *this; }
 	FTextBlockStyle& SetFontName(const ANSICHAR* InFontName) { Font = FSlateFontInfo(InFontName, Font.GetClampSize()); return *this; }
-	FTextBlockStyle& SetFontSize(uint16 InSize) { Font.Size = InSize; return *this; }
+	FTextBlockStyle& SetFontSize(float InSize) { Font.Size = InSize; return *this; }
 	FTextBlockStyle& SetTypefaceFontName(const FName& InTypefaceFontName) { Font.TypefaceFontName = InTypefaceFontName; return *this; }
 	FTextBlockStyle& SetFontMaterial(UObject* InMaterial) { Font.FontMaterial = InMaterial; return *this; }
 	FTextBlockStyle& SetFontOutlineMaterial(UObject* InMaterial) { Font.OutlineSettings.OutlineMaterial = InMaterial; return *this; }
@@ -799,7 +799,7 @@ struct FEditableTextStyle : public FSlateWidgetStyle
 	UPROPERTY(EditAnywhere, BlueprintReadWrite, Category=Appearance)
 	FSlateFontInfo Font;
 	FEditableTextStyle& SetFont(const FSlateFontInfo& InFont) { Font = InFont; return *this; }
-	FEditableTextStyle& SetFont(const FName& InFontName, uint16 InSize) { Font = FSlateFontInfo(InFontName, InSize); return *this; }
+	FEditableTextStyle& SetFont(const FName& InFontName, float InSize) { Font = FSlateFontInfo(InFontName, InSize); return *this; }
 	FEditableTextStyle& SetFontMaterial(UObject* InMaterial) { Font.FontMaterial = InMaterial; return *this; }
 	FEditableTextStyle& SetFontOutlineMaterial(UObject* InMaterial) { Font.OutlineSettings.OutlineMaterial = InMaterial; return *this; }
 
@@ -980,7 +980,7 @@ struct FEditableTextBoxStyle : public FSlateWidgetStyle
 	FSlateFontInfo Font_DEPRECATED;
 #endif
 	FEditableTextBoxStyle& SetFont(const FSlateFontInfo& InFont) { TextStyle.Font = InFont; return *this; }
-	FEditableTextBoxStyle& SetFont(const FName& InFontName, uint16 InSize) { return SetFont(FSlateFontInfo(InFontName, InSize)); }
+	FEditableTextBoxStyle& SetFont(const FName& InFontName, float InSize) { return SetFont(FSlateFontInfo(InFontName, InSize)); }
 
 	/** The style of the text block, which dictates the font, color, and shadow options. Style overrides all other properties! */
 	UPROPERTY(EditAnywhere, BlueprintReadWrite, Category = Appearance)
diff --git a/Engine/Source/Runtime/SlateCore/Public/Styling/StarshipCoreStyle.h b/Engine/Source/Runtime/SlateCore/Public/Styling/StarshipCoreStyle.h
--- a/Engine/Source/Runtime/SlateCore/Public/Styling/StarshipCoreStyle.h
+++ b/Engine/Source/Runtime/SlateCore/Public/Styling/StarshipCoreStyle.h
@@ -45,7 +45,7 @@ public:
 	static SLATECORE_API TSharedRef<const FCompositeFont> GetDefaultFont();
 
 	/** Get a font style using the default for for Slate */
-	static SLATECORE_API FSlateFontInfo GetDefaultFontStyle(const FName InTypefaceFontName, const int32 InSize, const FFontOutlineSettings& InOutlineSettings = FFontOutlineSettings());
+	static SLATECORE_API FSlateFontInfo GetDefaultFontStyle(const FName InTypefaceFontName, const float InSize, const FFontOutlineSettings& InOutlineSettings = FFontOutlineSettings());
 
 	static SLATECORE_API void ResetToDefault( );
 
diff --git a/Engine/Source/Runtime/SlateCore/Public/Styling/UMGCoreStyle.h b/Engine/Source/Runtime/SlateCore/Public/Styling/UMGCoreStyle.h
--- a/Engine/Source/Runtime/SlateCore/Public/Styling/UMGCoreStyle.h
+++ b/Engine/Source/Runtime/SlateCore/Public/Styling/UMGCoreStyle.h
@@ -24,7 +24,7 @@ public:
 	static SLATECORE_API TSharedRef<const FCompositeFont> GetDefaultFont();
 
 	/** Get a font style using the default for for Slate */
-	static SLATECORE_API FSlateFontInfo GetDefaultFontStyle(const FName InTypefaceFontName, const int32 InSize, const FFontOutlineSettings& InOutlineSettings = FFontOutlineSettings());
+	static SLATECORE_API FSlateFontInfo GetDefaultFontStyle(const FName InTypefaceFontName, const float InSize, const FFontOutlineSettings& InOutlineSettings = FFontOutlineSettings());
 
 	static SLATECORE_API void ResetToDefault( );
 
diff --git a/Engine/Source/Runtime/UMG/Private/WidgetBlueprintLibrary.cpp b/Engine/Source/Runtime/UMG/Private/WidgetBlueprintLibrary.cpp
--- a/Engine/Source/Runtime/UMG/Private/WidgetBlueprintLibrary.cpp
+++ b/Engine/Source/Runtime/UMG/Private/WidgetBlueprintLibrary.cpp
@@ -271,7 +271,7 @@ void UWidgetBlueprintLibrary::DrawText(FPaintContext& Context, const FString& In
 		Tint);
 }
 
-void UWidgetBlueprintLibrary::DrawTextFormatted(FPaintContext& Context, const FText& Text, FVector2D Position, UFont* Font, int32 FontSize, FName FontTypeFace, FLinearColor Tint)
+void UWidgetBlueprintLibrary::DrawTextFormatted(FPaintContext& Context, const FText& Text, FVector2D Position, UFont* Font, float FontSize, FName FontTypeFace, FLinearColor Tint)
 {
 	if ( Font )
 	{
diff --git a/Engine/Source/Runtime/UMG/Public/Blueprint/WidgetBlueprintLibrary.h b/Engine/Source/Runtime/UMG/Public/Blueprint/WidgetBlueprintLibrary.h
--- a/Engine/Source/Runtime/UMG/Public/Blueprint/WidgetBlueprintLibrary.h
+++ b/Engine/Source/Runtime/UMG/Public/Blueprint/WidgetBlueprintLibrary.h
@@ -129,7 +129,7 @@ public:
 	 * @param Tint			Color to render the line.
 	 */
 	UFUNCTION(BlueprintCallable, Category="Painting", meta=(DisplayName="Draw Text"))
-	static void DrawTextFormatted(UPARAM(ref) FPaintContext& Context, const FText& Text, FVector2D Position, UFont* Font, int32 FontSize = 16, FName FontTypeFace = FName(TEXT("Regular")), FLinearColor Tint = FLinearColor::White);
+	static void DrawTextFormatted(UPARAM(ref) FPaintContext& Context, const FText& Text, FVector2D Position, UFont* Font, float FontSize = 16.0f, FName FontTypeFace = FName(TEXT("Regular")), FLinearColor Tint = FLinearColor::White);
 
 	/**
 	 * The event reply to use when you choose to handle an event.  This will prevent the event 
