Added move semantics to TSlotBase::operator[] declarative syntax.
Removed a reference counting workaround that was added in CL# 3718295 for an issue that was fixed in CL# 3718782.
Removed unnecessary reference count churn from TDecl, and thus every SNew call.

diff --git a/Engine/Source/Runtime/SlateCore/Public/SlotBase.h b/Engine/Source/Runtime/SlateCore/Public/SlotBase.h
--- a/Engine/Source/Runtime/SlateCore/Public/SlotBase.h
+++ b/Engine/Source/Runtime/SlateCore/Public/SlotBase.h
@@ -47,11 +47,14 @@ public:
 	SLATECORE_API void SetOwner(const FChildren& Children);
 
 	/** Attach the child widget the slot now owns. */
+	FORCEINLINE_DEBUGGABLE void AttachWidget(TSharedRef<SWidget>&& InWidget)
+	{
+		DetatchParentFromContent();
+		Widget = MoveTemp(InWidget);
+		AfterContentOrOwnerAssigned();
+	}
 	FORCEINLINE_DEBUGGABLE void AttachWidget( const TSharedRef<SWidget>& InWidget )
 	{
-		// TODO: If we don't hold a reference here, ~SWidget() could called on the old widget before the assignment takes place
-		// The behavior of TShareRef is going to change in the near future to avoid this issue and this should then be reverted.
-		TSharedRef<SWidget> LocalWidget = Widget;
 		DetatchParentFromContent();
 		Widget = InWidget;
 		AfterContentOrOwnerAssigned();
@@ -130,6 +133,11 @@ class TSlotBase : public FSlotBase
 public:
 	using FSlotBase::FSlotBase;
 
+	SlotType& operator[](TSharedRef<SWidget>&& InChildWidget)
+	{
+		this->AttachWidget(MoveTemp(InChildWidget));
+		return static_cast<SlotType&>(*this);
+	}
 	SlotType& operator[]( const TSharedRef<SWidget>& InChildWidget )
 	{
 		this->AttachWidget(InChildWidget);
@@ -164,6 +172,11 @@ public:
 
 	public:
 		/** Attach the child widget the slot will own. */
+		typename SlotType::FSlotArguments& operator[](TSharedRef<SWidget>&& InChildWidget)
+		{
+			ChildWidget = MoveTemp(InChildWidget);
+			return Me();
+		}
 		typename SlotType::FSlotArguments& operator[](const TSharedRef<SWidget>& InChildWidget)
 		{
 			ChildWidget = InChildWidget;
diff --git a/Engine/Source/Runtime/SlateCore/Public/Widgets/DeclarativeSyntaxSupport.h b/Engine/Source/Runtime/SlateCore/Public/Widgets/DeclarativeSyntaxSupport.h
--- a/Engine/Source/Runtime/SlateCore/Public/Widgets/DeclarativeSyntaxSupport.h
+++ b/Engine/Source/Runtime/SlateCore/Public/Widgets/DeclarativeSyntaxSupport.h
@@ -739,7 +739,7 @@ namespace RequiredArgs
 		}
 
 		template<class WidgetType>
-		void CallConstruct(const TSharedRef<WidgetType>& OnWidget, const typename WidgetType::FArguments& WithNamedArgs) const
+		void CallConstruct(WidgetType* OnWidget, const typename WidgetType::FArguments& WithNamedArgs) const
 		{
 			// YOUR WIDGET MUST IMPLEMENT void Construct(const FArguments& InArgs)
 			OnWidget->Construct(WithNamedArgs);
@@ -755,7 +755,7 @@ namespace RequiredArgs
 		}
 
 		template<class WidgetType>
-		void CallConstruct(const TSharedRef<WidgetType>& OnWidget, const typename WidgetType::FArguments& WithNamedArgs) const
+		void CallConstruct(WidgetType* OnWidget, const typename WidgetType::FArguments& WithNamedArgs) const
 		{
 			// YOUR WIDGET MUST IMPLEMENT void Construct(const FArguments& InArgs)
 			OnWidget->Construct(WithNamedArgs, Forward<Arg0Type>(Arg0));
@@ -774,7 +774,7 @@ namespace RequiredArgs
 		}
 
 		template<class WidgetType>
-		void CallConstruct(const TSharedRef<WidgetType>& OnWidget, const typename WidgetType::FArguments& WithNamedArgs) const
+		void CallConstruct(WidgetType* OnWidget, const typename WidgetType::FArguments& WithNamedArgs) const
 		{
 			// YOUR WIDGET MUST IMPLEMENT Construct(const FArguments& InArgs)
 			OnWidget->Construct(WithNamedArgs, Forward<Arg0Type>(Arg0), Forward<Arg1Type>(Arg1));
@@ -795,7 +795,7 @@ namespace RequiredArgs
 		}
 
 		template<class WidgetType>
-		void CallConstruct(const TSharedRef<WidgetType>& OnWidget, const typename WidgetType::FArguments& WithNamedArgs) const
+		void CallConstruct(WidgetType* OnWidget, const typename WidgetType::FArguments& WithNamedArgs) const
 		{
 			// YOUR WIDGET MUST IMPLEMENT Construct(const FArguments& InArgs)
 			OnWidget->Construct(WithNamedArgs, Forward<Arg0Type>(Arg0), Forward<Arg1Type>(Arg1), Forward<Arg2Type>(Arg2));
@@ -818,7 +818,7 @@ namespace RequiredArgs
 		}
 
 		template<class WidgetType>
-		void CallConstruct(const TSharedRef<WidgetType>& OnWidget, const typename WidgetType::FArguments& WithNamedArgs) const
+		void CallConstruct(WidgetType* OnWidget, const typename WidgetType::FArguments& WithNamedArgs) const
 		{
 			// YOUR WIDGET MUST IMPLEMENT Construct(const FArguments& InArgs)
 			OnWidget->Construct(WithNamedArgs, Forward<Arg0Type>(Arg0), Forward<Arg1Type>(Arg1), Forward<Arg2Type>(Arg2), Forward<Arg3Type>(Arg3));
@@ -843,7 +843,7 @@ namespace RequiredArgs
 		}
 
 		template<class WidgetType>
-		void CallConstruct(const TSharedRef<WidgetType>& OnWidget, const typename WidgetType::FArguments& WithNamedArgs) const
+		void CallConstruct(WidgetType* OnWidget, const typename WidgetType::FArguments& WithNamedArgs) const
 		{
 			// YOUR WIDGET MUST IMPLEMENT Construct(const FArguments& InArgs)
 			OnWidget->Construct(WithNamedArgs, Forward<Arg0Type>(Arg0), Forward<Arg1Type>(Arg1), Forward<Arg2Type>(Arg2), Forward<Arg3Type>(Arg3), Forward<Arg4Type>(Arg4));
@@ -893,31 +893,6 @@ namespace RequiredArgs
 }
 
 
-/** Normal widgets are allocated directly by the TSlateDecl. */
-template<typename WidgetType, bool IsDerived>
-struct TWidgetAllocator
-{
-	static TSharedRef<WidgetType> PrivateAllocateWidget()
-	{
-		return MakeShared<WidgetType>();
-	}
-};
-
-/**
- * SUserWidgets are allocated in the corresponding CPP file, so that
- * the implementer can return an implementation that differs from the
- * public interface. @see SUserWidgetExample
- */
-template<typename WidgetType>
-struct TWidgetAllocator< WidgetType, true >
-{
-	static TSharedRef<WidgetType> PrivateAllocateWidget()
-	{
-		return WidgetType::New();
-	}
-};
-
-
 /**
  * Utility class used during widget instantiation.
  * Performs widget allocation and construction.
@@ -931,10 +906,24 @@ template<class WidgetType, typename RequiredArgsPayloadType>
 struct TSlateDecl
 {
 	TSlateDecl( const ANSICHAR* InType, const ANSICHAR* InFile, int32 OnLine, RequiredArgsPayloadType&& InRequiredArgs )
-		: _Widget( TWidgetAllocator<WidgetType, TIsDerivedFrom<WidgetType, SUserWidget>::IsDerived >::PrivateAllocateWidget() )
-		, _RequiredArgs(InRequiredArgs)
+		: _RequiredArgs(InRequiredArgs)
 	{
-		_Widget->SetDebugInfo( InType, InFile, OnLine, sizeof(WidgetType) );
+		if constexpr (std::is_base_of_v<SUserWidget, WidgetType>)
+		{
+			/**
+			 * SUserWidgets are allocated in the corresponding CPP file, so that
+			 * the implementer can return an implementation that differs from the
+			 * public interface. @see SUserWidgetExample
+			 */
+			_Widget = WidgetType::New();
+		}
+		else
+		{
+			/** Normal widgets are allocated directly by the TSlateDecl. */
+			_Widget = MakeShared<WidgetType>();
+		}
+
+		_Widget->SetDebugInfo(InType, InFile, OnLine, sizeof(WidgetType));
 	}
 
 	/**
@@ -942,10 +931,11 @@ struct TSlateDecl
 	 * @see SAssignNew
 	 */
 	template<class ExposeAsWidgetType>
-	TSlateDecl& Expose( TSharedPtr<ExposeAsWidgetType>& OutVarToInit )
+	TSlateDecl&& Expose( TSharedPtr<ExposeAsWidgetType>& OutVarToInit ) &&
 	{
+		// Can't move out _Widget here because operator<<= needs it
 		OutVarToInit = _Widget;
-		return *this;
+		return MoveTemp(*this);
 	}
 
 	/**
@@ -953,10 +943,11 @@ struct TSlateDecl
 	 * @see SAssignNew
 	 */
 	template<class ExposeAsWidgetType>
-	TSlateDecl& Expose( TSharedRef<ExposeAsWidgetType>& OutVarToInit )
+	TSlateDecl&& Expose( TSharedRef<ExposeAsWidgetType>& OutVarToInit ) &&
 	{
-		OutVarToInit = _Widget;
-		return *this;
+		// Can't move out _Widget here because operator<<= needs it
+		OutVarToInit = _Widget.ToSharedRef();
+		return MoveTemp(*this);
 	}
 
 	/**
@@ -964,10 +955,11 @@ struct TSlateDecl
 	 * @see SAssignNew
 	 */
 	template<class ExposeAsWidgetType>
-	TSlateDecl& Expose( TWeakPtr<ExposeAsWidgetType>& OutVarToInit )
+	TSlateDecl&& Expose( TWeakPtr<ExposeAsWidgetType>& OutVarToInit ) &&
 	{
-		OutVarToInit = _Widget;
-		return *this;
+		// Can't move out _Widget here because operator<<= needs it
+		OutVarToInit = _Widget.ToWeakPtr();
+		return MoveTemp(*this);
 	}
 
 	/**
@@ -977,17 +969,17 @@ struct TSlateDecl
 	 *
 	 * @return A reference to the widget that we constructed.
 	 */
-	TSharedRef<WidgetType> operator<<=( const typename WidgetType::FArguments& InArgs ) const
+	TSharedRef<WidgetType> operator<<=( const typename WidgetType::FArguments& InArgs ) &&
 	{
 		_Widget->SWidgetConstruct(InArgs);
-		_RequiredArgs.CallConstruct(_Widget, InArgs);
+		_RequiredArgs.CallConstruct(_Widget.Get(), InArgs);
 		_Widget->CacheVolatility();
 		_Widget->bIsDeclarativeSyntaxConstructionCompleted = true;
 
-		return _Widget;
+		return MoveTemp(_Widget).ToSharedRef();
 	}
 
-	const TSharedRef<WidgetType> _Widget;
+	TSharedPtr<WidgetType> _Widget;
 	RequiredArgsPayloadType& _RequiredArgs;
 };
 
