Remove payload indirection on slate elements.
Yields better cache coherency and measured better performance

diff --git a/Engine/Plugins/Tests/WidgetAutomationTests/Resources/TestBaselineFiles/UIAutomation_SlateTestTextBlock_ColorAndOpacity.json b/Engine/Plugins/Tests/WidgetAutomationTests/Resources/TestBaselineFiles/UIAutomation_SlateTestTextBlock_ColorAndOpacity.json
--- a/Engine/Plugins/Tests/WidgetAutomationTests/Resources/TestBaselineFiles/UIAutomation_SlateTestTextBlock_ColorAndOpacity.json
+++ b/Engine/Plugins/Tests/WidgetAutomationTests/Resources/TestBaselineFiles/UIAutomation_SlateTestTextBlock_ColorAndOpacity.json
@@ -8,7 +8,7 @@
 		"ForegroundColor": "(R=0.527115,G=0.527115,B=0.527115,A=1.000000)",
 		"SubduedForeground": "(R=0.527115,G=0.527115,B=0.527115,A=0.600000)"
 	},
-	"ShapedTextPayload":
+	"ShapedTextElement":
 	{
 		"Tint": "(R=0.500000,G=0.400000,B=0.300000,A=0.900000)",
 		"OutlineTint": "(R=0.000000,G=0.000000,B=0.000000,A=1.000000)"
diff --git a/Engine/Plugins/Tests/WidgetAutomationTests/Source/WidgetAutomationTests/Private/SlateTestHelper.cpp b/Engine/Plugins/Tests/WidgetAutomationTests/Source/WidgetAutomationTests/Private/SlateTestHelper.cpp
--- a/Engine/Plugins/Tests/WidgetAutomationTests/Source/WidgetAutomationTests/Private/SlateTestHelper.cpp
+++ b/Engine/Plugins/Tests/WidgetAutomationTests/Source/WidgetAutomationTests/Private/SlateTestHelper.cpp
@@ -2,6 +2,7 @@
 
 #include "SlateTestHelper.h"
 #include "CoreMinimal.h"
+#include "Rendering/DrawElements.h"
 #include "Rendering/DrawElementPayloads.h"
 #include "Serialization/JsonReader.h"
 #include "Serialization/JsonSerializer.h"
@@ -59,12 +60,12 @@ void FSlateTestHelper::ToJson(const FWidgetStyle& InWidgetStyle, TSharedPtr<FJso
 	MyJsonObject->SetObjectField("WidgetStyle", InnerJsonObject);
 }
 
-void FSlateTestHelper::ToJson(const FSlateShapedTextPayload& ShapedTextPayload, TSharedPtr<FJsonObject> MyJsonObject)
+void FSlateTestHelper::ToJson(const FSlateShapedTextElement& ShapedTextElement, TSharedPtr<FJsonObject> MyJsonObject)
 {
 	TSharedPtr<FJsonObject> InnerJsonObject = MakeShared<FJsonObject>();
-	InnerJsonObject->SetStringField("Tint", ShapedTextPayload.GetTint().ToString());
-	InnerJsonObject->SetStringField("OutlineTint", ShapedTextPayload.GetOutlineTint().ToString());
-	MyJsonObject->SetObjectField("ShapedTextPayload", InnerJsonObject);
+	InnerJsonObject->SetStringField("Tint", ShapedTextElement.GetTint().ToString());
+	InnerJsonObject->SetStringField("OutlineTint", ShapedTextElement.GetOutlineTint().ToString());
+	MyJsonObject->SetObjectField("ShapedTextElement", InnerJsonObject);
 }
 }
 #endif
\ No newline at end of file
diff --git a/Engine/Plugins/Tests/WidgetAutomationTests/Source/WidgetAutomationTests/Private/WidgetMockNonTemplate.cpp b/Engine/Plugins/Tests/WidgetAutomationTests/Source/WidgetAutomationTests/Private/WidgetMockNonTemplate.cpp
--- a/Engine/Plugins/Tests/WidgetAutomationTests/Source/WidgetAutomationTests/Private/WidgetMockNonTemplate.cpp
+++ b/Engine/Plugins/Tests/WidgetAutomationTests/Source/WidgetAutomationTests/Private/WidgetMockNonTemplate.cpp
@@ -104,11 +104,10 @@ void FWidgetMockNonTemplate::OnPaintDataValidation(const FGeometry& AllottedGeom
 				if (CachedElementList != nullptr)
 				{
 					//@TODO Support all element types and create their own JSON conversion function.
-					for (const FSlateDrawElement& DrawElements : CachedElementList->DrawElements.Get<(uint8)EElementType::ET_ShapedText>().Elements)
+					for (const FSlateShapedTextElement& DrawElements : CachedElementList->DrawElements.Get<(uint8)EElementType::ET_ShapedText>())
 					{
-						FSlateShapedTextPayload ShapedTextPayload = DrawElements.GetDataPayload<FSlateShapedTextPayload>();
-						if (TestJsonObject->GetObjectField("ShapedTextPayload")->GetStringField("Tint") == ShapedTextPayload.GetTint().ToString()
-							&& TestJsonObject->GetObjectField("ShapedTextPayload")->GetStringField("OutlineTint") == ShapedTextPayload.GetOutlineTint().ToString())
+						if (TestJsonObject->GetObjectField("ShapedTextElement")->GetStringField("Tint") == DrawElements.GetTint().ToString()
+							&& TestJsonObject->GetObjectField("ShapedTextElement")->GetStringField("OutlineTint") == DrawElements.GetOutlineTint().ToString())
 						{
 							if (!DataValidations.Contains(FTestFunctionNames::NAME_OnPaint))
 							{
@@ -138,10 +137,9 @@ void FWidgetMockNonTemplate::OnPaintDataValidation(const FGeometry& AllottedGeom
 			if (CachedElementList != nullptr)
 			{
 				//@TODO Support all element types and create their own JSON conversion function.
-				for (const FSlateDrawElement& DrawElements : CachedElementList->DrawElements.Get<(uint8)EElementType::ET_ShapedText>().Elements)
+				for (const FSlateShapedTextElement& DrawElements : CachedElementList->DrawElements.Get<(uint8)EElementType::ET_ShapedText>())
 				{
-					FSlateShapedTextPayload ShapedTextPayload = DrawElements.GetDataPayload<FSlateShapedTextPayload>();
-					FSlateTestHelper::ToJson(ShapedTextPayload, WriteJsonObject);
+					FSlateTestHelper::ToJson(DrawElements, WriteJsonObject);
 				}
 			}
 		}
diff --git a/Engine/Plugins/Tests/WidgetAutomationTests/Source/WidgetAutomationTests/Public/SlateTestHelper.h b/Engine/Plugins/Tests/WidgetAutomationTests/Source/WidgetAutomationTests/Public/SlateTestHelper.h
--- a/Engine/Plugins/Tests/WidgetAutomationTests/Source/WidgetAutomationTests/Public/SlateTestHelper.h
+++ b/Engine/Plugins/Tests/WidgetAutomationTests/Source/WidgetAutomationTests/Public/SlateTestHelper.h
@@ -9,7 +9,7 @@
 #if WITH_DEV_AUTOMATION_TESTS
 
 class JsonObject;
-struct FSlateShapedTextPayload;
+struct FSlateShapedTextElement;
 
 namespace UE::SlateWidgetAutomationTest
 {
@@ -28,7 +28,7 @@ public:
 	// JSON conversion functions.
 	static void ToJson(const FGeometry& MyGeometry, TSharedPtr<FJsonObject> MyJsonObject);
 	static void ToJson(const FWidgetStyle& InWidgetStyle, TSharedPtr<FJsonObject> MyJsonObject);
-	static void ToJson(const FSlateShapedTextPayload& ShapedTextPayload, TSharedPtr<FJsonObject> MyJsonObject);
+	static void ToJson(const FSlateShapedTextElement& ShapedTextElement, TSharedPtr<FJsonObject> MyJsonObject);
 };
 }
 #endif
\ No newline at end of file
diff --git a/Engine/Source/Developer/SlateReflector/Private/VisualTreeCapture.cpp b/Engine/Source/Developer/SlateReflector/Private/VisualTreeCapture.cpp
--- a/Engine/Source/Developer/SlateReflector/Private/VisualTreeCapture.cpp
+++ b/Engine/Source/Developer/SlateReflector/Private/VisualTreeCapture.cpp
@@ -59,11 +59,11 @@ void FVisualEntry::Resolve(const FSlateWindowElementList& ElementList)
 		return;
 	}
 
-	auto ResolveBounds = [&](const FSlateDrawElementContainer& Container, uint8 InElementType)
+	auto ResolveBounds = [&](const auto& Container, uint8 InElementType)
 	{
 		if (InElementType == (uint8)ElementType)
 		{
-			const FSlateDrawElement& Element = Container.Elements[ElementIndex];
+			const FSlateDrawElement& Element = Container[ElementIndex];
 			const FSlateRenderTransform& Transform = Element.GetRenderTransform();
 			const FVector2D LocalSize = Element.GetLocalSize();
 
@@ -187,14 +187,13 @@ void FVisualTreeCapture::AddInvalidationRootCachedEntries(TSharedRef<FVisualTree
 	const TArray<TSharedPtr<FSlateCachedElementList>>& CachedElements = Data.GetCachedElementLists();
 	for (const TSharedPtr<FSlateCachedElementList>& CachedElement : CachedElements)
 	{
-		auto AddTypedEntries = [&](FSlateDrawElementContainer& Container)
+		auto AddTypedEntries = [&](auto& Container)
 		{
-			const FSlateDrawElementArray& DrawElementArray = Container.Elements;
 			const SWidget* Widget = CachedElement->OwningWidget;
 			// todo, should check if parents has the metadata also
 			if (Widget && !Widget->GetMetaData<FInvisibleToWidgetReflectorMetaData>())
 			{
-				for (const FSlateDrawElement& Element : DrawElementArray)
+				for (const FSlateDrawElement& Element : Container)
 				{
 					const int32 EntryIndex = Tree->Entries.Emplace(Widget->AsShared(), Element);
 
diff --git a/Engine/Source/Runtime/SlateCore/Private/Debugging/SlateDebugging.cpp b/Engine/Source/Runtime/SlateCore/Private/Debugging/SlateDebugging.cpp
--- a/Engine/Source/Runtime/SlateCore/Private/Debugging/SlateDebugging.cpp
+++ b/Engine/Source/Runtime/SlateCore/Private/Debugging/SlateDebugging.cpp
@@ -10,6 +10,7 @@
 #include "Misc/StringBuilder.h"
 #include "ProfilingDebugging/CsvProfiler.h"
 #include "Rendering/DrawElements.h"
+#include "Rendering/DrawElementPayloads.h"
 #include "SlateGlobals.h"
 #include "Styling/CoreStyle.h"
 #include "Types/ReflectionMetadata.h"
diff --git a/Engine/Source/Runtime/SlateCore/Private/FastUpdate/SlateInvalidationRoot.cpp b/Engine/Source/Runtime/SlateCore/Private/FastUpdate/SlateInvalidationRoot.cpp
--- a/Engine/Source/Runtime/SlateCore/Private/FastUpdate/SlateInvalidationRoot.cpp
+++ b/Engine/Source/Runtime/SlateCore/Private/FastUpdate/SlateInvalidationRoot.cpp
@@ -15,6 +15,7 @@
 #include "Trace/SlateTrace.h"
 #include "Types/ReflectionMetadata.h"
 #include "Types/SlateAttributeMetaData.h"
+#include "Rendering/DrawElementPayloads.h"
 
 CSV_DECLARE_CATEGORY_MODULE_EXTERN(SLATECORE_API, Slate);
 
@@ -378,7 +379,8 @@ FSlateInvalidationResult FSlateInvalidationRoot::PaintInvalidationRoot(const FSl
 	if (!Context.bAllowFastPathUpdate || bNeedsSlowPath || GSlateIsInInvalidationSlowPath)
 	{
 		SCOPED_NAMED_EVENT(Slate_PaintSlowPath, FColor::Red);
-
+		
+		// Clears existing cached element lists
 		ClearAllFastPathData(!Context.bAllowFastPathUpdate);
 
 		GSlateIsOnFastUpdatePath = false;
diff --git a/Engine/Source/Runtime/SlateCore/Private/Rendering/DrawElementTypes.cpp b/Engine/Source/Runtime/SlateCore/Private/Rendering/DrawElementTypes.cpp
new file mode 100644
--- /dev/null
+++ b/Engine/Source/Runtime/SlateCore/Private/Rendering/DrawElementTypes.cpp
@@ -0,0 +1,742 @@
+// Copyright Epic Games, Inc. All Rights Reserved.
+
+#include "Rendering/DrawElementTypes.h"
+#include "Rendering/DrawElements.h"
+#include "Widgets/SWidget.h"
+#include "Widgets/SWindow.h"
+#include "HAL/IConsoleManager.h"
+#include "Types/ReflectionMetadata.h"
+#include "Fonts/ShapedTextFwd.h"
+#include "Fonts/FontCache.h"
+#include "Rendering/SlateObjectReferenceCollector.h"
+#include "Debugging/SlateDebugging.h"
+#include "Application/SlateApplicationBase.h"
+
+#include <limits>
+#include <type_traits>
+
+DECLARE_CYCLE_STAT(TEXT("FSlateDrawElement::Make Time"), STAT_SlateDrawElementMakeTime, STATGROUP_SlateVerbose);
+DECLARE_CYCLE_STAT(TEXT("FSlateDrawElement::MakeCustomVerts Time"), STAT_SlateDrawElementMakeCustomVertsTime, STATGROUP_Slate);
+DECLARE_CYCLE_STAT(TEXT("FSlateDrawElement::Prebatch Time"), STAT_SlateDrawElementPrebatchTime, STATGROUP_Slate);
+
+DEFINE_STAT(STAT_SlateBufferPoolMemory);
+DEFINE_STAT(STAT_SlateCachedDrawElementMemory);
+
+static bool bApplyDisabledEffectOnWidgets = true;
+
+static void HandleApplyDisabledEffectOnWidgetsToggled(IConsoleVariable* CVar)
+{
+	if (FSlateApplicationBase::IsInitialized())
+	{
+		FSlateApplicationBase::Get().InvalidateAllWidgets(false);
+	}
+}
+
+static FAutoConsoleVariableRef CVarApplyDisabledEffectOnWidgets(
+	TEXT("Slate.ApplyDisabledEffectOnWidgets"),
+	bApplyDisabledEffectOnWidgets,
+	TEXT("If true, disabled game-layer widgets will have alpha multiplied by 0.45."),
+	FConsoleVariableDelegate::CreateStatic(&HandleApplyDisabledEffectOnWidgetsToggled)
+);
+
+static bool IsResourceObjectValid(UObject*& InObject)
+{
+	if (InObject != nullptr && (!IsValidChecked(InObject) || InObject->IsUnreachable() || InObject->HasAnyFlags(RF_BeginDestroyed)))
+	{
+		UE_LOG(LogSlate, Warning, TEXT("Attempted to access resource for %s which is pending kill, unreachable or pending destroy"), *InObject->GetName());
+		return false;
+	}
+
+	return true;
+}
+
+static bool ShouldCull(const FSlateWindowElementList& ElementList)
+{
+	const FSlateClippingManager& ClippingManager = ElementList.GetClippingManager();
+	const int32 CurrentIndex = ClippingManager.GetClippingIndex();
+	if (CurrentIndex != INDEX_NONE)
+	{
+		const FSlateClippingState& ClippingState = ClippingManager.GetClippingStates()[CurrentIndex];
+		if (ClippingState.GetClippingMethod() == EClippingMethod::Scissor)
+		{
+			return ClippingState.HasZeroArea();
+		}
+		else if (ClippingState.GetClippingMethod() == EClippingMethod::Stencil)
+		{
+			FSlateRect WindowRect = ElementList.GetPaintWindow()->GetClippingRectangleInWindow();
+			if (WindowRect.GetArea() > 0)
+			{
+				for (const FSlateClippingZone& Stencil : ClippingState.StencilQuads)
+				{
+					bool bOverlapping = false;
+					FSlateRect ClippedStencil = Stencil.GetBoundingBox().IntersectionWith(WindowRect, bOverlapping);
+					
+					if (!bOverlapping || ClippedStencil.GetArea() <= 0)
+					{
+						return true;
+					}
+				}
+			}
+		}
+
+	}
+
+	return false;
+}
+
+static bool ShouldCull(const FSlateWindowElementList& ElementList, const FPaintGeometry& PaintGeometry)
+{
+	const FVector2f LocalSize = UE::Slate::CastToVector2f(PaintGeometry.GetLocalSize());
+	const float DrawScale = PaintGeometry.DrawScale;
+
+	const FVector2f PixelSize = (LocalSize * DrawScale);
+	if (PixelSize.X <= 0.f || PixelSize.Y <= 0.f)
+	{
+		return true;
+	}
+
+	return ShouldCull(ElementList);
+}
+
+static bool ShouldCull(const FSlateWindowElementList& ElementList, const FPaintGeometry& PaintGeometry, const FSlateBrush* InBrush)
+{
+	if (ShouldCull(ElementList, PaintGeometry))
+	{
+		return true;
+	}
+
+	if (InBrush->GetDrawType() == ESlateBrushDrawType::NoDrawType)
+	{
+		return true;
+	}
+
+	UObject* ResourceObject = InBrush->GetResourceObject();
+	if (!IsResourceObjectValid(ResourceObject))
+	{
+		return true;
+	}
+
+	return false;
+}
+
+
+static bool ShouldCull(const FSlateWindowElementList& ElementList, const FPaintGeometry& PaintGeometry, const FLinearColor& InTint)
+{
+	if (InTint.A == 0 || ShouldCull(ElementList, PaintGeometry))
+	{
+		return true;
+	}
+
+	return false;
+}
+
+static bool ShouldCull(const FSlateWindowElementList& ElementList, const FPaintGeometry& PaintGeometry, const FLinearColor& InTint, const FString& InText)
+{
+	if (InTint.A == 0 || InText.Len() == 0 || ShouldCull(ElementList, PaintGeometry))
+	{
+		return true;
+	}
+
+	return false;
+}
+
+static bool ShouldCull(const FSlateWindowElementList& ElementList, const FPaintGeometry& PaintGeometry, const FSlateBrush* InBrush, const FLinearColor& InTint)
+{
+	bool bIsFullyTransparent = InTint.A == 0 && (InBrush->OutlineSettings.Color.GetSpecifiedColor().A == 0 || InBrush->OutlineSettings.bUseBrushTransparency);
+	if (bIsFullyTransparent || ShouldCull(ElementList, PaintGeometry, InBrush))
+	{
+		return true;
+	}
+
+	return false;
+}
+
+void FSlateDrawElement::Init(FSlateWindowElementList& ElementList, EElementType InElementType, uint32 InLayer, const FPaintGeometry& PaintGeometry, ESlateDrawEffect InDrawEffects)
+{
+	if (ElementList.GetIsInGameLayer() && !bApplyDisabledEffectOnWidgets)
+	{
+		InDrawEffects &= ~ESlateDrawEffect::DisabledEffect;
+	}
+
+	// Apply the pixel snapping effects, if it's unset/inherited we don't do anything.
+	switch (ElementList.GetPixelSnappingMethod())
+	{
+		case EWidgetPixelSnapping::Disabled:
+			InDrawEffects |= ESlateDrawEffect::NoPixelSnapping;
+			break;
+		case EWidgetPixelSnapping::Inherit:
+			// NOTE: Do nothing on inherit.
+		case EWidgetPixelSnapping::SnapToPixel:
+			// NOTE: We don't remove the drawing effects pixel snapping here.  The default behavior of InDrawEffects is always to snap to pixels
+			// so all we need to do here is not to snap.  This will allow the continued behavior that if someone does set snap to pixel at a
+			// higher level, it won't break people that directly set InDrawEffects |= ESlateDrawEffect::NoPixelSnapping, before drawing their
+			// draw element.
+		default:
+			break;
+	}
+
+	RenderTransform = PaintGeometry.GetAccumulatedRenderTransform();
+	Position = PaintGeometry.DrawPosition;
+	Scale = PaintGeometry.DrawScale;
+	LocalSize = FVector2f(PaintGeometry.GetLocalSize());
+	ClipStateHandle.SetPreCachedClipIndex(ElementList.GetClippingIndex());
+
+#if UE_SLATE_VERIFY_PIXELSIZE
+	const FVector2f PixelSize = FVector2f(LocalSize * Scale);
+	ensureMsgf(PixelSize.X >= 0.f && PixelSize.X <= (float)std::numeric_limits<uint16>::max(), TEXT("The size X '%f' is too small or big to fit in the SlateVertex buffer."), PixelSize.X);
+	ensureMsgf(PixelSize.Y >= 0.f && PixelSize.Y <= (float)std::numeric_limits<uint16>::max(), TEXT("The size Y '%f' is too small or big to fit in the SlateVertex buffer."), PixelSize.Y);
+#endif
+
+	LayerId = InLayer;
+
+	ElementType = InElementType;
+	DrawEffects = InDrawEffects;
+	
+	// Calculate the layout to render transform as this is needed by several calculations downstream.
+	const FSlateLayoutTransform InverseLayoutTransform(Inverse(FSlateLayoutTransform(Scale, Position)));
+
+	{
+		// This is a workaround because we want to keep track of the various Scenes 
+		// in use throughout the UI. We keep a synchronized set with the render thread on the SlateRenderer and 
+		// use indices to synchronize between them.
+		FSlateRenderer* Renderer = FSlateApplicationBase::Get().GetRenderer();
+		checkSlow(Renderer);
+		int32 SceneIndexLong = Renderer->GetCurrentSceneIndex();
+		ensureMsgf(SceneIndexLong <= std::numeric_limits<int8>::max(), TEXT("The Scene index is saved as an int8 in the DrawElements."));
+		SceneIndex = (int8)SceneIndexLong;
+	}
+
+	BatchFlags = ESlateBatchDrawFlag::None;
+	BatchFlags |= static_cast<ESlateBatchDrawFlag>(static_cast<uint32>(InDrawEffects) & static_cast<uint32>(ESlateDrawEffect::NoBlending | ESlateDrawEffect::PreMultipliedAlpha | ESlateDrawEffect::NoGamma | ESlateDrawEffect::InvertAlpha));
+
+	static_assert(((__underlying_type(ESlateDrawEffect))ESlateDrawEffect::NoBlending) == ((__underlying_type(ESlateBatchDrawFlag))ESlateBatchDrawFlag::NoBlending), "Must keep ESlateBatchDrawFlag and ESlateDrawEffect partial matches");
+	static_assert(((__underlying_type(ESlateDrawEffect))ESlateDrawEffect::PreMultipliedAlpha) == ((__underlying_type(ESlateBatchDrawFlag))ESlateBatchDrawFlag::PreMultipliedAlpha), "Must keep ESlateBatchDrawFlag and ESlateDrawEffect partial matches");
+	static_assert(((__underlying_type(ESlateDrawEffect))ESlateDrawEffect::NoGamma) == ((__underlying_type(ESlateBatchDrawFlag))ESlateBatchDrawFlag::NoGamma), "Must keep ESlateBatchDrawFlag and ESlateDrawEffect partial matches");
+	static_assert(((__underlying_type(ESlateDrawEffect))ESlateDrawEffect::InvertAlpha) == ((__underlying_type(ESlateBatchDrawFlag))ESlateBatchDrawFlag::InvertAlpha), "Must keep ESlateBatchDrawFlag and ESlateDrawEffect partial matches");
+	if ((InDrawEffects & ESlateDrawEffect::ReverseGamma) != ESlateDrawEffect::None)
+	{
+		BatchFlags |= ESlateBatchDrawFlag::ReverseGamma;
+	}
+}
+
+void FSlateDrawElement::ApplyPositionOffset(UE::Slate::FDeprecateVector2DParameter InOffset)
+{
+	SetPosition(Position + InOffset);
+	RenderTransform = Concatenate(RenderTransform, InOffset);
+
+	// Recompute cached layout to render transform
+	const FSlateLayoutTransform InverseLayoutTransform(Inverse(FSlateLayoutTransform(Scale, Position)));
+}
+
+void FSlateDrawElement::MakeDebugQuad( FSlateWindowElementList& ElementList, uint32 InLayer, const FPaintGeometry& PaintGeometry, FLinearColor Tint)
+{
+	PaintGeometry.CommitTransformsIfUsingLegacyConstructor();
+
+	if (ShouldCull(ElementList))
+	{
+		return;
+	}
+
+	FSlateBoxElement& Element = ElementList.AddUninitialized<EElementType::ET_DebugQuad>();
+
+	Element.SetTint(Tint);
+	Element.Init(ElementList, EElementType::ET_DebugQuad, InLayer, PaintGeometry, ESlateDrawEffect::None);
+}
+
+FSlateDrawElement& FSlateDrawElement::MakeBoxInternal(
+	FSlateWindowElementList& ElementList,
+	uint32 InLayer,
+	const FPaintGeometry& PaintGeometry,
+	const FSlateBrush* InBrush,
+	ESlateDrawEffect InDrawEffects,
+	const FLinearColor& InTint
+)
+{
+	EElementType ElementType = (InBrush->DrawAs == ESlateBrushDrawType::Border) ? EElementType::ET_Border : (InBrush->DrawAs == ESlateBrushDrawType::RoundedBox) ? EElementType::ET_RoundedBox : EElementType::ET_Box;
+
+	// Cast to Rounded Rect to get the internal parameters 
+	// New payload type - inherit from BoxPayload 
+
+	switch (ElementType)
+	{
+		case EElementType::ET_Box:
+		{
+			FSlateBoxElement& Element = ElementList.AddUninitialized<EElementType::ET_Box>();
+			Element.SetTint(InTint);
+			Element.SetBrush(InBrush, PaintGeometry.GetLocalSize(), PaintGeometry.DrawScale);
+			Element.Init(ElementList, ElementType, InLayer, PaintGeometry, InDrawEffects);
+			return Element;
+		}
+		break;
+		case EElementType::ET_RoundedBox:
+		{
+			FSlateRoundedBoxElement& Element = ElementList.AddUninitialized<EElementType::ET_RoundedBox>();
+			FVector4f CornerRadii = FVector4f(InBrush->OutlineSettings.CornerRadii);
+			if (InBrush->OutlineSettings.RoundingType == ESlateBrushRoundingType::HalfHeightRadius)
+			{
+				const float UniformRadius = PaintGeometry.GetLocalSize().Y / 2.0f;
+				CornerRadii = FVector4f(UniformRadius, UniformRadius, UniformRadius, UniformRadius);
+			}
+			Element.SetRadius(CornerRadii);
+
+			if (InBrush->OutlineSettings.bUseBrushTransparency)
+			{
+				FLinearColor Color = InBrush->OutlineSettings.Color.GetSpecifiedColor().CopyWithNewOpacity(InTint.A);
+				Element.SetOutline(Color, InBrush->OutlineSettings.Width);
+			}
+			else
+			{
+				Element.SetOutline(InBrush->OutlineSettings.Color.GetSpecifiedColor(), InBrush->OutlineSettings.Width);
+			}
+			Element.SetTint(InTint);
+			Element.SetBrush(InBrush, PaintGeometry.GetLocalSize(), PaintGeometry.DrawScale);
+			Element.Init(ElementList, ElementType, InLayer, PaintGeometry, InDrawEffects);
+			return Element;
+		}
+		break;
+		case EElementType::ET_Border:
+		{
+			FSlateBoxElement& Element = ElementList.AddUninitialized<EElementType::ET_Border>();
+			Element.SetTint(InTint);
+			Element.SetBrush(InBrush, PaintGeometry.GetLocalSize(), PaintGeometry.DrawScale);
+			Element.Init(ElementList, ElementType, InLayer, PaintGeometry, InDrawEffects);
+			return Element;
+
+		}
+		break;
+		default:
+		{
+			checkf(0, TEXT("Invalid box element"));
+			FSlateBoxElement& Element = ElementList.AddUninitialized<EElementType::ET_Box>();
+			Element.SetTint(InTint);
+			Element.SetBrush(InBrush, PaintGeometry.GetLocalSize(), PaintGeometry.DrawScale);
+			Element.Init(ElementList, ElementType, InLayer, PaintGeometry, InDrawEffects);
+			return Element;
+		}
+	}
+}
+
+void FSlateDrawElement::MakeBox(
+	FSlateWindowElementList& ElementList,
+	uint32 InLayer, 
+	const FPaintGeometry& PaintGeometry, 
+	const FSlateBrush* InBrush,
+	ESlateDrawEffect InDrawEffects, 
+	const FLinearColor& InTint)
+{
+	PaintGeometry.CommitTransformsIfUsingLegacyConstructor();
+
+	if (ShouldCull(ElementList, PaintGeometry, InBrush, InTint))
+	{
+		return;
+	}
+
+	MakeBoxInternal(ElementList, InLayer, PaintGeometry, InBrush, InDrawEffects, InTint);
+}
+
+void FSlateDrawElement::MakeRotatedBox(
+	FSlateWindowElementList& ElementList,
+	uint32 InLayer,
+	const FPaintGeometry& PaintGeometry,
+	const FSlateBrush* InBrush,
+	ESlateDrawEffect InDrawEffects,
+	float Angle2D,
+	UE::Slate::FDeprecateOptionalVector2DParameter InRotationPoint,
+	ERotationSpace RotationSpace,
+	const FLinearColor& InTint)
+{
+	PaintGeometry.CommitTransformsIfUsingLegacyConstructor();
+
+	if (ShouldCull(ElementList, PaintGeometry, InBrush, InTint))
+	{
+		return;
+	}
+
+	FSlateDrawElement& DrawElement = MakeBoxInternal(ElementList, InLayer, PaintGeometry, InBrush, InDrawEffects, InTint);
+	
+	if (Angle2D != 0.0f)
+	{
+		const FVector2f RotationPoint = GetRotationPoint(PaintGeometry, InRotationPoint, RotationSpace);
+		const FSlateRenderTransform RotationTransform = Concatenate(Inverse(RotationPoint), FQuat2f(Angle2D), RotationPoint);
+		DrawElement.SetRenderTransform(Concatenate(RotationTransform, DrawElement.GetRenderTransform()));
+	}
+}
+
+#if !UE_BUILD_SHIPPING
+namespace UE::Slate::Private
+{
+	bool GLogPaintedText = false;
+	FAutoConsoleVariableRef LogPaintedTextRef(
+		TEXT("Slate.LogPaintedText"),
+		GLogPaintedText,
+		TEXT("If true, all text that is visible to the user will be logged when it is painted. This will log the full text to be painted, not the truncated or clipped version based on UI constraints.")
+	);
+}
+#endif
+
+void FSlateDrawElement::MakeText( FSlateWindowElementList& ElementList, uint32 InLayer, const FPaintGeometry& PaintGeometry, const FString& InText, const int32 StartIndex, const int32 EndIndex, const FSlateFontInfo& InFontInfo, ESlateDrawEffect InDrawEffects, const FLinearColor& InTint )
+{
+	SCOPE_CYCLE_COUNTER(STAT_SlateDrawElementMakeTime)
+	PaintGeometry.CommitTransformsIfUsingLegacyConstructor();
+
+	if (ShouldCull(ElementList, PaintGeometry, InTint, InText))
+	{
+		return;
+	}
+#if !UE_BUILD_SHIPPING
+	if (UE::Slate::Private::GLogPaintedText)
+	{
+		UE_LOG(LogSlate, Log, TEXT("MakeText: '%s'."), *InText);
+	}
+#endif
+	FSlateTextElement& Element = ElementList.AddUninitialized<EElementType::ET_Text>();
+
+	Element.SetTint(InTint);
+	Element.SetText(InText, InFontInfo, StartIndex, EndIndex);
+
+	Element.Init(ElementList, EElementType::ET_Text, InLayer, PaintGeometry, InDrawEffects);
+}
+
+void FSlateDrawElement::MakeText( FSlateWindowElementList& ElementList, uint32 InLayer, const FPaintGeometry& PaintGeometry, const FString& InText, const FSlateFontInfo& InFontInfo, ESlateDrawEffect InDrawEffects, const FLinearColor& InTint )
+{
+	SCOPE_CYCLE_COUNTER(STAT_SlateDrawElementMakeTime)
+	PaintGeometry.CommitTransformsIfUsingLegacyConstructor();
+
+	// Don't try and render empty text
+	if (InText.Len() == 0)
+	{
+		return;
+	}
+
+	if (ShouldCull(ElementList, PaintGeometry, InTint, InText))
+	{
+		return;
+	}
+
+	// Don't do anything if there the font would be completely transparent 
+	if (InTint.A == 0 && !InFontInfo.OutlineSettings.IsVisible())
+	{
+		return;
+	}
+#if !UE_BUILD_SHIPPING
+	if (UE::Slate::Private::GLogPaintedText)
+	{
+		UE_LOG(LogSlate, Log, TEXT("MakeText: '%s'."), *InText);
+	}
+#endif
+	FSlateTextElement& Element = ElementList.AddUninitialized<EElementType::ET_Text>();
+
+	Element.SetTint(InTint);
+	Element.SetText(InText, InFontInfo);
+
+	Element.Init(ElementList, EElementType::ET_Text, InLayer, PaintGeometry, InDrawEffects);
+}
+
+namespace SlateDrawElement
+{
+#if SLATE_CHECK_UOBJECT_RENDER_RESOURCES
+	const FName MaterialInterfaceClassName = "MaterialInterface";
+	void CheckInvalidUMaterial(const UObject* InMaterialResource, const TCHAR* MaterialMessage)
+	{
+		if (InMaterialResource && GSlateCheckUObjectRenderResources)
+		{
+			bool bIsValidLowLevel = InMaterialResource->IsValidLowLevelFast(false);
+			if (!bIsValidLowLevel || !IsValid(InMaterialResource) || InMaterialResource->GetClass()->GetFName() == MaterialInterfaceClassName)
+			{
+				UE_LOG(LogSlate, Error, TEXT("Material '%s' is not valid. PendingKill:'%d'. ValidLowLevelFast:'%d'. InvalidClass:'%d'")
+					, MaterialMessage
+					, (bIsValidLowLevel ? !IsValid(InMaterialResource) : false)
+					, bIsValidLowLevel
+					, (bIsValidLowLevel ? InMaterialResource->GetClass()->GetFName() == MaterialInterfaceClassName : false));
+
+				if (bIsValidLowLevel)
+				{
+					UE_LOG(LogSlate, Error, TEXT("Material name: '%s'"), *InMaterialResource->GetFullName());
+				}
+
+				const TCHAR* Message = TEXT("We detected an invalid resource in FSlateDrawElement. Check the log for more detail.");
+				if (GSlateCheckUObjectRenderResourcesShouldLogFatal)
+				{
+					UE_LOG(LogSlate, Fatal, TEXT("%s"), Message);
+				}
+				else
+				{
+					ensureAlwaysMsgf(false, TEXT("%s"), Message);
+				}
+			}
+		}
+	}
+#endif
+}
+
+void FSlateDrawElement::MakeShapedText(FSlateWindowElementList& ElementList, uint32 InLayer, const FPaintGeometry& PaintGeometry, const FShapedGlyphSequenceRef& InShapedGlyphSequence, ESlateDrawEffect InDrawEffects, const FLinearColor& BaseTint, const FLinearColor& OutlineTint, FTextOverflowArgs TextOverflowArgs)
+{
+	SCOPE_CYCLE_COUNTER(STAT_SlateDrawElementMakeTime)
+	PaintGeometry.CommitTransformsIfUsingLegacyConstructor();
+
+	if (InShapedGlyphSequence->GetGlyphsToRender().Num() == 0)
+	{
+		return;
+	}
+
+	if (ShouldCull(ElementList, PaintGeometry))
+	{
+		return;
+	}
+
+	// Don't do anything if there the font would be completely transparent 
+	if ((BaseTint.A == 0 && InShapedGlyphSequence->GetFontOutlineSettings().OutlineSize == 0) || 
+		(BaseTint.A == 0 && OutlineTint.A == 0))
+	{
+		return;
+	}
+
+#if SLATE_CHECK_UOBJECT_RENDER_RESOURCES
+	SlateDrawElement::CheckInvalidUMaterial(InShapedGlyphSequence->GetFontMaterial(), TEXT("Font Material"));
+	SlateDrawElement::CheckInvalidUMaterial(InShapedGlyphSequence->GetFontOutlineSettings().OutlineMaterial, TEXT("Outline Material"));
+#endif
+
+	FSlateShapedTextElement& Element = ElementList.AddUninitialized<EElementType::ET_ShapedText>();
+
+	Element.SetTint(BaseTint);
+	Element.SetShapedText(ElementList, InShapedGlyphSequence, OutlineTint);
+	Element.SetOverflowArgs(TextOverflowArgs);
+
+	Element.Init(ElementList, EElementType::ET_ShapedText, InLayer, PaintGeometry, InDrawEffects);
+}
+
+void FSlateDrawElement::MakeGradient( FSlateWindowElementList& ElementList, uint32 InLayer, const FPaintGeometry& PaintGeometry, TArray<FSlateGradientStop> InGradientStops, EOrientation InGradientType, ESlateDrawEffect InDrawEffects, FVector4f CornerRadius)
+{
+	PaintGeometry.CommitTransformsIfUsingLegacyConstructor();
+
+	if (ShouldCull(ElementList, PaintGeometry))
+	{
+		return;
+	}
+
+	FSlateGradientElement& Element = ElementList.AddUninitialized<EElementType::ET_Gradient>();
+
+	Element.SetGradient(InGradientStops, InGradientType, CornerRadius);
+
+	Element.Init(ElementList, EElementType::ET_Gradient, InLayer, PaintGeometry, InDrawEffects);
+}
+
+void FSlateDrawElement::MakeSpline(FSlateWindowElementList& ElementList, uint32 InLayer, const FPaintGeometry& PaintGeometry, const UE::Slate::FDeprecateVector2DParameter InStart, const UE::Slate::FDeprecateVector2DParameter InStartDir, const UE::Slate::FDeprecateVector2DParameter InEnd, const UE::Slate::FDeprecateVector2DParameter InEndDir, float InThickness, ESlateDrawEffect InDrawEffects, const FLinearColor& InTint)
+{
+	PaintGeometry.CommitTransformsIfUsingLegacyConstructor();
+
+	if (ShouldCull(ElementList))
+	{
+		return;
+	}
+	FSlateSplineElement& Element = ElementList.AddUninitialized<EElementType::ET_Spline>();
+
+	Element.SetHermiteSpline(InStart, InStartDir, InEnd, InEndDir, InThickness, InTint);
+
+	Element.Init(ElementList, EElementType::ET_Spline, InLayer, PaintGeometry, InDrawEffects);
+}
+
+void FSlateDrawElement::MakeCubicBezierSpline(FSlateWindowElementList& ElementList, uint32 InLayer, const FPaintGeometry& PaintGeometry, const UE::Slate::FDeprecateVector2DParameter P0, const UE::Slate::FDeprecateVector2DParameter P1, const UE::Slate::FDeprecateVector2DParameter P2, const UE::Slate::FDeprecateVector2DParameter P3, float InThickness, ESlateDrawEffect InDrawEffects, const FLinearColor& InTint)
+{
+	PaintGeometry.CommitTransformsIfUsingLegacyConstructor();
+
+	if (ShouldCull(ElementList))
+	{
+		return;
+	}
+	FSlateSplineElement& Element = ElementList.AddUninitialized<EElementType::ET_Spline>();
+
+	Element.SetCubicBezier(P0, P1, P2, P3, InThickness, InTint);
+
+	Element.Init(ElementList, EElementType::ET_Spline, InLayer, PaintGeometry, InDrawEffects);
+}
+
+void FSlateDrawElement::MakeDrawSpaceSpline(FSlateWindowElementList& ElementList, uint32 InLayer, const UE::Slate::FDeprecateVector2DParameter InStart, const UE::Slate::FDeprecateVector2DParameter InStartDir, const UE::Slate::FDeprecateVector2DParameter InEnd, const UE::Slate::FDeprecateVector2DParameter InEndDir, float InThickness, ESlateDrawEffect InDrawEffects, const FLinearColor& InTint)
+{
+	MakeSpline(ElementList, InLayer, FPaintGeometry(), InStart, InStartDir, InEnd, InEndDir, InThickness, InDrawEffects, InTint);
+}
+
+#if UE_ENABLE_SLATE_VECTOR_DEPRECATION_MECHANISMS
+void FSlateDrawElement::MakeLines(FSlateWindowElementList& ElementList, uint32 InLayer, const FPaintGeometry& PaintGeometry, const TArray<FVector2d>& Points, ESlateDrawEffect InDrawEffects, const FLinearColor& InTint, bool bAntialias, float Thickness)
+{
+	if (ShouldCull(ElementList) || Points.Num() < 2)
+	{
+		return;
+	}
+
+	TArray<FVector2f> NewVector;
+	NewVector.Reserve(Points.Num());
+	for (FVector2d Vect : Points)
+	{
+		NewVector.Add(UE::Slate::CastToVector2f(Vect));
+	}
+	MakeLines(ElementList, InLayer, PaintGeometry, MoveTemp(NewVector), InDrawEffects, InTint, bAntialias, Thickness);
+}
+#endif // UE_ENABLE_SLATE_VECTOR_DEPRECATION_MECHANISMS
+
+void FSlateDrawElement::MakeLines(FSlateWindowElementList& ElementList, uint32 InLayer, const FPaintGeometry& PaintGeometry, TArray<FVector2f> Points, ESlateDrawEffect InDrawEffects, const FLinearColor& InTint, bool bAntialias, float Thickness)
+{
+	PaintGeometry.CommitTransformsIfUsingLegacyConstructor();
+
+	if (ShouldCull(ElementList) || Points.Num() < 2)
+	{
+		return;
+	}
+
+	FSlateLineElement& Element = ElementList.AddUninitialized<EElementType::ET_Line>();
+
+	Element.SetTint(InTint);
+	Element.SetThickness(Thickness);
+	Element.SetLines(MoveTemp(Points), bAntialias);
+
+	ESlateDrawEffect DrawEffects = InDrawEffects;
+	if (bAntialias)
+	{
+		// If the line is to be anti-aliased, we cannot reliably snap
+		// the generated vertices
+		DrawEffects |= ESlateDrawEffect::NoPixelSnapping;
+	}
+
+	Element.Init(ElementList, EElementType::ET_Line, InLayer, PaintGeometry, DrawEffects);
+}
+
+#if UE_ENABLE_SLATE_VECTOR_DEPRECATION_MECHANISMS
+void FSlateDrawElement::MakeLines(FSlateWindowElementList& ElementList, uint32 InLayer, const FPaintGeometry& PaintGeometry, const TArray<FVector2d>& Points, const TArray<FLinearColor>& PointColors, ESlateDrawEffect InDrawEffects, const FLinearColor& InTint, bool bAntialias, float Thickness)
+{
+	if (ShouldCull(ElementList) || Points.Num() < 2)
+	{
+		return;
+	}
+
+	TArray<FVector2f> NewVector;
+	NewVector.Reserve(Points.Num());
+	for (FVector2d Vect : Points)
+	{
+		NewVector.Add(UE::Slate::CastToVector2f(Vect));
+	}
+	MakeLines(ElementList, InLayer, PaintGeometry, MoveTemp(NewVector), PointColors, InDrawEffects, InTint, bAntialias, Thickness);
+}
+#endif
+
+void FSlateDrawElement::MakeLines( FSlateWindowElementList& ElementList, uint32 InLayer, const FPaintGeometry& PaintGeometry, TArray<FVector2f> Points, TArray<FLinearColor> PointColors, ESlateDrawEffect InDrawEffects, const FLinearColor& InTint, bool bAntialias, float Thickness )
+{
+	PaintGeometry.CommitTransformsIfUsingLegacyConstructor();
+
+	if (ShouldCull(ElementList) || Points.Num() < 2)
+	{
+		return;
+	}
+
+	FSlateLineElement& Element = ElementList.AddUninitialized<EElementType::ET_Line>();
+
+	Element.SetTint(InTint);
+	Element.SetThickness(Thickness);
+	Element.SetLines(MoveTemp(Points), bAntialias, MoveTemp(PointColors));
+
+	Element.Init(ElementList, EElementType::ET_Line, InLayer, PaintGeometry, InDrawEffects);
+}
+
+void FSlateDrawElement::MakeViewport( FSlateWindowElementList& ElementList, uint32 InLayer, const FPaintGeometry& PaintGeometry, TSharedPtr<const ISlateViewport> Viewport, ESlateDrawEffect InDrawEffects, const FLinearColor& InTint )
+{
+	PaintGeometry.CommitTransformsIfUsingLegacyConstructor();
+
+	if (ShouldCull(ElementList))
+	{
+		return;
+	}
+
+	FSlateViewportElement& Element = ElementList.AddUninitialized<EElementType::ET_Viewport>();
+
+	Element.SetViewport(Viewport, InTint);
+	check(Element.RenderTargetResource == nullptr || !Element.RenderTargetResource->Debug_IsDestroyed());
+
+	Element.Init(ElementList, EElementType::ET_Viewport, InLayer, PaintGeometry, InDrawEffects);
+}
+
+
+void FSlateDrawElement::MakeCustom( FSlateWindowElementList& ElementList, uint32 InLayer, TSharedPtr<ICustomSlateElement, ESPMode::ThreadSafe> CustomDrawer )
+{
+	if (ShouldCull(ElementList))
+	{
+		return;
+	}
+
+	FSlateCustomDrawerElement& Element = ElementList.AddUninitialized<EElementType::ET_Custom>();
+
+	Element.SetCustomDrawer(CustomDrawer);
+	
+	Element.Init(ElementList, EElementType::ET_Custom, InLayer, FPaintGeometry(), ESlateDrawEffect::None);
+	Element.RenderTransform = FSlateRenderTransform();
+}
+
+
+void FSlateDrawElement::MakeCustomVerts(FSlateWindowElementList& ElementList, uint32 InLayer, const FSlateResourceHandle& InRenderResourceHandle, const TArray<FSlateVertex>& InVerts, const TArray<SlateIndex>& InIndexes, ISlateUpdatableInstanceBuffer* InInstanceData, uint32 InInstanceOffset, uint32 InNumInstances, ESlateDrawEffect InDrawEffects)
+{
+	SCOPE_CYCLE_COUNTER(STAT_SlateDrawElementMakeCustomVertsTime);
+	
+	if (ShouldCull(ElementList))
+	{
+		return;
+	}
+
+	FSlateCustomVertsElement& Element = ElementList.AddUninitialized<EElementType::ET_CustomVerts>();
+
+	const FSlateShaderResourceProxy* RenderingProxy = InRenderResourceHandle.GetResourceProxy();
+	ISlateUpdatableInstanceBufferRenderProxy* RenderProxy = InInstanceData ? InInstanceData->GetRenderProxy() : nullptr;
+	Element.SetCustomVerts(RenderingProxy, InVerts, InIndexes, RenderProxy, InInstanceOffset, InNumInstances);
+
+	Element.Init(ElementList, EElementType::ET_CustomVerts, InLayer, FPaintGeometry(), InDrawEffects);
+	Element.RenderTransform = FSlateRenderTransform();
+}
+
+void FSlateDrawElement::MakePostProcessPass(FSlateWindowElementList& ElementList, uint32 InLayer, const FPaintGeometry& PaintGeometry, const FVector4f& Params, int32 DownsampleAmount, FVector4f CornerRadius)
+{
+	PaintGeometry.CommitTransformsIfUsingLegacyConstructor();
+
+	if (ShouldCull(ElementList))
+	{
+		return;
+	}
+
+	FSlatePostProcessElement& Element = ElementList.AddUninitialized<EElementType::ET_PostProcessPass>();
+
+	Element.DownsampleAmount = DownsampleAmount;
+	Element.PostProcessData = Params;
+	Element.CornerRadius = CornerRadius;
+
+	Element.Init(ElementList, EElementType::ET_PostProcessPass, InLayer, PaintGeometry, ESlateDrawEffect::None);
+}
+
+FSlateDrawElement::FSlateDrawElement()
+	: bIsCached(false)
+{
+
+}
+
+FVector2f FSlateDrawElement::GetRotationPoint(const FPaintGeometry& PaintGeometry, const TOptional<FVector2f>& UserRotationPoint, ERotationSpace RotationSpace)
+{
+	FVector2f RotationPoint(0, 0);
+
+	const FVector2f LocalSize = UE::Slate::CastToVector2f(PaintGeometry.GetLocalSize());
+
+	switch (RotationSpace)
+	{
+		case RelativeToElement:
+		{
+			// If the user did not specify a rotation point, we rotate about the center of the element
+			RotationPoint = UserRotationPoint.Get(LocalSize * 0.5f);
+		}
+		break;
+		case RelativeToWorld:
+		{
+			// its in world space, must convert the point to local space.
+			RotationPoint = TransformPoint(Inverse(PaintGeometry.GetAccumulatedRenderTransform()), UserRotationPoint.Get(FVector2f::ZeroVector));
+		}
+		break;
+	default:
+		check(0);
+		break;
+	}
+
+	return RotationPoint;
+}
\ No newline at end of file
diff --git a/Engine/Source/Runtime/SlateCore/Private/Rendering/DrawElements.cpp b/Engine/Source/Runtime/SlateCore/Private/Rendering/DrawElements.cpp
--- a/Engine/Source/Runtime/SlateCore/Private/Rendering/DrawElements.cpp
+++ b/Engine/Source/Runtime/SlateCore/Private/Rendering/DrawElements.cpp
@@ -13,144 +13,33 @@
 #include "Application/SlateApplicationBase.h"
 
 #include <limits>
+#include <type_traits>
+#include <utility>
 
-DECLARE_CYCLE_STAT(TEXT("FSlateDrawElement::Make Time"), STAT_SlateDrawElementMakeTime, STATGROUP_SlateVerbose);
-DECLARE_CYCLE_STAT(TEXT("FSlateDrawElement::MakeCustomVerts Time"), STAT_SlateDrawElementMakeCustomVertsTime, STATGROUP_Slate);
-DECLARE_CYCLE_STAT(TEXT("FSlateDrawElement::Prebatch Time"), STAT_SlateDrawElementPrebatchTime, STATGROUP_Slate);
-
-DEFINE_STAT(STAT_SlateBufferPoolMemory);
-DEFINE_STAT(STAT_SlateCachedDrawElementMemory);
-
-static bool bApplyDisabledEffectOnWidgets = true;
-
-static void HandleApplyDisabledEffectOnWidgetsToggled(IConsoleVariable* CVar)
-{
-	if (FSlateApplicationBase::IsInitialized())
-	{
-		FSlateApplicationBase::Get().InvalidateAllWidgets(false);
-	}
-}
-
-static FAutoConsoleVariableRef CVarApplyDisabledEffectOnWidgets(
-	TEXT("Slate.ApplyDisabledEffectOnWidgets"),
-	bApplyDisabledEffectOnWidgets,
-	TEXT("If true, disabled game-layer widgets will have alpha multiplied by 0.45."),
-	FConsoleVariableDelegate::CreateStatic(&HandleApplyDisabledEffectOnWidgetsToggled)
-);
-
-static bool IsResourceObjectValid(UObject*& InObject)
-{
-	if (InObject != nullptr && (!IsValidChecked(InObject) || InObject->IsUnreachable() || InObject->HasAnyFlags(RF_BeginDestroyed)))
-	{
-		UE_LOG(LogSlate, Warning, TEXT("Attempted to access resource for %s which is pending kill, unreachable or pending destroy"), *InObject->GetName());
-		return false;
-	}
-
-	return true;
-}
-
-static bool ShouldCull(const FSlateWindowElementList& ElementList)
-{
-	const FSlateClippingManager& ClippingManager = ElementList.GetClippingManager();
-	const int32 CurrentIndex = ClippingManager.GetClippingIndex();
-	if (CurrentIndex != INDEX_NONE)
-	{
-		const FSlateClippingState& ClippingState = ClippingManager.GetClippingStates()[CurrentIndex];
-		if (ClippingState.GetClippingMethod() == EClippingMethod::Scissor)
-		{
-			return ClippingState.HasZeroArea();
-		}
-		else if (ClippingState.GetClippingMethod() == EClippingMethod::Stencil)
-		{
-			FSlateRect WindowRect = ElementList.GetPaintWindow()->GetClippingRectangleInWindow();
-			if (WindowRect.GetArea() > 0)
-			{
-				for (const FSlateClippingZone& Stencil : ClippingState.StencilQuads)
-				{
-					bool bOverlapping = false;
-					FSlateRect ClippedStencil = Stencil.GetBoundingBox().IntersectionWith(WindowRect, bOverlapping);
-					
-					if (!bOverlapping || ClippedStencil.GetArea() <= 0)
-					{
-						return true;
-					}
-				}
-			}
-		}
-
-	}
-
-	return false;
-}
-
-static bool ShouldCull(const FSlateWindowElementList& ElementList, const FPaintGeometry& PaintGeometry)
-{
-	const FVector2f LocalSize = UE::Slate::CastToVector2f(PaintGeometry.GetLocalSize());
-	const float DrawScale = PaintGeometry.DrawScale;
-
-	const FVector2f PixelSize = (LocalSize * DrawScale);
-	if (PixelSize.X <= 0.f || PixelSize.Y <= 0.f)
-	{
-		return true;
-	}
-
-	return ShouldCull(ElementList);
-}
-
-static bool ShouldCull(const FSlateWindowElementList& ElementList, const FPaintGeometry& PaintGeometry, const FSlateBrush* InBrush)
+/** 
+ * Helper struct to determine if an element has implemented AddReferencedObject 
+ * 
+ * For an explantion of how this detection: https://en.cppreference.com/w/cpp/types/void_t
+ */
+template <typename T>
+struct TIsReferencedObjects
 {
-	if (ShouldCull(ElementList, PaintGeometry))
-	{
-		return true;
-	}
-
-	if (InBrush->GetDrawType() == ESlateBrushDrawType::NoDrawType)
+	template <typename U, typename = void>
+	struct IsReferencedObjectsHelper : std::false_type
 	{
-		return true;
-	}
-
-	UObject* ResourceObject = InBrush->GetResourceObject();
-	if (!IsResourceObjectValid(ResourceObject))
-	{
-		return true;
-	}
-
-	return false;
-}
-
-
-static bool ShouldCull(const FSlateWindowElementList& ElementList, const FPaintGeometry& PaintGeometry, const FLinearColor& InTint)
-{
-	if (InTint.A == 0 || ShouldCull(ElementList, PaintGeometry))
-	{
-		return true;
-	}
-
-	return false;
-}
-
-static bool ShouldCull(const FSlateWindowElementList& ElementList, const FPaintGeometry& PaintGeometry, const FLinearColor& InTint, const FString& InText)
-{
-	if (InTint.A == 0 || InText.Len() == 0 || ShouldCull(ElementList, PaintGeometry))
-	{
-		return true;
-	}
-
-	return false;
-}
+	};
 
-static bool ShouldCull(const FSlateWindowElementList& ElementList, const FPaintGeometry& PaintGeometry, const FSlateBrush* InBrush, const FLinearColor& InTint)
-{
-	bool bIsFullyTransparent = InTint.A == 0 && (InBrush->OutlineSettings.Color.GetSpecifiedColor().A == 0 || InBrush->OutlineSettings.bUseBrushTransparency);
-	if (bIsFullyTransparent || ShouldCull(ElementList, PaintGeometry, InBrush))
+	template <typename U>
+	struct IsReferencedObjectsHelper<U, std::void_t<decltype(std::declval<U>().AddReferencedObjects(std::declval<FReferenceCollector&>()))>> : std::true_type
 	{
-		return true;
-	}
-
-	return false;
-}
+	};
 
+	static constexpr bool Value = IsReferencedObjectsHelper<T>::value;
+};
 
+static_assert(!TIsReferencedObjects<FSlateBoxElement>::Value, "FSlateBoxElement should have not AddReferencedObjects.");
+static_assert(TIsReferencedObjects<FSlateTextElement>::Value, "FSlateTextElement should have AddReferencedObjects.");
+static_assert(TIsReferencedObjects<FSlateShapedTextElement>::Value, "FSlateShapedTextElement should have AddReferencedObjects.");
 
 FSlateWindowElementList::FSlateWindowElementList(const TSharedPtr<SWindow>& InPaintWindow)
 	: WeakPaintWindow(InPaintWindow)
@@ -187,615 +76,6 @@ FSlateWindowElementList::~FSlateWindowElementList()
 	}*/
 }
 
-void FSlateDrawElement::Init(FSlateWindowElementList& ElementList, EElementType InElementType, uint32 InLayer, const FPaintGeometry& PaintGeometry, ESlateDrawEffect InDrawEffects)
-{
-	if (ElementList.GetIsInGameLayer() && !bApplyDisabledEffectOnWidgets)
-	{
-		InDrawEffects &= ~ESlateDrawEffect::DisabledEffect;
-	}
-
-	// Apply the pixel snapping effects, if it's unset/inherited we don't do anything.
-	switch (ElementList.GetPixelSnappingMethod())
-	{
-		case EWidgetPixelSnapping::Disabled:
-			InDrawEffects |= ESlateDrawEffect::NoPixelSnapping;
-			break;
-		case EWidgetPixelSnapping::Inherit:
-			// NOTE: Do nothing on inherit.
-		case EWidgetPixelSnapping::SnapToPixel:
-			// NOTE: We don't remove the drawing effects pixel snapping here.  The default behavior of InDrawEffects is always to snap to pixels
-			// so all we need to do here is not to snap.  This will allow the continued behavior that if someone does set snap to pixel at a
-			// higher level, it won't break people that directly set InDrawEffects |= ESlateDrawEffect::NoPixelSnapping, before drawing their
-			// draw element.
-		default:
-			break;
-	}
-
-	RenderTransform = PaintGeometry.GetAccumulatedRenderTransform();
-	Position = PaintGeometry.DrawPosition;
-	Scale = PaintGeometry.DrawScale;
-	LocalSize = FVector2f(PaintGeometry.GetLocalSize());
-	ClipStateHandle.SetPreCachedClipIndex(ElementList.GetClippingIndex());
-
-#if UE_SLATE_VERIFY_PIXELSIZE
-	const FVector2f PixelSize = FVector2f(LocalSize * Scale);
-	ensureMsgf(PixelSize.X >= 0.f && PixelSize.X <= (float)std::numeric_limits<uint16>::max(), TEXT("The size X '%f' is too small or big to fit in the SlateVertex buffer."), PixelSize.X);
-	ensureMsgf(PixelSize.Y >= 0.f && PixelSize.Y <= (float)std::numeric_limits<uint16>::max(), TEXT("The size Y '%f' is too small or big to fit in the SlateVertex buffer."), PixelSize.Y);
-#endif
-
-	LayerId = InLayer;
-
-	ElementType = InElementType;
-	DrawEffects = InDrawEffects;
-	
-	// Calculate the layout to render transform as this is needed by several calculations downstream.
-	const FSlateLayoutTransform InverseLayoutTransform(Inverse(FSlateLayoutTransform(Scale, Position)));
-
-	{
-		// This is a workaround because we want to keep track of the various Scenes 
-		// in use throughout the UI. We keep a synchronized set with the render thread on the SlateRenderer and 
-		// use indices to synchronize between them.
-		FSlateRenderer* Renderer = FSlateApplicationBase::Get().GetRenderer();
-		checkSlow(Renderer);
-		int32 SceneIndexLong = Renderer->GetCurrentSceneIndex();
-		ensureMsgf(SceneIndexLong <= std::numeric_limits<int8>::max(), TEXT("The Scene index is saved as an int8 in the DrawElements."));
-		SceneIndex = (int8)SceneIndexLong;
-	}
-
-	BatchFlags = ESlateBatchDrawFlag::None;
-	BatchFlags |= static_cast<ESlateBatchDrawFlag>(static_cast<uint32>(InDrawEffects) & static_cast<uint32>(ESlateDrawEffect::NoBlending | ESlateDrawEffect::PreMultipliedAlpha | ESlateDrawEffect::NoGamma | ESlateDrawEffect::InvertAlpha));
-
-	static_assert(((__underlying_type(ESlateDrawEffect))ESlateDrawEffect::NoBlending) == ((__underlying_type(ESlateBatchDrawFlag))ESlateBatchDrawFlag::NoBlending), "Must keep ESlateBatchDrawFlag and ESlateDrawEffect partial matches");
-	static_assert(((__underlying_type(ESlateDrawEffect))ESlateDrawEffect::PreMultipliedAlpha) == ((__underlying_type(ESlateBatchDrawFlag))ESlateBatchDrawFlag::PreMultipliedAlpha), "Must keep ESlateBatchDrawFlag and ESlateDrawEffect partial matches");
-	static_assert(((__underlying_type(ESlateDrawEffect))ESlateDrawEffect::NoGamma) == ((__underlying_type(ESlateBatchDrawFlag))ESlateBatchDrawFlag::NoGamma), "Must keep ESlateBatchDrawFlag and ESlateDrawEffect partial matches");
-	static_assert(((__underlying_type(ESlateDrawEffect))ESlateDrawEffect::InvertAlpha) == ((__underlying_type(ESlateBatchDrawFlag))ESlateBatchDrawFlag::InvertAlpha), "Must keep ESlateBatchDrawFlag and ESlateDrawEffect partial matches");
-	if ((InDrawEffects & ESlateDrawEffect::ReverseGamma) != ESlateDrawEffect::None)
-	{
-		BatchFlags |= ESlateBatchDrawFlag::ReverseGamma;
-	}
-}
-
-void FSlateDrawElement::ApplyPositionOffset(UE::Slate::FDeprecateVector2DParameter InOffset)
-{
-	SetPosition(Position + InOffset);
-	RenderTransform = Concatenate(RenderTransform, InOffset);
-
-	// Recompute cached layout to render transform
-	const FSlateLayoutTransform InverseLayoutTransform(Inverse(FSlateLayoutTransform(Scale, Position)));
-}
-
-void FSlateDrawElement::MakeDebugQuad( FSlateWindowElementList& ElementList, uint32 InLayer, const FPaintGeometry& PaintGeometry, FLinearColor Tint)
-{
-	PaintGeometry.CommitTransformsIfUsingLegacyConstructor();
-
-	if (ShouldCull(ElementList))
-	{
-		return;
-	}
-
-	FSlateDrawElement& Element = ElementList.AddUninitialized<EElementType::ET_DebugQuad>();
-	FSlateBoxPayload& BoxPayload = ElementList.CreatePayload<FSlateBoxPayload>(Element);
-
-	BoxPayload.SetTint(Tint);
-	Element.Init(ElementList, EElementType::ET_DebugQuad, InLayer, PaintGeometry, ESlateDrawEffect::None);
-}
-
-FSlateDrawElement& FSlateDrawElement::MakeBoxInternal(
-	FSlateWindowElementList& ElementList,
-	uint32 InLayer,
-	const FPaintGeometry& PaintGeometry,
-	const FSlateBrush* InBrush,
-	ESlateDrawEffect InDrawEffects,
-	const FLinearColor& InTint
-)
-{
-	EElementType ElementType = (InBrush->DrawAs == ESlateBrushDrawType::Border) ? EElementType::ET_Border : (InBrush->DrawAs == ESlateBrushDrawType::RoundedBox) ? EElementType::ET_RoundedBox : EElementType::ET_Box;
-
-	// Cast to Rounded Rect to get the internal parameters 
-	// New payload type - inherit from BoxPayload 
-
-	FSlateDrawElement& Element = ElementList.AddUninitializedLookup(ElementType);
-
-	FSlateBoxPayload* BoxPayload;
-	if ( ElementType == EElementType::ET_RoundedBox )
-	{
-		FSlateRoundedBoxPayload* RBoxPayload = &ElementList.CreatePayload<FSlateRoundedBoxPayload>(Element);
-		FVector4f CornerRadii = FVector4f(InBrush->OutlineSettings.CornerRadii);
-		if (InBrush->OutlineSettings.RoundingType == ESlateBrushRoundingType::HalfHeightRadius)
-		{
-			const float UniformRadius = PaintGeometry.GetLocalSize().Y / 2.0f;
-			CornerRadii = FVector4f(UniformRadius, UniformRadius, UniformRadius, UniformRadius);
-		}
-		RBoxPayload->SetRadius(CornerRadii);
-
-		if (InBrush->OutlineSettings.bUseBrushTransparency)
-		{
-			FLinearColor Color = InBrush->OutlineSettings.Color.GetSpecifiedColor().CopyWithNewOpacity(InTint.A);
-			RBoxPayload->SetOutline(Color, InBrush->OutlineSettings.Width);
-		}
-		else
-		{
-			RBoxPayload->SetOutline(InBrush->OutlineSettings.Color.GetSpecifiedColor(), InBrush->OutlineSettings.Width);
-		}
-		BoxPayload = RBoxPayload;
-	}
-	else 
-	{
-		BoxPayload = &ElementList.CreatePayload<FSlateBoxPayload>(Element);
-	}
-
-	BoxPayload->SetTint(InTint);
-	BoxPayload->SetBrush(InBrush, PaintGeometry.GetLocalSize(), PaintGeometry.DrawScale);
-
-	Element.Init(ElementList, ElementType, InLayer, PaintGeometry, InDrawEffects);
-
-	return Element;
-}
-
-void FSlateDrawElement::MakeBox(
-	FSlateWindowElementList& ElementList,
-	uint32 InLayer, 
-	const FPaintGeometry& PaintGeometry, 
-	const FSlateBrush* InBrush,
-	ESlateDrawEffect InDrawEffects, 
-	const FLinearColor& InTint)
-{
-	PaintGeometry.CommitTransformsIfUsingLegacyConstructor();
-
-	if (ShouldCull(ElementList, PaintGeometry, InBrush, InTint))
-	{
-		return;
-	}
-
-	MakeBoxInternal(ElementList, InLayer, PaintGeometry, InBrush, InDrawEffects, InTint);
-}
-
-void FSlateDrawElement::MakeRotatedBox(
-	FSlateWindowElementList& ElementList,
-	uint32 InLayer,
-	const FPaintGeometry& PaintGeometry,
-	const FSlateBrush* InBrush,
-	ESlateDrawEffect InDrawEffects,
-	float Angle2D,
-	UE::Slate::FDeprecateOptionalVector2DParameter InRotationPoint,
-	ERotationSpace RotationSpace,
-	const FLinearColor& InTint)
-{
-	PaintGeometry.CommitTransformsIfUsingLegacyConstructor();
-
-	if (ShouldCull(ElementList, PaintGeometry, InBrush, InTint))
-	{
-		return;
-	}
-
-	FSlateDrawElement& DrawElement = MakeBoxInternal(ElementList, InLayer, PaintGeometry, InBrush, InDrawEffects, InTint);
-	
-	if (Angle2D != 0.0f)
-	{
-		const FVector2f RotationPoint = GetRotationPoint(PaintGeometry, InRotationPoint, RotationSpace);
-		const FSlateRenderTransform RotationTransform = Concatenate(Inverse(RotationPoint), FQuat2f(Angle2D), RotationPoint);
-		DrawElement.SetRenderTransform(Concatenate(RotationTransform, DrawElement.GetRenderTransform()));
-	}
-}
-
-#if !UE_BUILD_SHIPPING
-namespace UE::Slate::Private
-{
-	bool GLogPaintedText = false;
-	FAutoConsoleVariableRef LogPaintedTextRef(
-		TEXT("Slate.LogPaintedText"),
-		GLogPaintedText,
-		TEXT("If true, all text that is visible to the user will be logged when it is painted. This will log the full text to be painted, not the truncated or clipped version based on UI constraints.")
-	);
-}
-#endif
-
-void FSlateDrawElement::MakeText( FSlateWindowElementList& ElementList, uint32 InLayer, const FPaintGeometry& PaintGeometry, const FString& InText, const int32 StartIndex, const int32 EndIndex, const FSlateFontInfo& InFontInfo, ESlateDrawEffect InDrawEffects, const FLinearColor& InTint )
-{
-	SCOPE_CYCLE_COUNTER(STAT_SlateDrawElementMakeTime)
-	PaintGeometry.CommitTransformsIfUsingLegacyConstructor();
-
-	if (ShouldCull(ElementList, PaintGeometry, InTint, InText))
-	{
-		return;
-	}
-#if !UE_BUILD_SHIPPING
-	if (UE::Slate::Private::GLogPaintedText)
-	{
-		UE_LOG(LogSlate, Log, TEXT("MakeText: '%s'."), *InText);
-	}
-#endif
-	FSlateDrawElement& Element = ElementList.AddUninitialized<EElementType::ET_Text>();
-	FSlateTextPayload& DataPayload = ElementList.CreatePayload<FSlateTextPayload>(Element);
-
-	DataPayload.SetTint(InTint);
-	DataPayload.SetText(InText, InFontInfo, StartIndex, EndIndex);
-
-	Element.Init(ElementList, EElementType::ET_Text, InLayer, PaintGeometry, InDrawEffects);
-}
-
-void FSlateDrawElement::MakeText( FSlateWindowElementList& ElementList, uint32 InLayer, const FPaintGeometry& PaintGeometry, const FString& InText, const FSlateFontInfo& InFontInfo, ESlateDrawEffect InDrawEffects, const FLinearColor& InTint )
-{
-	SCOPE_CYCLE_COUNTER(STAT_SlateDrawElementMakeTime)
-	PaintGeometry.CommitTransformsIfUsingLegacyConstructor();
-
-	// Don't try and render empty text
-	if (InText.Len() == 0)
-	{
-		return;
-	}
-
-	if (ShouldCull(ElementList, PaintGeometry, InTint, InText))
-	{
-		return;
-	}
-
-	// Don't do anything if there the font would be completely transparent 
-	if (InTint.A == 0 && !InFontInfo.OutlineSettings.IsVisible())
-	{
-		return;
-	}
-#if !UE_BUILD_SHIPPING
-	if (UE::Slate::Private::GLogPaintedText)
-	{
-		UE_LOG(LogSlate, Log, TEXT("MakeText: '%s'."), *InText);
-	}
-#endif
-	FSlateDrawElement& Element = ElementList.AddUninitialized<EElementType::ET_Text>();
-
-	FSlateTextPayload& DataPayload = ElementList.CreatePayload<FSlateTextPayload>(Element);
-
-	DataPayload.SetTint(InTint);
-	DataPayload.SetText(InText, InFontInfo);
-
-	Element.Init(ElementList, EElementType::ET_Text, InLayer, PaintGeometry, InDrawEffects);
-}
-
-namespace SlateDrawElement
-{
-#if SLATE_CHECK_UOBJECT_RENDER_RESOURCES
-	const FName MaterialInterfaceClassName = "MaterialInterface";
-	void CheckInvalidUMaterial(const UObject* InMaterialResource, const TCHAR* MaterialMessage)
-	{
-		if (InMaterialResource && GSlateCheckUObjectRenderResources)
-		{
-			bool bIsValidLowLevel = InMaterialResource->IsValidLowLevelFast(false);
-			if (!bIsValidLowLevel || !IsValid(InMaterialResource) || InMaterialResource->GetClass()->GetFName() == MaterialInterfaceClassName)
-			{
-				UE_LOG(LogSlate, Error, TEXT("Material '%s' is not valid. PendingKill:'%d'. ValidLowLevelFast:'%d'. InvalidClass:'%d'")
-					, MaterialMessage
-					, (bIsValidLowLevel ? !IsValid(InMaterialResource) : false)
-					, bIsValidLowLevel
-					, (bIsValidLowLevel ? InMaterialResource->GetClass()->GetFName() == MaterialInterfaceClassName : false));
-
-				if (bIsValidLowLevel)
-				{
-					UE_LOG(LogSlate, Error, TEXT("Material name: '%s'"), *InMaterialResource->GetFullName());
-				}
-
-				const TCHAR* Message = TEXT("We detected an invalid resource in FSlateDrawElement. Check the log for more detail.");
-				if (GSlateCheckUObjectRenderResourcesShouldLogFatal)
-				{
-					UE_LOG(LogSlate, Fatal, TEXT("%s"), Message);
-				}
-				else
-				{
-					ensureAlwaysMsgf(false, TEXT("%s"), Message);
-				}
-			}
-		}
-	}
-#endif
-}
-
-void FSlateDrawElement::MakeShapedText(FSlateWindowElementList& ElementList, uint32 InLayer, const FPaintGeometry& PaintGeometry, const FShapedGlyphSequenceRef& InShapedGlyphSequence, ESlateDrawEffect InDrawEffects, const FLinearColor& BaseTint, const FLinearColor& OutlineTint, FTextOverflowArgs TextOverflowArgs)
-{
-	SCOPE_CYCLE_COUNTER(STAT_SlateDrawElementMakeTime)
-	PaintGeometry.CommitTransformsIfUsingLegacyConstructor();
-
-	if (InShapedGlyphSequence->GetGlyphsToRender().Num() == 0)
-	{
-		return;
-	}
-
-	if (ShouldCull(ElementList, PaintGeometry))
-	{
-		return;
-	}
-
-	// Don't do anything if there the font would be completely transparent 
-	if ((BaseTint.A == 0 && InShapedGlyphSequence->GetFontOutlineSettings().OutlineSize == 0) || 
-		(BaseTint.A == 0 && OutlineTint.A == 0))
-	{
-		return;
-	}
-
-#if SLATE_CHECK_UOBJECT_RENDER_RESOURCES
-	SlateDrawElement::CheckInvalidUMaterial(InShapedGlyphSequence->GetFontMaterial(), TEXT("Font Material"));
-	SlateDrawElement::CheckInvalidUMaterial(InShapedGlyphSequence->GetFontOutlineSettings().OutlineMaterial, TEXT("Outline Material"));
-#endif
-
-	FSlateDrawElement& Element = ElementList.AddUninitialized<EElementType::ET_ShapedText>();
-
-	FSlateShapedTextPayload& DataPayload = ElementList.CreatePayload<FSlateShapedTextPayload>(Element);
-	DataPayload.SetTint(BaseTint);
-	DataPayload.SetShapedText(ElementList, InShapedGlyphSequence, OutlineTint);
-	DataPayload.SetOverflowArgs(TextOverflowArgs);
-
-	Element.Init(ElementList, EElementType::ET_ShapedText, InLayer, PaintGeometry, InDrawEffects);
-}
-
-void FSlateDrawElement::MakeGradient( FSlateWindowElementList& ElementList, uint32 InLayer, const FPaintGeometry& PaintGeometry, TArray<FSlateGradientStop> InGradientStops, EOrientation InGradientType, ESlateDrawEffect InDrawEffects, FVector4f CornerRadius)
-{
-	PaintGeometry.CommitTransformsIfUsingLegacyConstructor();
-
-	if (ShouldCull(ElementList, PaintGeometry))
-	{
-		return;
-	}
-
-	FSlateDrawElement& Element = ElementList.AddUninitialized<EElementType::ET_Gradient>();
-
-	FSlateGradientPayload& DataPayload = ElementList.CreatePayload<FSlateGradientPayload>(Element);
-
-	DataPayload.SetGradient(InGradientStops, InGradientType, CornerRadius);
-
-	Element.Init(ElementList, EElementType::ET_Gradient, InLayer, PaintGeometry, InDrawEffects);
-}
-
-void FSlateDrawElement::MakeSpline(FSlateWindowElementList& ElementList, uint32 InLayer, const FPaintGeometry& PaintGeometry, const UE::Slate::FDeprecateVector2DParameter InStart, const UE::Slate::FDeprecateVector2DParameter InStartDir, const UE::Slate::FDeprecateVector2DParameter InEnd, const UE::Slate::FDeprecateVector2DParameter InEndDir, float InThickness, ESlateDrawEffect InDrawEffects, const FLinearColor& InTint)
-{
-	PaintGeometry.CommitTransformsIfUsingLegacyConstructor();
-
-	if (ShouldCull(ElementList))
-	{
-		return;
-	}
-	FSlateDrawElement& Element = ElementList.AddUninitialized<EElementType::ET_Spline>();
-
-	FSlateSplinePayload& DataPayload = ElementList.CreatePayload<FSlateSplinePayload>(Element);
-
-	DataPayload.SetHermiteSpline(InStart, InStartDir, InEnd, InEndDir, InThickness, InTint);
-
-	Element.Init(ElementList, EElementType::ET_Spline, InLayer, PaintGeometry, InDrawEffects);
-}
-
-void FSlateDrawElement::MakeCubicBezierSpline(FSlateWindowElementList& ElementList, uint32 InLayer, const FPaintGeometry& PaintGeometry, const UE::Slate::FDeprecateVector2DParameter P0, const UE::Slate::FDeprecateVector2DParameter P1, const UE::Slate::FDeprecateVector2DParameter P2, const UE::Slate::FDeprecateVector2DParameter P3, float InThickness, ESlateDrawEffect InDrawEffects, const FLinearColor& InTint)
-{
-	PaintGeometry.CommitTransformsIfUsingLegacyConstructor();
-
-	if (ShouldCull(ElementList))
-	{
-		return;
-	}
-	FSlateDrawElement& Element = ElementList.AddUninitialized<EElementType::ET_Spline>();
-
-	FSlateSplinePayload& DataPayload = ElementList.CreatePayload<FSlateSplinePayload>(Element);
-
-	DataPayload.SetCubicBezier(P0, P1, P2, P3, InThickness, InTint);
-
-	Element.Init(ElementList, EElementType::ET_Spline, InLayer, PaintGeometry, InDrawEffects);
-}
-
-void FSlateDrawElement::MakeDrawSpaceSpline(FSlateWindowElementList& ElementList, uint32 InLayer, const UE::Slate::FDeprecateVector2DParameter InStart, const UE::Slate::FDeprecateVector2DParameter InStartDir, const UE::Slate::FDeprecateVector2DParameter InEnd, const UE::Slate::FDeprecateVector2DParameter InEndDir, float InThickness, ESlateDrawEffect InDrawEffects, const FLinearColor& InTint)
-{
-	MakeSpline(ElementList, InLayer, FPaintGeometry(), InStart, InStartDir, InEnd, InEndDir, InThickness, InDrawEffects, InTint);
-}
-
-#if UE_ENABLE_SLATE_VECTOR_DEPRECATION_MECHANISMS
-void FSlateDrawElement::MakeLines(FSlateWindowElementList& ElementList, uint32 InLayer, const FPaintGeometry& PaintGeometry, const TArray<FVector2d>& Points, ESlateDrawEffect InDrawEffects, const FLinearColor& InTint, bool bAntialias, float Thickness)
-{
-	if (ShouldCull(ElementList) || Points.Num() < 2)
-	{
-		return;
-	}
-
-	TArray<FVector2f> NewVector;
-	NewVector.Reserve(Points.Num());
-	for (FVector2d Vect : Points)
-	{
-		NewVector.Add(UE::Slate::CastToVector2f(Vect));
-	}
-	MakeLines(ElementList, InLayer, PaintGeometry, MoveTemp(NewVector), InDrawEffects, InTint, bAntialias, Thickness);
-}
-#endif // UE_ENABLE_SLATE_VECTOR_DEPRECATION_MECHANISMS
-
-void FSlateDrawElement::MakeLines(FSlateWindowElementList& ElementList, uint32 InLayer, const FPaintGeometry& PaintGeometry, TArray<FVector2f> Points, ESlateDrawEffect InDrawEffects, const FLinearColor& InTint, bool bAntialias, float Thickness)
-{
-	PaintGeometry.CommitTransformsIfUsingLegacyConstructor();
-
-	if (ShouldCull(ElementList) || Points.Num() < 2)
-	{
-		return;
-	}
-
-	FSlateDrawElement& Element = ElementList.AddUninitialized<EElementType::ET_Line>();
-
-	FSlateLinePayload& DataPayload = ElementList.CreatePayload<FSlateLinePayload>(Element);
-
-	DataPayload.SetTint(InTint);
-	DataPayload.SetThickness(Thickness);
-	DataPayload.SetLines(MoveTemp(Points), bAntialias);
-
-	ESlateDrawEffect DrawEffects = InDrawEffects;
-	if (bAntialias)
-	{
-		// If the line is to be anti-aliased, we cannot reliably snap
-		// the generated vertices
-		DrawEffects |= ESlateDrawEffect::NoPixelSnapping;
-	}
-
-	Element.Init(ElementList, EElementType::ET_Line, InLayer, PaintGeometry, DrawEffects);
-}
-
-#if UE_ENABLE_SLATE_VECTOR_DEPRECATION_MECHANISMS
-void FSlateDrawElement::MakeLines(FSlateWindowElementList& ElementList, uint32 InLayer, const FPaintGeometry& PaintGeometry, const TArray<FVector2d>& Points, const TArray<FLinearColor>& PointColors, ESlateDrawEffect InDrawEffects, const FLinearColor& InTint, bool bAntialias, float Thickness)
-{
-	if (ShouldCull(ElementList) || Points.Num() < 2)
-	{
-		return;
-	}
-
-	TArray<FVector2f> NewVector;
-	NewVector.Reserve(Points.Num());
-	for (FVector2d Vect : Points)
-	{
-		NewVector.Add(UE::Slate::CastToVector2f(Vect));
-	}
-	MakeLines(ElementList, InLayer, PaintGeometry, MoveTemp(NewVector), PointColors, InDrawEffects, InTint, bAntialias, Thickness);
-}
-#endif
-
-void FSlateDrawElement::MakeLines( FSlateWindowElementList& ElementList, uint32 InLayer, const FPaintGeometry& PaintGeometry, TArray<FVector2f> Points, TArray<FLinearColor> PointColors, ESlateDrawEffect InDrawEffects, const FLinearColor& InTint, bool bAntialias, float Thickness )
-{
-	PaintGeometry.CommitTransformsIfUsingLegacyConstructor();
-
-	if (ShouldCull(ElementList) || Points.Num() < 2)
-	{
-		return;
-	}
-
-	FSlateDrawElement& Element = ElementList.AddUninitialized<EElementType::ET_Line>();
-
-	FSlateLinePayload& DataPayload = ElementList.CreatePayload<FSlateLinePayload>(Element);
-	DataPayload.SetTint(InTint);
-	DataPayload.SetThickness(Thickness);
-	DataPayload.SetLines(MoveTemp(Points), bAntialias, MoveTemp(PointColors));
-
-	Element.Init(ElementList, EElementType::ET_Line, InLayer, PaintGeometry, InDrawEffects);
-}
-
-void FSlateDrawElement::MakeViewport( FSlateWindowElementList& ElementList, uint32 InLayer, const FPaintGeometry& PaintGeometry, TSharedPtr<const ISlateViewport> Viewport, ESlateDrawEffect InDrawEffects, const FLinearColor& InTint )
-{
-	PaintGeometry.CommitTransformsIfUsingLegacyConstructor();
-
-	if (ShouldCull(ElementList))
-	{
-		return;
-	}
-
-	FSlateDrawElement& Element = ElementList.AddUninitialized<EElementType::ET_Viewport>();
-	FSlateViewportPayload& DataPayload = ElementList.CreatePayload<FSlateViewportPayload>(Element);
-
-	DataPayload.SetViewport(Viewport, InTint);
-	check(DataPayload.RenderTargetResource == nullptr || !DataPayload.RenderTargetResource->Debug_IsDestroyed());
-
-	Element.Init(ElementList, EElementType::ET_Viewport, InLayer, PaintGeometry, InDrawEffects);
-}
-
-
-void FSlateDrawElement::MakeCustom( FSlateWindowElementList& ElementList, uint32 InLayer, TSharedPtr<ICustomSlateElement, ESPMode::ThreadSafe> CustomDrawer )
-{
-	if (ShouldCull(ElementList))
-	{
-		return;
-	}
-
-	FSlateDrawElement& Element = ElementList.AddUninitialized<EElementType::ET_Custom>();
-
-	FSlateCustomDrawerPayload& DataPayload = ElementList.CreatePayload<FSlateCustomDrawerPayload>(Element);
-	DataPayload.SetCustomDrawer(CustomDrawer);
-	
-	Element.Init(ElementList, EElementType::ET_Custom, InLayer, FPaintGeometry(), ESlateDrawEffect::None);
-	Element.RenderTransform = FSlateRenderTransform();
-}
-
-
-void FSlateDrawElement::MakeCustomVerts(FSlateWindowElementList& ElementList, uint32 InLayer, const FSlateResourceHandle& InRenderResourceHandle, const TArray<FSlateVertex>& InVerts, const TArray<SlateIndex>& InIndexes, ISlateUpdatableInstanceBuffer* InInstanceData, uint32 InInstanceOffset, uint32 InNumInstances, ESlateDrawEffect InDrawEffects)
-{
-	SCOPE_CYCLE_COUNTER(STAT_SlateDrawElementMakeCustomVertsTime);
-	
-	if (ShouldCull(ElementList))
-	{
-		return;
-	}
-
-	FSlateDrawElement& Element = ElementList.AddUninitialized<EElementType::ET_CustomVerts>();
-	FSlateCustomVertsPayload& DataPayload = ElementList.CreatePayload<FSlateCustomVertsPayload>(Element);
-
-	const FSlateShaderResourceProxy* RenderingProxy = InRenderResourceHandle.GetResourceProxy();
-	ISlateUpdatableInstanceBufferRenderProxy* RenderProxy = InInstanceData ? InInstanceData->GetRenderProxy() : nullptr;
-	DataPayload.SetCustomVerts(RenderingProxy, InVerts, InIndexes, RenderProxy, InInstanceOffset, InNumInstances);
-
-	Element.Init(ElementList, EElementType::ET_CustomVerts, InLayer, FPaintGeometry(), InDrawEffects);
-	Element.RenderTransform = FSlateRenderTransform();
-}
-
-void FSlateDrawElement::MakePostProcessPass(FSlateWindowElementList& ElementList, uint32 InLayer, const FPaintGeometry& PaintGeometry, const FVector4f& Params, int32 DownsampleAmount, FVector4f CornerRadius)
-{
-	PaintGeometry.CommitTransformsIfUsingLegacyConstructor();
-
-	if (ShouldCull(ElementList))
-	{
-		return;
-	}
-
-	FSlateDrawElement& Element = ElementList.AddUninitialized<EElementType::ET_PostProcessPass>();
-
-	FSlatePostProcessPayload& DataPayload = ElementList.CreatePayload<FSlatePostProcessPayload>(Element);
-	DataPayload.DownsampleAmount = DownsampleAmount;
-	DataPayload.PostProcessData = Params;
-	DataPayload.CornerRadius = CornerRadius;
-
-	Element.Init(ElementList, EElementType::ET_PostProcessPass, InLayer, PaintGeometry, ESlateDrawEffect::None);
-}
-
-FSlateDrawElement::FSlateDrawElement()
-	: DataPayload(nullptr)
-	, bIsCached(false)
-{
-
-}
-
-FSlateDrawElement::~FSlateDrawElement()
-{
-	if (bIsCached)
-	{
-		delete DataPayload;
-	}
-	else if (DataPayload)
-	{
-		// Allocated by a memstack so we just need to call the destructor manually
-		DataPayload->~FSlateDataPayload();
-	}
-}
-
-FVector2f FSlateDrawElement::GetRotationPoint(const FPaintGeometry& PaintGeometry, const TOptional<FVector2f>& UserRotationPoint, ERotationSpace RotationSpace)
-{
-	FVector2f RotationPoint(0, 0);
-
-	const FVector2f LocalSize = UE::Slate::CastToVector2f(PaintGeometry.GetLocalSize());
-
-	switch (RotationSpace)
-	{
-		case RelativeToElement:
-		{
-			// If the user did not specify a rotation point, we rotate about the center of the element
-			RotationPoint = UserRotationPoint.Get(LocalSize * 0.5f);
-		}
-		break;
-		case RelativeToWorld:
-		{
-			// its in world space, must convert the point to local space.
-			RotationPoint = TransformPoint(Inverse(PaintGeometry.GetAccumulatedRenderTransform()), UserRotationPoint.Get(FVector2f::ZeroVector));
-		}
-		break;
-	default:
-		check(0);
-		break;
-	}
-
-	return RotationPoint;
-}
-
-void FSlateDrawElement::AddReferencedObjects(FReferenceCollector& Collector)
-{
-	if(DataPayload)
-	{
-		DataPayload->AddReferencedObjects(Collector);
-	}
-}
-
 void FSlateWindowElementList::SetIsInGameLayer(bool bInGameLayer)
 {
 	bIsInGameLayer = bInGameLayer;
@@ -1104,9 +384,9 @@ void FSlateWindowElementList::ResetElementList()
 
 	ClippingManager.ResetClippingState();
 
-	auto ResetElementContainer = [](FSlateDrawElementContainer& Container)
+	auto ResetElementContainer = [](auto& Container)
 	{
-		Container.Elements.Reset();
+		Container.Reset();
 	};
 
 	VisitTupleElements(ResetElementContainer, UncachedDrawElements);
@@ -1130,14 +410,19 @@ void FSlateWindowElementList::ResetElementList()
 	RenderTargetWindow = nullptr;
 }
 
-
 void FSlateWindowElementList::AddReferencedObjects(FReferenceCollector& Collector)
 {
-	auto AddReferencedElementContainer = [&](FSlateDrawElementContainer& Container)
+	auto AddReferencedElementContainer = [&](auto& Container)
 	{
-		for (FSlateDrawElement& Element : Container.Elements)
+		// Container should just be some FSlateDrawElementArray<T>
+		using FSlateElementType = typename std::remove_reference<decltype(Container)>::type::ElementType;
+
+		if constexpr (TIsReferencedObjects<FSlateElementType>::Value)
 		{
-			Element.AddReferencedObjects(Collector);
+			for (FSlateElementType& Element : Container)
+			{
+				Element.AddReferencedObjects(Collector);
+			}
 		}
 	};
 
@@ -1307,11 +592,17 @@ void FSlateCachedElementList::AddCachedClipState(FSlateCachedClipState& ClipStat
 
 void FSlateCachedElementList::AddReferencedObjects(FReferenceCollector& Collector)
 {
-	auto AddReferencedElementContainer = [&](FSlateDrawElementContainer& Container)
+	auto AddReferencedElementContainer = [&](auto& Container)
 	{
-		for (FSlateDrawElement& Element : Container.Elements)
+		// Container should just be some FSlateDrawElementArray<T>
+		using FSlateElementType = typename std::remove_reference<decltype(Container)>::type::ElementType;
+
+		if constexpr (TIsReferencedObjects<FSlateElementType>::Value)
 		{
-			Element.AddReferencedObjects(Collector);
+			for (FSlateElementType& Element : Container)
+			{
+				Element.AddReferencedObjects(Collector);
+			}
 		}
 	};
 
@@ -1321,12 +612,12 @@ void FSlateCachedElementList::AddReferencedObjects(FReferenceCollector& Collecto
 bool FSlateCachedElementList::IsEmpty()
 {
 	bool bElementsFound = false;
-	
-	auto CheckElementsExist = [&](FSlateDrawElementContainer& Container)
+
+	auto CheckElementsExist = [&](auto& Container)
 	{
 		if (!bElementsFound)
 		{
-			bElementsFound = Container.Elements.Num() > 0;
+			bElementsFound = Container.Num() > 0;
 		}
 	};
 
@@ -1338,10 +629,10 @@ bool FSlateCachedElementList::IsEmpty()
 int32 FSlateCachedElementList::NumElements()
 {
 	int32 NumElements = 0;
-	
-	auto CountNumElements = [&](FSlateDrawElementContainer& Container)
+
+	auto CountNumElements = [&](auto& Container)
 	{
-		NumElements += Container.Elements.Num();
+		NumElements += Container.Num();
 	};
 
 	VisitTupleElements(CountNumElements, DrawElements);
@@ -1352,9 +643,9 @@ int32 FSlateCachedElementList::NumElements()
 void FSlateCachedElementList::DestroyCachedData()
 {
 	// Clear out any cached draw elements
-	auto ResetElementContainer = [](FSlateDrawElementContainer& Container)
+	auto ResetElementContainer = [](auto& Container)
 	{
-		Container.Elements.Reset();
+		Container.Reset();
 	};
 
 	VisitTupleElements(ResetElementContainer, DrawElements);
@@ -1422,11 +713,11 @@ bool FSlateCachedElementsHandle::HasCachedElements() const
 	{
 		bool bElementsFound = false;
 
-		auto CheckElementsExist = [&](FSlateDrawElementContainer& Container)
+		auto CheckElementsExist = [&](auto& Container)
 		{
 			if (!bElementsFound)
 			{
-				bElementsFound = Container.Elements.Num() > 0;
+				bElementsFound = Container.Num() > 0;
 			}
 		};
 
@@ -1436,4 +727,4 @@ bool FSlateCachedElementsHandle::HasCachedElements() const
 	}
 
 	return false;
-}
+}
\ No newline at end of file
diff --git a/Engine/Source/Runtime/SlateCore/Private/Rendering/ElementBatcher.cpp b/Engine/Source/Runtime/SlateCore/Private/Rendering/ElementBatcher.cpp
--- a/Engine/Source/Runtime/SlateCore/Private/Rendering/ElementBatcher.cpp
+++ b/Engine/Source/Runtime/SlateCore/Private/Rendering/ElementBatcher.cpp
@@ -4,7 +4,6 @@
 #include "Fonts/SlateFontInfo.h"
 #include "Fonts/FontCache.h"
 #include "Rendering/RenderingCommon.h"
-#include "Rendering/DrawElementPayloads.h"
 #include "Rendering/DrawElements.h"
 #include "Rendering/RenderingPolicy.h"
 #include "Rendering/ShaderResourceManager.h"
@@ -13,6 +12,7 @@
 #include "ProfilingDebugging/ScopedTimers.h"
 #include "Debugging/SlateDebugging.h"
 #include "Widgets/SWidgetUtils.h"
+#include "Rendering/DrawElementPayloads.h"
 
 CSV_DECLARE_CATEGORY_MODULE_EXTERN(SLATECORE_API, Slate);
 
@@ -141,19 +141,22 @@ void FSlateBatchData::CombineBatches(FSlateRenderBatch& FirstBatch, FSlateRender
 		// when merging verts we have to offset the indices in the second batch based on the first batches existing number of verts
 		const int32 BatchOffset = FirstBatch.NumVertices;
 
-		// Final vertices is assumed to have the first batch already in it
-		FirstBatch.NumVertices += SecondBatch.NumVertices;
-		FirstBatch.NumIndices += SecondBatch.NumIndices;
+		if (SecondBatch.HasVertexData())
+		{
+			// Final vertices is assumed to have the first batch already in it
+			FirstBatch.NumVertices += SecondBatch.NumVertices;
+			FirstBatch.NumIndices += SecondBatch.NumIndices;
+
+			FinalVertices.Append(&(*SecondBatch.SourceVertices)[SecondBatch.VertexOffset], SecondBatch.NumVertices);
 
-		FinalVertices.Append(&(*SecondBatch.SourceVertices)[SecondBatch.VertexOffset], SecondBatch.NumVertices);
+			FinalIndices.Reserve(FinalIndices.Num() + SecondBatch.NumIndices);
 
-		FinalIndices.Reserve(FinalIndices.Num()+SecondBatch.NumIndices);
-		
-		// Get source indices at the source index offset and shift each index by the batches current offset
-		for (int32 i = 0; i < SecondBatch.NumIndices; ++i)
-		{
-			const int32 FinalIndex = (*SecondBatch.SourceIndices)[i + SecondBatch.IndexOffset] + BatchOffset;
-			FinalIndices.Add(FinalIndex);
+			// Get source indices at the source index offset and shift each index by the batches current offset
+			for (int32 i = 0; i < SecondBatch.NumIndices; ++i)
+			{
+				const int32 FinalIndex = (*SecondBatch.SourceIndices)[i + SecondBatch.IndexOffset] + BatchOffset;
+				FinalIndices.Add(FinalIndex);
+			}
 		}
 	}
 
@@ -356,120 +359,25 @@ void FSlateElementBatcher::AddElements(FSlateWindowElementList& WindowElementLis
 
 void FSlateElementBatcher::AddElementsInternal(const FSlateDrawElementMap& DrawElements, FVector2f ViewportSize)
 {
-	auto HandleUnMapped = [&](const FSlateDrawElementArray& Elements, FVector2f ViewportSize)
-	{
-		for (const FSlateDrawElement& DrawElement : Elements)
-		{
-			// Determine what type of element to add
-			switch (DrawElement.GetElementType())
-			{
-			case EElementType::ET_Box:
-			case EElementType::ET_RoundedBox:
-			{
-				checkf(0, TEXT("All box elements should be mapped and use bulk element processing"));
-			}
-			break;
-			case EElementType::ET_Border:
-			{
-				SCOPED_NAMED_EVENT_TEXT("Slate::AddBorderElement", FColor::Magenta);
-				STAT(ElementStat_Borders++);
-				DrawElement.IsPixelSnapped() ? AddBorderElement<ESlateVertexRounding::Enabled>(DrawElement) : AddBorderElement<ESlateVertexRounding::Disabled>(DrawElement);
-			}
-			break;
-			case EElementType::ET_Text:
-			{
-				SCOPED_NAMED_EVENT_TEXT("Slate::AddTextElement", FColor::Magenta);
-				STAT(ElementStat_Text++);
-				DrawElement.IsPixelSnapped() ? AddTextElement<ESlateVertexRounding::Enabled>(DrawElement) : AddTextElement<ESlateVertexRounding::Disabled>(DrawElement);
-			}
-			break;
-			case EElementType::ET_ShapedText:
-			{
-				SCOPED_NAMED_EVENT_TEXT("Slate::AddShapedTextElement", FColor::Magenta);
-				STAT(ElementStat_ShapedText++);
-				DrawElement.IsPixelSnapped() ? AddShapedTextElement<ESlateVertexRounding::Enabled>(DrawElement) : AddShapedTextElement<ESlateVertexRounding::Disabled>(DrawElement);
-			}
-			break;
-			case EElementType::ET_Line:
-			{
-				checkf(0, TEXT("All line elements should be mapped and use bulk element processing"));
-			}
-			break;
-			case EElementType::ET_DebugQuad:
-			{
-				SCOPED_NAMED_EVENT_TEXT("Slate::AddDebugQuadElement", FColor::Magenta);
-				STAT(ElementStat_Other++);
-				DrawElement.IsPixelSnapped() ? AddDebugQuadElement<ESlateVertexRounding::Enabled>(DrawElement) : AddDebugQuadElement<ESlateVertexRounding::Disabled>(DrawElement);
-			}
-			break;
-			case EElementType::ET_Spline:
-			{
-				SCOPED_NAMED_EVENT_TEXT("Slate::AddSplineElement", FColor::Magenta);
-				// Note that we ignore pixel snapping here; see implementation for more info.
-				STAT(ElementStat_Other++);
-				AddSplineElement(DrawElement);
-			}
-			break;
-			case EElementType::ET_Gradient:
-			{
-				SCOPED_NAMED_EVENT_TEXT("Slate::AddGradientElement", FColor::Magenta);
-				STAT(ElementStat_Other++);
-				DrawElement.IsPixelSnapped() ? AddGradientElement<ESlateVertexRounding::Enabled>(DrawElement) : AddGradientElement<ESlateVertexRounding::Disabled>(DrawElement);
-			}
-			break;
-			case EElementType::ET_Viewport:
-			{
-				SCOPED_NAMED_EVENT_TEXT("Slate::AddViewportElement", FColor::Magenta);
-				STAT(ElementStat_Other++);
-				DrawElement.IsPixelSnapped() ? AddViewportElement<ESlateVertexRounding::Enabled>(DrawElement) : AddViewportElement<ESlateVertexRounding::Disabled>(DrawElement);
-			}
-			break;
-			case EElementType::ET_Custom:
-			{
-				SCOPED_NAMED_EVENT_TEXT("Slate::AddCustomElement", FColor::Magenta);
-				STAT(ElementStat_Other++);
-				AddCustomElement(DrawElement);
-			}
-			break;
-			case EElementType::ET_CustomVerts:
-			{
-				SCOPED_NAMED_EVENT_TEXT("Slate::AddCustomVertsElement", FColor::Magenta);
-				STAT(ElementStat_Other++);
-				AddCustomVerts(DrawElement);
-			}
-			break;
-			case EElementType::ET_PostProcessPass:
-			{
-				SCOPED_NAMED_EVENT_TEXT("Slate::AddPostProcessElement", FColor::Magenta);
-				STAT(ElementStat_Other++);
-				AddPostProcessPass(DrawElement, ViewportSize);
-			}
-			break;
-			default:
-				checkf(0, TEXT("Invalid element type"));
-			}
-		}
-	};
-
 	{
 		SCOPED_NAMED_EVENT_TEXT("Slate::AddBoxElements", FColor::Magenta);
-		const FSlateDrawElementContainer& Container = DrawElements.Get<(uint8)EElementType::ET_Box>();
-		STAT(ElementStat_Other += Container.Elements.Num());
-		AddBoxElements(Container.Elements);
+		const FSlateDrawElementArray<FSlateBoxElement>& Container = DrawElements.Get<(uint8)EElementType::ET_Box>();
+		STAT(ElementStat_Other += Container.Num());
+		AddBoxElements(Container);
 	}
 
 	{
 		SCOPED_NAMED_EVENT_TEXT("Slate::AddRoundedBoxElements", FColor::Magenta);
-		const FSlateDrawElementContainer& Container = DrawElements.Get<(uint8)EElementType::ET_RoundedBox>();
-		STAT(ElementStat_Other += Container.Elements.Num());
-		AddBoxElements(Container.Elements);
+		const FSlateDrawElementArray<FSlateRoundedBoxElement>& Container = DrawElements.Get<(uint8)EElementType::ET_RoundedBox>();
+		STAT(ElementStat_Other += Container.Num());
+		AddBoxElements(Container);
 	}
 
 	{
 		SCOPED_NAMED_EVENT_TEXT("Slate::AddBorderElement", FColor::Magenta);
-		const FSlateDrawElementContainer& Container = DrawElements.Get<(uint8)EElementType::ET_Border>();
-		STAT(ElementStat_Other += Container.Elements.Num());
-		for (const FSlateDrawElement& DrawElement : Container.Elements)
+		const FSlateDrawElementArray<FSlateBoxElement>& Container = DrawElements.Get<(uint8)EElementType::ET_Border>();
+		STAT(ElementStat_Other += Container.Num());
+		for (const FSlateBoxElement& DrawElement : Container)
 		{
 			DrawElement.IsPixelSnapped() ? AddBorderElement<ESlateVertexRounding::Enabled>(DrawElement) : AddBorderElement<ESlateVertexRounding::Disabled>(DrawElement);
 		}
@@ -477,9 +385,9 @@ void FSlateElementBatcher::AddElementsInternal(const FSlateDrawElementMap& DrawE
 
 	{
 		SCOPED_NAMED_EVENT_TEXT("Slate::AddTextElement", FColor::Magenta);
-		const FSlateDrawElementContainer& Container = DrawElements.Get<(uint8)EElementType::ET_Text>();
-		STAT(ElementStat_Other += Container.Elements.Num());
-		for (const FSlateDrawElement& DrawElement : Container.Elements)
+		const FSlateDrawElementArray<FSlateTextElement>& Container = DrawElements.Get<(uint8)EElementType::ET_Text>();
+		STAT(ElementStat_Other += Container.Num());
+		for (const FSlateTextElement& DrawElement : Container)
 		{
 			DrawElement.IsPixelSnapped() ? AddTextElement<ESlateVertexRounding::Enabled>(DrawElement) : AddTextElement<ESlateVertexRounding::Disabled>(DrawElement);
 		}
@@ -487,9 +395,9 @@ void FSlateElementBatcher::AddElementsInternal(const FSlateDrawElementMap& DrawE
 
 	{
 		SCOPED_NAMED_EVENT_TEXT("Slate::AddShapedTextElement", FColor::Magenta);
-		const FSlateDrawElementContainer& Container = DrawElements.Get<(uint8)EElementType::ET_ShapedText>();
-		STAT(ElementStat_Other += Container.Elements.Num());
-		for (const FSlateDrawElement& DrawElement : Container.Elements)
+		const FSlateDrawElementArray<FSlateShapedTextElement>& Container = DrawElements.Get<(uint8)EElementType::ET_ShapedText>();
+		STAT(ElementStat_Other += Container.Num());
+		for (const FSlateShapedTextElement& DrawElement : Container)
 		{
 			DrawElement.IsPixelSnapped() ? AddShapedTextElement<ESlateVertexRounding::Enabled>(DrawElement) : AddShapedTextElement<ESlateVertexRounding::Disabled>(DrawElement);
 		}
@@ -497,16 +405,16 @@ void FSlateElementBatcher::AddElementsInternal(const FSlateDrawElementMap& DrawE
 
 	{
 		SCOPED_NAMED_EVENT_TEXT("Slate::AddLineElements", FColor::Magenta);
-		const FSlateDrawElementContainer& Container = DrawElements.Get<(uint8)EElementType::ET_Line>();
-		STAT(ElementStat_Other += Container.Elements.Num());
-		AddLineElements(Container.Elements);
+		const FSlateDrawElementArray<FSlateLineElement>& Container = DrawElements.Get<(uint8)EElementType::ET_Line>();
+		STAT(ElementStat_Other += Container.Num());
+		AddLineElements(Container);
 	}
 
 	{
 		SCOPED_NAMED_EVENT_TEXT("Slate::AddDebugQuadElement", FColor::Magenta);
-		const FSlateDrawElementContainer& Container = DrawElements.Get<(uint8)EElementType::ET_DebugQuad>();
-		STAT(ElementStat_Other += Container.Elements.Num());
-		for (const FSlateDrawElement& DrawElement : Container.Elements)
+		const FSlateDrawElementArray<FSlateBoxElement>& Container = DrawElements.Get<(uint8)EElementType::ET_DebugQuad>();
+		STAT(ElementStat_Other += Container.Num());
+		for (const FSlateBoxElement& DrawElement : Container)
 		{
 			DrawElement.IsPixelSnapped() ? AddDebugQuadElement<ESlateVertexRounding::Enabled>(DrawElement) : AddDebugQuadElement<ESlateVertexRounding::Disabled>(DrawElement);
 		}
@@ -514,11 +422,11 @@ void FSlateElementBatcher::AddElementsInternal(const FSlateDrawElementMap& DrawE
 
 	{
 		SCOPED_NAMED_EVENT_TEXT("Slate::AddSplineElement", FColor::Magenta);
-		const FSlateDrawElementContainer& Container = DrawElements.Get<(uint8)EElementType::ET_Spline>();
+		const FSlateDrawElementArray<FSlateSplineElement>& Container = DrawElements.Get<(uint8)EElementType::ET_Spline>();
 
 		// Note that we ignore pixel snapping here; see implementation for more info.
-		STAT(ElementStat_Other += Container.Elements.Num());
-		for (const FSlateDrawElement& DrawElement : Container.Elements)
+		STAT(ElementStat_Other += Container.Num());
+		for (const FSlateSplineElement& DrawElement : Container)
 		{
 			AddSplineElement(DrawElement);
 		}
@@ -526,9 +434,9 @@ void FSlateElementBatcher::AddElementsInternal(const FSlateDrawElementMap& DrawE
 
 	{
 		SCOPED_NAMED_EVENT_TEXT("Slate::AddGradientElement", FColor::Magenta);
-		const FSlateDrawElementContainer& Container = DrawElements.Get<(uint8)EElementType::ET_Gradient>();
-		STAT(ElementStat_Other += Container.Elements.Num());
-		for (const FSlateDrawElement& DrawElement : Container.Elements)
+		const FSlateDrawElementArray<FSlateGradientElement>& Container = DrawElements.Get<(uint8)EElementType::ET_Gradient>();
+		STAT(ElementStat_Other += Container.Num());
+		for (const FSlateGradientElement& DrawElement : Container)
 		{
 			DrawElement.IsPixelSnapped() ? AddGradientElement<ESlateVertexRounding::Enabled>(DrawElement) : AddGradientElement<ESlateVertexRounding::Disabled>(DrawElement);
 		}
@@ -536,9 +444,9 @@ void FSlateElementBatcher::AddElementsInternal(const FSlateDrawElementMap& DrawE
 
 	{
 		SCOPED_NAMED_EVENT_TEXT("Slate::AddViewportElement", FColor::Magenta);
-		const FSlateDrawElementContainer& Container = DrawElements.Get<(uint8)EElementType::ET_Viewport>();
-		STAT(ElementStat_Other += Container.Elements.Num());
-		for (const FSlateDrawElement& DrawElement : Container.Elements)
+		const FSlateDrawElementArray<FSlateViewportElement>& Container = DrawElements.Get<(uint8)EElementType::ET_Viewport>();
+		STAT(ElementStat_Other += Container.Num());
+		for (const FSlateViewportElement& DrawElement : Container)
 		{
 			DrawElement.IsPixelSnapped() ? AddViewportElement<ESlateVertexRounding::Enabled>(DrawElement) : AddViewportElement<ESlateVertexRounding::Disabled>(DrawElement);
 		}
@@ -546,9 +454,9 @@ void FSlateElementBatcher::AddElementsInternal(const FSlateDrawElementMap& DrawE
 
 	{
 		SCOPED_NAMED_EVENT_TEXT("Slate::AddCustomElement", FColor::Magenta);
-		const FSlateDrawElementContainer& Container = DrawElements.Get<(uint8)EElementType::ET_Custom>();
-		STAT(ElementStat_Other += Container.Elements.Num());
-		for (const FSlateDrawElement& DrawElement : Container.Elements)
+		const FSlateDrawElementArray<FSlateCustomDrawerElement>& Container = DrawElements.Get<(uint8)EElementType::ET_Custom>();
+		STAT(ElementStat_Other += Container.Num());
+		for (const FSlateCustomDrawerElement& DrawElement : Container)
 		{
 			AddCustomElement(DrawElement);
 		}
@@ -556,9 +464,9 @@ void FSlateElementBatcher::AddElementsInternal(const FSlateDrawElementMap& DrawE
 
 	{
 		SCOPED_NAMED_EVENT_TEXT("Slate::AddCustomVertsElement", FColor::Magenta);
-		const FSlateDrawElementContainer& Container = DrawElements.Get<(uint8)EElementType::ET_CustomVerts>();
-		STAT(ElementStat_Other += Container.Elements.Num());
-		for (const FSlateDrawElement& DrawElement : Container.Elements)
+		const FSlateDrawElementArray<FSlateCustomVertsElement>& Container = DrawElements.Get<(uint8)EElementType::ET_CustomVerts>();
+		STAT(ElementStat_Other += Container.Num());
+		for (const FSlateCustomVertsElement& DrawElement : Container)
 		{
 			AddCustomVerts(DrawElement);
 		}
@@ -566,18 +474,17 @@ void FSlateElementBatcher::AddElementsInternal(const FSlateDrawElementMap& DrawE
 
 	{
 		SCOPED_NAMED_EVENT_TEXT("Slate::AddPostProcessElement", FColor::Magenta);
-		const FSlateDrawElementContainer& Container = DrawElements.Get<(uint8)EElementType::ET_PostProcessPass>();
-		STAT(ElementStat_Other += Container.Elements.Num());
-		for (const FSlateDrawElement& DrawElement : Container.Elements)
+		const FSlateDrawElementArray<FSlatePostProcessElement>& Container = DrawElements.Get<(uint8)EElementType::ET_PostProcessPass>();
+		STAT(ElementStat_Other += Container.Num());
+		for (const FSlatePostProcessElement& DrawElement : Container)
 		{
 			AddPostProcessPass(DrawElement, ViewportSize);
 		}
 	}
 
 	{
-		SCOPED_NAMED_EVENT_TEXT("Slate::AddNonMappedElement", FColor::Magenta);
-		const FSlateDrawElementContainer& Container = DrawElements.Get<(uint8)EElementType::ET_NonMapped>();
-		HandleUnMapped(Container.Elements, ViewportSize);
+		const FSlateDrawElementArray<FSlateDrawElement>& Container = DrawElements.Get<(uint8)EElementType::ET_NonMapped>();
+		ensure(Container.IsEmpty());
 	}
 }
 
@@ -630,9 +537,9 @@ void FSlateElementBatcher::AddCachedElements(FSlateCachedElementData& CachedElem
 #endif
 }
 
-template<typename ElementAdder, typename ElementBatchParamCreator, typename ElementBatchReserver>
+template<typename ElementType, typename ElementAdder, typename ElementBatchParamCreator, typename ElementBatchReserver>
 void FSlateElementBatcher::GenerateIndexedVertexBatches(
-	const FSlateDrawElementArray& DrawElements
+	const FSlateDrawElementArray<ElementType>& DrawElements
 	, ElementAdder&& InElementAdder
 	, ElementBatchParamCreator&& InElementBatchParamCreator
 	, ElementBatchReserver&& InElementBatchReserver)
@@ -652,7 +559,7 @@ void FSlateElementBatcher::GenerateIndexedVertexBatches(
 		int32 BatchIndexEnd = Index;
 		while (DrawElements.IsValidIndex(++BatchIndexEnd))
 		{
-			const FSlateDrawElement& NextDrawElement = DrawElements[BatchIndexEnd];
+			const ElementType& NextDrawElement = DrawElements[BatchIndexEnd];
 			FSlateRenderBatchParams NextBatchParams;
 			InElementBatchParamCreator(NextDrawElement, NextBatchParams);
 
@@ -662,7 +569,7 @@ void FSlateElementBatcher::GenerateIndexedVertexBatches(
 			}
 		}
 
-		// Process valid range
+		// Process valid range, we always create the batch - even if it may not have any elements later on (Ex: 1 point lines).
 		FSlateRenderBatch& RenderBatch = CreateRenderBatch(
 			NewBatchParams.Layer
 			, NewBatchParams.ShaderParams
@@ -686,11 +593,9 @@ void FSlateElementBatcher::GenerateIndexedVertexBatches(
 }
 
 template<ESlateVertexRounding Rounding>
-void FSlateElementBatcher::AddDebugQuadElement(const FSlateDrawElement& DrawElement)
+void FSlateElementBatcher::AddDebugQuadElement(const FSlateBoxElement& DrawElement)
 {
-	const FSlateBoxPayload& DrawElementPayload = DrawElement.GetDataPayload<FSlateBoxPayload>();
-
-	const FColor Tint = PackVertexColor(DrawElementPayload.GetTint());
+	const FColor Tint = PackVertexColor(DrawElement.GetTint());
 	const FSlateRenderTransform& RenderTransform = DrawElement.GetRenderTransform();
 	const FVector2f LocalSize = DrawElement.GetLocalSize();
 	ESlateDrawEffect InDrawEffects = DrawElement.GetDrawEffects();
@@ -698,7 +603,7 @@ void FSlateElementBatcher::AddDebugQuadElement(const FSlateDrawElement& DrawElem
 
 	FSlateRenderBatch& RenderBatch = CreateRenderBatch(Layer, FShaderParams(), nullptr, ESlateDrawPrimitive::TriangleList, ESlateShader::Default, ESlateDrawEffect::None, ESlateBatchDrawFlag::None, DrawElement);
 	
-	const FColor Color = PackVertexColor(DrawElementPayload.GetTint());
+	const FColor Color = PackVertexColor(DrawElement.GetTint());
 
 	// Determine the four corners of the quad
 	FVector2f TopLeft = FVector2f::ZeroVector;
@@ -737,14 +642,14 @@ FORCEINLINE void IndexQuad(FSlateRenderBatch& RenderBatch, int32 TopLeft, int32
 	RenderBatch.AddIndex(TopLeft);
 }
 
-void FSlateElementBatcher::AddBoxElements(const FSlateDrawElementArray& DrawElements)
+template<typename ElementType>
+void FSlateElementBatcher::AddBoxElements(const FSlateDrawElementArray<ElementType>& DrawElements)
 {
-	auto AddBoxElementInternal = [&](FSlateRenderBatch& RenderBatch, const FSlateDrawElement& DrawElement)
+	auto AddBoxElementInternal = [&](FSlateRenderBatch& RenderBatch, const ElementType& DrawElement)
 	{
 		const ESlateVertexRounding Rounding = DrawElement.IsPixelSnapped() ? ESlateVertexRounding::Enabled : ESlateVertexRounding::Disabled;
-		const FSlateBoxPayload& DrawElementPayload = DrawElement.GetDataPayload<FSlateBoxPayload>();
 
-		const FColor Tint = PackVertexColor(DrawElementPayload.GetTint());
+		const FColor Tint = PackVertexColor(DrawElement.GetTint());
 		const FSlateRenderTransform& ElementRenderTransform = DrawElement.GetRenderTransform();
 		const FSlateRenderTransform RenderTransform = DrawElement.GetRenderTransform();// GetBoxRenderTransform(DrawElement);
 		const FVector2f LocalSize = DrawElement.GetLocalSize();
@@ -768,7 +673,7 @@ void FSlateElementBatcher::AddBoxElements(const FSlateDrawElementArray& DrawElem
 
 		FVector2f HalfTexel;
 
-		const FSlateShaderResourceProxy* ResourceProxy = DrawElementPayload.GetResourceProxy();
+		const FSlateShaderResourceProxy* ResourceProxy = DrawElement.GetResourceProxy();
 		FSlateShaderResource* Resource = nullptr;
 		if (ResourceProxy)
 		{
@@ -781,7 +686,7 @@ void FSlateElementBatcher::AddBoxElements(const FSlateDrawElementArray& DrawElem
 			// Texel offset
 			HalfTexel = FVector2f(PixelCenterOffset / TextureWidth, PixelCenterOffset / TextureHeight);
 
-			const FBox2f& BrushUV = DrawElementPayload.GetBrushUVRegion();
+			const FBox2f& BrushUV = DrawElement.GetBrushUVRegion();
 			//In case brush has valid UV region - use it instead of proxy UV
 			if (BrushUV.bIsValid)
 			{
@@ -804,11 +709,11 @@ void FSlateElementBatcher::AddBoxElements(const FSlateDrawElementArray& DrawElem
 		}
 
 
-		const ESlateBrushTileType::Type TilingRule = DrawElementPayload.GetBrushTiling();
+		const ESlateBrushTileType::Type TilingRule = DrawElement.GetBrushTiling();
 		const bool bTileHorizontal = (TilingRule == ESlateBrushTileType::Both || TilingRule == ESlateBrushTileType::Horizontal);
 		const bool bTileVertical = (TilingRule == ESlateBrushTileType::Both || TilingRule == ESlateBrushTileType::Vertical);
 
-		const ESlateBrushMirrorType::Type MirroringRule = DrawElementPayload.GetBrushMirroring();
+		const ESlateBrushMirrorType::Type MirroringRule = DrawElement.GetBrushMirroring();
 		const bool bMirrorHorizontal = (MirroringRule == ESlateBrushMirrorType::Both || MirroringRule == ESlateBrushMirrorType::Horizontal);
 		const bool bMirrorVertical = (MirroringRule == ESlateBrushMirrorType::Both || MirroringRule == ESlateBrushMirrorType::Vertical);
 
@@ -820,16 +725,14 @@ void FSlateElementBatcher::AddBoxElements(const FSlateDrawElementArray& DrawElem
 		ESlateShader ShaderType = ESlateShader::Default;
 		FShaderParams ShaderParams;
 		FColor SecondaryColor;
-		if (DrawElement.GetElementType() == EElementType::ET_RoundedBox)
+		if constexpr (std::is_same<ElementType, FSlateRoundedBoxElement>::value)
 		{
-
 			ShaderType = ESlateShader::RoundedBox;
-			const FSlateRoundedBoxPayload& RoundedPayload = DrawElement.GetDataPayload<FSlateRoundedBoxPayload>();
 
-			ShaderParams.PixelParams = FVector4f(0, RoundedPayload.GetOutlineWeight(), LocalSize.X, LocalSize.Y);//RadiusWeight;
-			ShaderParams.PixelParams2 = RoundedPayload.GetRadius();
+			ShaderParams.PixelParams = FVector4f(0, DrawElement.GetOutlineWeight(), LocalSize.X, LocalSize.Y);//RadiusWeight;
+			ShaderParams.PixelParams2 = DrawElement.GetRadius();
 
-			SecondaryColor = PackVertexColor(RoundedPayload.OutlineColor);
+			SecondaryColor = PackVertexColor(DrawElement.OutlineColor);
 		}
 
 		float HorizontalTiling = bTileHorizontal ? LocalSize.X / TextureWidth : 1.0f;
@@ -840,14 +743,14 @@ void FSlateElementBatcher::AddBoxElements(const FSlateDrawElementArray& DrawElem
 		// The start index of these vertices in the index buffer
 		const uint32 IndexStart = RenderBatch.GetNumVertices();
 
-		const FMargin& Margin = DrawElementPayload.GetBrushMargin();
+		const FMargin& Margin = DrawElement.GetBrushMargin();
 
 		const FVector2f TopRight = FVector2f(BotRight.X, TopLeft.Y);
 		const FVector2f BotLeft = FVector2f(TopLeft.X, BotRight.Y);
 
 		const FColor FeatherColor(0, 0, 0, 0);
 
-		if (DrawElementPayload.GetBrushDrawType() != ESlateBrushDrawType::Image &&
+		if (DrawElement.GetBrushDrawType() != ESlateBrushDrawType::Image &&
 			(Margin.Left != 0.0f || Margin.Top != 0.0f || Margin.Right != 0.0f || Margin.Bottom != 0.0f))
 		{
 			// Create 9 quads for the box element based on the following diagram
@@ -1166,25 +1069,23 @@ void FSlateElementBatcher::AddBoxElements(const FSlateDrawElementArray& DrawElem
 		}
 	};
 	
-	auto GenerateBoxElementBatchParams = [&](const FSlateDrawElement& InDrawElement, FSlateRenderBatchParams& OutBatchParameters)
+	auto GenerateBoxElementBatchParams = [&](const ElementType& InDrawElement, FSlateRenderBatchParams& OutBatchParameters)
 	{
-		const FSlateBoxPayload& DrawElementPayload = InDrawElement.GetDataPayload<FSlateBoxPayload>();
 		const FVector2f LocalSize = InDrawElement.GetLocalSize();
 		const ESlateDrawEffect DrawEffects = InDrawElement.GetDrawEffects();
 
 		// Shader type and params
 		ESlateShader ShaderType = ESlateShader::Default;
-		if (InDrawElement.GetElementType() == EElementType::ET_RoundedBox)
+		if constexpr (std::is_same<ElementType, FSlateRoundedBoxElement>::value)
 		{
 			ShaderType = ESlateShader::RoundedBox;
-			const FSlateRoundedBoxPayload& RoundedPayload = InDrawElement.GetDataPayload<FSlateRoundedBoxPayload>();
-
-			OutBatchParameters.ShaderParams.PixelParams = FVector4f(0, RoundedPayload.GetOutlineWeight(), LocalSize.X, LocalSize.Y);//RadiusWeight;
-			OutBatchParameters.ShaderParams.PixelParams2 = RoundedPayload.GetRadius();
+			
+			OutBatchParameters.ShaderParams.PixelParams = FVector4f(0, InDrawElement.GetOutlineWeight(), LocalSize.X, LocalSize.Y);//RadiusWeight;
+			OutBatchParameters.ShaderParams.PixelParams2 = InDrawElement.GetRadius();
 		}
 
 		// Shader Resource
-		const FSlateShaderResourceProxy* ResourceProxy = DrawElementPayload.GetResourceProxy();
+		const FSlateShaderResourceProxy* ResourceProxy = InDrawElement.GetResourceProxy();
 		FSlateShaderResource* Resource = nullptr;
 		if (ResourceProxy)
 		{
@@ -1193,11 +1094,11 @@ void FSlateElementBatcher::AddBoxElements(const FSlateDrawElementArray& DrawElem
 		}
 
 		// Draw Flags
-		const ESlateBrushTileType::Type TilingRule = DrawElementPayload.GetBrushTiling();
+		const ESlateBrushTileType::Type TilingRule = InDrawElement.GetBrushTiling();
 		const bool bTileHorizontal = (TilingRule == ESlateBrushTileType::Both || TilingRule == ESlateBrushTileType::Horizontal);
 		const bool bTileVertical = (TilingRule == ESlateBrushTileType::Both || TilingRule == ESlateBrushTileType::Vertical);
 
-		const ESlateBrushMirrorType::Type MirroringRule = DrawElementPayload.GetBrushMirroring();
+		const ESlateBrushMirrorType::Type MirroringRule = InDrawElement.GetBrushMirroring();
 		const bool bMirrorHorizontal = (MirroringRule == ESlateBrushMirrorType::Both || MirroringRule == ESlateBrushMirrorType::Horizontal);
 		const bool bMirrorVertical = (MirroringRule == ESlateBrushMirrorType::Both || MirroringRule == ESlateBrushMirrorType::Vertical);
 
@@ -1223,12 +1124,11 @@ void FSlateElementBatcher::AddBoxElements(const FSlateDrawElementArray& DrawElem
 		// For source of magic numbers, see internal method above
 		for (uint32 Index = InBatchStart; Index < InBatchEnd; Index++)
 		{
-			const FSlateDrawElement& DrawElement = DrawElements[Index];
+			const ElementType& DrawElement = DrawElements[Index];
 			const ESlateVertexRounding Rounding = DrawElement.IsPixelSnapped() ? ESlateVertexRounding::Enabled : ESlateVertexRounding::Disabled;
-			const FSlateBoxPayload& DrawElementPayload = DrawElement.GetDataPayload<FSlateBoxPayload>();
-			const FMargin& Margin = DrawElementPayload.GetBrushMargin();
+			const FMargin& Margin = DrawElement.GetBrushMargin();
 
-			if (DrawElementPayload.GetBrushDrawType() != ESlateBrushDrawType::Image &&
+			if (DrawElement.GetBrushDrawType() != ESlateBrushDrawType::Image &&
 				(Margin.Left != 0.0f || Margin.Top != 0.0f || Margin.Right != 0.0f || Margin.Bottom != 0.0f))
 			{
 				NumVertexes += 8 + 8;
@@ -1265,7 +1165,7 @@ namespace SlateElementBatcher
 #if SLATE_CHECK_UOBJECT_RENDER_RESOURCES
 	const FName MaterialInterfaceClassName = "MaterialInterface";
 
-	void CheckUObject(const FSlateTextPayload& InDrawElementPayload, const UObject* InFontMaterial)
+	void CheckUObject(const FSlateTextElement& InDrawElement, const UObject* InFontMaterial)
 	{
 		if (InFontMaterial && GSlateCheckUObjectRenderResources)
 		{
@@ -1273,7 +1173,7 @@ namespace SlateElementBatcher
 			if (!bIsValidLowLevel || !IsValid(InFontMaterial) || InFontMaterial->GetClass()->GetFName() == MaterialInterfaceClassName)
 			{
 				UE_LOG(LogSlate, Error, TEXT("We are rendering a string with an invalid font. The string is: '%s'")
-					, InDrawElementPayload.GetText());
+					, InDrawElement.GetText());
 				// We expect to log more info in the SlateMaterialResource.
 				//In case we crash before that, we also log some info here.
 				UE_LOG(LogSlate, Error, TEXT("Material is not valid. PendingKill:'%d'. ValidLowLevelFast:'%d'. InvalidClass:'%d'")
@@ -1287,14 +1187,13 @@ namespace SlateElementBatcher
 }
 
 template<ESlateVertexRounding Rounding>
-void FSlateElementBatcher::AddTextElement(const FSlateDrawElement& DrawElement)
+void FSlateElementBatcher::AddTextElement(const FSlateTextElement& DrawElement)
 {
-	const FSlateTextPayload& DrawElementPayload = DrawElement.GetDataPayload<FSlateTextPayload>();
-	FColor BaseTint = PackVertexColor(DrawElementPayload.GetTint());
+	FColor BaseTint = PackVertexColor(DrawElement.GetTint());
 
-	const FFontOutlineSettings& OutlineSettings = DrawElementPayload.GetFontInfo().OutlineSettings;
+	const FFontOutlineSettings& OutlineSettings = DrawElement.GetFontInfo().OutlineSettings;
 
-	int32 Len = DrawElementPayload.GetTextLength();
+	int32 Len = DrawElement.GetTextLength();
 	ensure(Len > 0);
 
 	ESlateDrawEffect InDrawEffects = DrawElement.GetDrawEffects();
@@ -1314,12 +1213,12 @@ void FSlateElementBatcher::AddTextElement(const FSlateDrawElement& DrawElement)
 	FSlateFontCache& FontCache = *RenderingPolicy->GetFontCache();
 	FSlateShaderResourceManager& ResourceManager = *RenderingPolicy->GetResourceManager();
 
-	const UObject* BaseFontMaterial = DrawElementPayload.GetFontInfo().FontMaterial;
+	const UObject* BaseFontMaterial = DrawElement.GetFontInfo().FontMaterial;
 	const UObject* OutlineFontMaterial = OutlineSettings.OutlineMaterial;
 
 #if SLATE_CHECK_UOBJECT_RENDER_RESOURCES
-	SlateElementBatcher::CheckUObject(DrawElementPayload, BaseFontMaterial);
-	SlateElementBatcher::CheckUObject(DrawElementPayload, OutlineFontMaterial);
+	SlateElementBatcher::CheckUObject(DrawElement, BaseFontMaterial);
+	SlateElementBatcher::CheckUObject(DrawElement, OutlineFontMaterial);
 #endif
 
 	bool bOutlineFont = OutlineSettings.OutlineSize > 0;
@@ -1327,7 +1226,7 @@ void FSlateElementBatcher::AddTextElement(const FSlateDrawElement& DrawElement)
 
 	auto BuildFontGeometry = [&](const FFontOutlineSettings& InOutlineSettings, const FColor& InTint, const UObject* FontMaterial, int32 InLayer, float InOutlineHorizontalOffset)
 	{
-		FCharacterList& CharacterList = FontCache.GetCharacterList(DrawElementPayload.GetFontInfo(), FontScale, InOutlineSettings);
+		FCharacterList& CharacterList = FontCache.GetCharacterList(DrawElement.GetFontInfo(), FontScale, InOutlineSettings);
 
 		float MaxHeight = CharacterList.GetMaxHeight();
 
@@ -1377,7 +1276,7 @@ void FSlateElementBatcher::AddTextElement(const FSlateDrawElement& DrawElement)
 		uint32 NumLines = 1;
 		for( uint32 CharIndex = 0; CharIndex < NumChars; ++CharIndex )
 		{
-			const TCHAR CurrentChar = DrawElementPayload.GetText()[ CharIndex ];
+			const TCHAR CurrentChar = DrawElement.GetText()[ CharIndex ];
 
 			ensure(CurrentChar != '\0');
 
@@ -1395,7 +1294,7 @@ void FSlateElementBatcher::AddTextElement(const FSlateDrawElement& DrawElement)
 			}
 			else
 			{
-				const FCharacterEntry& Entry = CharacterList.GetCharacter(CurrentChar, DrawElementPayload.GetFontInfo().FontFallback);
+				const FCharacterEntry& Entry = CharacterList.GetCharacter(CurrentChar, DrawElement.GetFontInfo().FontFallback);
 
 				if( Entry.Valid && (FontAtlasTexture == nullptr || Entry.TextureIndex != FontTextureIndex) )
 				{
@@ -1408,7 +1307,7 @@ void FSlateElementBatcher::AddTextElement(const FSlateDrawElement& DrawElement)
 					FontAtlasTexture = SlateFontTexture->GetSlateTexture();
 					check(FontAtlasTexture);
 
-					FontShaderResource = ResourceManager.GetFontShaderResource( FontTextureIndex, FontAtlasTexture, DrawElementPayload.GetFontInfo().FontMaterial );
+					FontShaderResource = ResourceManager.GetFontShaderResource( FontTextureIndex, FontAtlasTexture, DrawElement.GetFontInfo().FontMaterial );
 					check(FontShaderResource);
 
 					const bool bIsGrayscale = SlateFontTexture->IsGrayscale();
@@ -1586,12 +1485,11 @@ static bool IsAxisAlignedRotation(const FMatrix2x2& Matrix)
 }
 
 template<ESlateVertexRounding Rounding>
-void FSlateElementBatcher::AddShapedTextElement( const FSlateDrawElement& DrawElement )
+void FSlateElementBatcher::AddShapedTextElement( const FSlateShapedTextElement& DrawElement )
 {
-	const FSlateShapedTextPayload& DrawElementPayload = DrawElement.GetDataPayload<FSlateShapedTextPayload>();
-	const FShapedGlyphSequence* ShapedGlyphSequence = DrawElementPayload.GetShapedGlyphSequence().Get();
+	const FShapedGlyphSequence* ShapedGlyphSequence = DrawElement.GetShapedGlyphSequence().Get();
 
-	const FShapedGlyphSequence* OverflowGlyphSequence = DrawElementPayload.OverflowArgs.OverflowTextPtr.Get();
+	const FShapedGlyphSequence* OverflowGlyphSequence = DrawElement.OverflowArgs.OverflowTextPtr.Get();
 
 	checkSlow(ShapedGlyphSequence);
 
@@ -1599,7 +1497,7 @@ void FSlateElementBatcher::AddShapedTextElement( const FSlateDrawElement& DrawEl
 
 	ensure(ShapedGlyphSequence->GetGlyphsToRender().Num() > 0);
 
-	const FColor BaseTint = PackVertexColor(DrawElementPayload.GetTint());
+	const FColor BaseTint = PackVertexColor(DrawElement.GetTint());
 
 	FSlateFontCache& FontCache = *RenderingPolicy->GetFontCache();
 
@@ -1662,8 +1560,8 @@ void FSlateElementBatcher::AddShapedTextElement( const FSlateDrawElement& DrawEl
 		// Optimize by culling
 		// Todo: this doesn't work with cached clipping
 		BuildContext.bEnableCulling = false;
-		BuildContext.bForceEllipsis = DrawElementPayload.OverflowArgs.bIsLastVisibleBlock && DrawElementPayload.OverflowArgs.bIsNextBlockClipped;
-		BuildContext.OverflowDirection = DrawElementPayload.OverflowArgs.OverflowDirection;
+		BuildContext.bForceEllipsis = DrawElement.OverflowArgs.bIsLastVisibleBlock && DrawElement.OverflowArgs.bIsNextBlockClipped;
+		BuildContext.OverflowDirection = DrawElement.OverflowArgs.OverflowDirection;
 
 		if (ShapedGlyphSequence->GetGlyphsToRender().Num() > 200 || (OverflowGlyphSequence && BuildContext.OverflowDirection != ETextOverflowDirection::NoOverflow))
 		{
@@ -1683,7 +1581,7 @@ void FSlateElementBatcher::AddShapedTextElement( const FSlateDrawElement& DrawEl
 				// In checks below, ignore floating-point differences caused by transforming and untransforming the clip rect
 				const bool NeedLeftEllipsis = FMath::FloorToInt(BuildContext.LocalClipBoundingBoxLeft) > 0 && BuildContext.OverflowDirection == ETextOverflowDirection::RightToLeft;
 				const bool NeedRightEllipsis = ShapedGlyphSequence->GetMeasuredWidth() > FMath::CeilToInt(BuildContext.LocalClipBoundingBoxRight) && BuildContext.OverflowDirection == ETextOverflowDirection::LeftToRight;
-				if (!NeedLeftEllipsis && !NeedRightEllipsis && !DrawElementPayload.OverflowArgs.bIsNextBlockClipped)
+				if (!NeedLeftEllipsis && !NeedRightEllipsis && !DrawElement.OverflowArgs.bIsNextBlockClipped)
 				{
 					BuildContext.OverflowDirection = ETextOverflowDirection::NoOverflow;
 				}
@@ -1705,7 +1603,7 @@ void FSlateElementBatcher::AddShapedTextElement( const FSlateDrawElement& DrawEl
 	if (bOutlineFont)
 	{
 		// Build geometry for the outline
-		BuildFontGeometry(OutlineSettings, PackVertexColor(DrawElementPayload.GetOutlineTint()), OutlineFontMaterial, Layer, 0.f);
+		BuildFontGeometry(OutlineSettings, PackVertexColor(DrawElement.GetOutlineTint()), OutlineFontMaterial, Layer, 0.f);
 		
 		//The fill area was measured without an outline so it must be shifted by the scaled outline size
 		const float HorizontalOffset = FMath::RoundToFloat((float)OutlineSize * FontScale);
@@ -1721,26 +1619,25 @@ void FSlateElementBatcher::AddShapedTextElement( const FSlateDrawElement& DrawEl
 }
 
 template<ESlateVertexRounding Rounding>
-void FSlateElementBatcher::AddGradientElement( const FSlateDrawElement& DrawElement )
+void FSlateElementBatcher::AddGradientElement( const FSlateGradientElement& DrawElement )
 {
 	const FSlateRenderTransform& RenderTransform = DrawElement.GetRenderTransform();
 	const FVector2f LocalSize = DrawElement.GetLocalSize();
-	const FSlateGradientPayload& InPayload = DrawElement.GetDataPayload<FSlateGradientPayload>();
 	const ESlateDrawEffect InDrawEffects = DrawElement.GetDrawEffects();
 	const int32 Layer = DrawElement.GetLayer();
 	const float DrawScale = DrawElement.GetScale();
 
 	// There must be at least one gradient stop
-	check( InPayload.GradientStops.Num() > 0 );
+	check(DrawElement.GradientStops.Num() > 0 );
 
 	FShaderParams ShaderParams;
 
 	ESlateShader ShaderType = ESlateShader::Default;
-	if (InPayload.CornerRadius != FVector4f(0.f))
+	if (DrawElement.CornerRadius != FVector4f(0.f))
 	{
 		ShaderType = ESlateShader::RoundedBox;
 		ShaderParams.PixelParams = FVector4f(0.0f, 0.0f, LocalSize.X, LocalSize.Y);
-		ShaderParams.PixelParams2 = InPayload.CornerRadius;
+		ShaderParams.PixelParams2 = DrawElement.CornerRadius;
 	}
 
 	FSlateRenderBatch& RenderBatch = 
@@ -1761,14 +1658,14 @@ void FSlateElementBatcher::AddGradientElement( const FSlateDrawElement& DrawElem
 	FVector2f BotRight = FVector2f(LocalSize.X, LocalSize.Y);
 
 	// Copy the gradient stops.. We may need to add more
-	TArray<FSlateGradientStop> GradientStops = InPayload.GradientStops;
+	TArray<FSlateGradientStop> GradientStops = DrawElement.GradientStops;
 
-	const FSlateGradientStop& FirstStop = InPayload.GradientStops[0];
-	const FSlateGradientStop& LastStop = InPayload.GradientStops[ InPayload.GradientStops.Num() - 1 ];
+	const FSlateGradientStop& FirstStop = DrawElement.GradientStops[0];
+	const FSlateGradientStop& LastStop = DrawElement.GradientStops[DrawElement.GradientStops.Num() - 1 ];
 		
 	// Determine if the first and last stops are not at the start and end of the quad
 	// If they are not add a gradient stop with the same color as the first and/or last stop
-	if( InPayload.GradientType == Orient_Vertical )
+	if( DrawElement.GradientType == Orient_Vertical )
 	{
 		if( 0.0f < FirstStop.Position.X )
 		{
@@ -1813,7 +1710,7 @@ void FSlateElementBatcher::AddGradientElement( const FSlateDrawElement& DrawElem
 		FVector2f StartUV;
 		FVector2f EndUV;
 
-		if( InPayload.GradientType == Orient_Vertical )
+		if( DrawElement.GradientType == Orient_Vertical )
 		{
 			// Gradient stop is vertical so gradients to left to right
 			StartPt = TopLeft;
@@ -2128,7 +2025,7 @@ private:
 };
 
 
-void FSlateElementBatcher::AddSplineElement(const FSlateDrawElement& DrawElement)
+void FSlateElementBatcher::AddSplineElement(const FSlateSplineElement& DrawElement)
 {
 	// WHY NO PIXEL SNAPPING?
 	//
@@ -2139,7 +2036,6 @@ void FSlateElementBatcher::AddSplineElement(const FSlateDrawElement& DrawElement
 	// a floating point line width.
 
 	const FSlateRenderTransform& RenderTransform = DrawElement.GetRenderTransform();
-	const FSlateSplinePayload& InPayload = DrawElement.GetDataPayload<FSlateSplinePayload>();
 	const ESlateDrawEffect InDrawEffects = DrawElement.GetDrawEffects();
 	const int32 Layer = DrawElement.GetLayer();
 
@@ -2147,7 +2043,7 @@ void FSlateElementBatcher::AddSplineElement(const FSlateDrawElement& DrawElement
 	// Increasing this value will increase the fuzziness of line edges.
 	const float FilterRadius = 1.0f;
 
-	const float HalfThickness = 0.5f * FMath::Max(0.0f, InPayload.GetThickness());
+	const float HalfThickness = 0.5f * FMath::Max(0.0f, DrawElement.GetThickness());
 
 	// The amount we increase each side of the line to generate enough pixels
 	const float FilteredHalfThickness = HalfThickness + FilterRadius;
@@ -2159,25 +2055,25 @@ void FSlateElementBatcher::AddSplineElement(const FSlateDrawElement& DrawElement
 	// Thickness is given in screenspace, convert it to local space for geometry build
 	const float LocalHalfThickness = FilteredHalfThickness / DrawElement.GetScale();
 
-	const FColor SplineColor = (InPayload.GradientStops.Num()==1) ? PackVertexColor(InPayload.GradientStops[0].Color) : PackVertexColor(InPayload.GetTint());
+	const FColor SplineColor = (DrawElement.GradientStops.Num()==1) ? PackVertexColor(DrawElement.GradientStops[0].Color) : PackVertexColor(DrawElement.GetTint());
 
 	FSplineBuilder SplineBuilder(
 		RenderBatch,
-		InPayload.P0,
+		DrawElement.P0,
 		LocalHalfThickness,
 		RenderTransform,
 		SplineColor
 	);
 
-	if (/*const bool bNoGradient = */InPayload.GradientStops.Num() <= 1)
+	if (/*const bool bNoGradient = */DrawElement.GradientStops.Num() <= 1)
 	{
 		// Normal scenario where there is no color gradient.
-		SplineBuilder.BuildBezierGeometry(InPayload.P0, InPayload.P1, InPayload.P2, InPayload.P3);
+		SplineBuilder.BuildBezierGeometry(DrawElement.P0, DrawElement.P1, DrawElement.P2, DrawElement.P3);
 	}
 	else
 	{
 		// Deprecated scenario _WithColorGradient
-		SplineBuilder.BuildBezierGeometry_WithColorGradient( InPayload.GradientStops, 1, InPayload.P0, InPayload.P1, InPayload.P2, InPayload.P3, *this);
+		SplineBuilder.BuildBezierGeometry_WithColorGradient(DrawElement.GradientStops, 1, DrawElement.P0, DrawElement.P1, DrawElement.P2, DrawElement.P3, *this);
 	}
 }
 
@@ -2472,7 +2368,7 @@ private:
 
 		if (SizeSquared > SMALL_NUMBER)
 		{
-			return DirSum * FMath::InvSqrtEst(SizeSquared);
+			return DirSum * FMath::InvSqrt(SizeSquared);
 		}
 
 		// Inputs are antiparallel, so any perpendicular vector suffices
@@ -2489,18 +2385,17 @@ private:
 };
 
 
-void FSlateElementBatcher::AddLineElements( const FSlateDrawElementArray& DrawElements )
+void FSlateElementBatcher::AddLineElements( const FSlateDrawElementArray<FSlateLineElement>& DrawElements )
 {
-	auto AddLineElementInternal = [&](FSlateRenderBatch& RenderBatch, const FSlateDrawElement& DrawElement)
+	auto AddLineElementInternal = [&](FSlateRenderBatch& RenderBatch, const FSlateLineElement& DrawElement)
 	{
 		const ESlateVertexRounding Rounding = DrawElement.IsPixelSnapped() ? ESlateVertexRounding::Enabled : ESlateVertexRounding::Disabled;
-		const FSlateLinePayload& DrawElementPayload = DrawElement.GetDataPayload<FSlateLinePayload>();
 		const FSlateRenderTransform& RenderTransform = DrawElement.GetRenderTransform();
 		const ESlateDrawEffect DrawEffects = DrawElement.GetDrawEffects();
 		const int32 Layer = DrawElement.GetLayer();
 
-		const TArray<FVector2f>& Points = DrawElementPayload.GetPoints();
-		const TArray<FLinearColor>& PointColors = DrawElementPayload.GetPointColors();
+		const TArray<FVector2f>& Points = DrawElement.GetPoints();
+		const TArray<FLinearColor>& PointColors = DrawElement.GetPointColors();
 
 		const int32 NumPoints = Points.Num();
 		if (NumPoints < 2)
@@ -2509,15 +2404,15 @@ void FSlateElementBatcher::AddLineElements( const FSlateDrawElementArray& DrawEl
 		}
 
 		// Pack element tint and per-point colors
-		const FColor PackedTint = PackVertexColor(DrawElementPayload.GetTint());
+		const FColor PackedTint = PackVertexColor(DrawElement.GetTint());
 		TArray<FColor> PackedColors;
 		PackedColors.Reserve(PointColors.Num());
 		for (const FLinearColor& PointColor : PointColors)
 		{
-			PackedColors.Add(PackVertexColor(PointColor * DrawElementPayload.GetTint()));
+			PackedColors.Add(PackVertexColor(PointColor * DrawElement.GetTint()));
 		}
 
-		if (DrawElementPayload.IsAntialiased())
+		if (DrawElement.IsAntialiased())
 		{
 			// Filter size to use for anti-aliasing.
 			// Increasing this value will increase the fuzziness of line edges.
@@ -2527,7 +2422,7 @@ void FSlateElementBatcher::AddLineElements( const FSlateDrawElementArray& DrawEl
 			// A split can also be forced by repeating points in the input
 			const float MiterAngleLimit = 90.0f - KINDA_SMALL_NUMBER;
 
-			const float HalfThickness = 0.5f * FMath::Max(0.0f, DrawElementPayload.GetThickness());
+			const float HalfThickness = 0.5f * FMath::Max(0.0f, DrawElement.GetThickness());
 
 			FLineBuilder LineBuilder(
 				RenderBatch,
@@ -2541,7 +2436,7 @@ void FSlateElementBatcher::AddLineElements( const FSlateDrawElementArray& DrawEl
 		}
 		else
 		{
-			if (DrawElementPayload.GetThickness() == 1.0f)
+			if (DrawElement.GetThickness() == 1.0f)
 			{
 				const FColor StartColor = PackedColors.Num() ? PackedColors[0] : PackedTint;
 				RenderBatch.EmplaceVertex(FSlateVertex::Make(RenderTransform, Points[0], FVector2f::ZeroVector, StartColor, {}, Rounding));
@@ -2561,7 +2456,7 @@ void FSlateElementBatcher::AddLineElements( const FSlateDrawElementArray& DrawEl
 			{
 
 				// Thickness is given in screenspace, convert it to local space for geometry build
-				const float LocalHalfThickness = (DrawElementPayload.GetThickness() * 0.5f) / DrawElement.GetScale();
+				const float LocalHalfThickness = (DrawElement.GetThickness() * 0.5f) / DrawElement.GetScale();
 
 				for (int32 Point = 0; Point < NumPoints - 1; ++Point)
 				{
@@ -2584,21 +2479,20 @@ void FSlateElementBatcher::AddLineElements( const FSlateDrawElementArray& DrawEl
 			}
 		}
 	};
-
-	auto GenerateLineElementBatchParams = [&](const FSlateDrawElement& InDrawElement, FSlateRenderBatchParams& OutBatchParameters)
+	
+	auto GenerateLineElementBatchParams = [&](const FSlateLineElement& InDrawElement, FSlateRenderBatchParams& OutBatchParameters)
 	{
-		const FSlateLinePayload& DrawElementPayload = InDrawElement.GetDataPayload<FSlateLinePayload>();
 		const ESlateDrawEffect DrawEffects = InDrawElement.GetDrawEffects();
 
 		// Shader type and params
 		ESlateShader ShaderType = ESlateShader::Default;
-		if (DrawElementPayload.IsAntialiased())
+		if (InDrawElement.IsAntialiased())
 		{
 			// Filter size to use for anti-aliasing.
 			// Increasing this value will increase the fuzziness of line edges.
 			const float FilterRadius = 1.0f;
 
-			const float HalfThickness = 0.5f * FMath::Max(0.0f, DrawElementPayload.GetThickness());
+			const float HalfThickness = 0.5f * FMath::Max(0.0f, InDrawElement.GetThickness());
 
 			// The amount we increase each side of the line to generate enough pixels
 			const float FilteredHalfThickness = HalfThickness + FilterRadius;
@@ -2612,13 +2506,13 @@ void FSlateElementBatcher::AddLineElements( const FSlateDrawElementArray& DrawEl
 		// Primitive type
 		ESlateDrawPrimitive PrimitiveType = ESlateDrawPrimitive::LineList;
 
-		if (DrawElementPayload.IsAntialiased())
+		if (InDrawElement.IsAntialiased())
 		{
 			PrimitiveType = ESlateDrawPrimitive::TriangleList;
 		}
 		else
 		{
-			if (DrawElementPayload.GetThickness() == 1.0f)
+			if (InDrawElement.GetThickness() == 1.0f)
 			{
 				PrimitiveType = ESlateDrawPrimitive::LineList;
 			}
@@ -2646,15 +2540,13 @@ void FSlateElementBatcher::AddLineElements( const FSlateDrawElementArray& DrawEl
 		// For source of magic numbers, see internal method above
 		for (uint32 Index = InBatchStart; Index < InBatchEnd; Index++)
 		{
-
-			const FSlateDrawElement& DrawElement = DrawElements[Index];
+			const FSlateLineElement& DrawElement = DrawElements[Index];
 			const ESlateVertexRounding Rounding = DrawElement.IsPixelSnapped() ? ESlateVertexRounding::Enabled : ESlateVertexRounding::Disabled;
-			const FSlateLinePayload& DrawElementPayload = DrawElement.GetDataPayload<FSlateLinePayload>();
 			const FSlateRenderTransform& RenderTransform = DrawElement.GetRenderTransform();
 			const ESlateDrawEffect DrawEffects = DrawElement.GetDrawEffects();
 			const int32 Layer = DrawElement.GetLayer();
 
-			const TArray<FVector2f>& Points = DrawElementPayload.GetPoints();
+			const TArray<FVector2f>& Points = DrawElement.GetPoints();
 
 			const int32 NumPoints = Points.Num();
 			if (NumPoints < 2)
@@ -2662,13 +2554,13 @@ void FSlateElementBatcher::AddLineElements( const FSlateDrawElementArray& DrawEl
 				return;
 			}
 
-			if (DrawElementPayload.IsAntialiased())
+			if (DrawElement.IsAntialiased())
 			{
 				const float FilterRadius = 1.0f;
 
 				const float MiterAngleLimit = 90.0f - KINDA_SMALL_NUMBER;
 
-				const float HalfThickness = 0.5f * FMath::Max(0.0f, DrawElementPayload.GetThickness());
+				const float HalfThickness = 0.5f * FMath::Max(0.0f, DrawElement.GetThickness());
 
 				FLineBuilder LineBuilder(
 					RenderBatch,
@@ -2682,7 +2574,7 @@ void FSlateElementBatcher::AddLineElements( const FSlateDrawElementArray& DrawEl
 			}
 			else
 			{
-				if (DrawElementPayload.GetThickness() == 1.0f)
+				if (DrawElement.GetThickness() == 1.0f)
 				{
 					NumVertexes += NumPoints;
 					NumIndices += (NumPoints - 1) * 2;
@@ -2704,22 +2596,21 @@ void FSlateElementBatcher::AddLineElements( const FSlateDrawElementArray& DrawEl
 }
 
 template<ESlateVertexRounding Rounding>
-void FSlateElementBatcher::AddViewportElement( const FSlateDrawElement& DrawElement )
+void FSlateElementBatcher::AddViewportElement( const FSlateViewportElement& DrawElement )
 {
 	const FSlateRenderTransform& RenderTransform = DrawElement.GetRenderTransform();
 	const FVector2f LocalSize = DrawElement.GetLocalSize();
-	const FSlateViewportPayload& DrawElementPayload = DrawElement.GetDataPayload<FSlateViewportPayload>();
 	ESlateDrawEffect InDrawEffects = DrawElement.GetDrawEffects();
 	const int32 Layer = DrawElement.GetLayer();
 
-	const FColor FinalColor = PackVertexColor(DrawElementPayload.GetTint());
+	const FColor FinalColor = PackVertexColor(DrawElement.GetTint());
 
 	ESlateBatchDrawFlag DrawFlags = DrawElement.GetBatchFlags();
 
-	FSlateShaderResource* ViewportResource = DrawElementPayload.RenderTargetResource;
+	FSlateShaderResource* ViewportResource = DrawElement.RenderTargetResource;
 	ESlateShader ShaderType = ESlateShader::Default;
 
-	if(DrawElementPayload.bViewportTextureAlphaOnly )
+	if(DrawElement.bViewportTextureAlphaOnly )
 	{
 		// This is a slight hack, but the grayscale font shader is the same as the general shader except it reads alpha only textures and doesn't support tiling
 		ShaderType = ESlateShader::GrayscaleFont;
@@ -2728,9 +2619,9 @@ void FSlateElementBatcher::AddViewportElement( const FSlateDrawElement& DrawElem
 	FSlateRenderBatch& RenderBatch = CreateRenderBatch(UsedBatchData, Layer, FShaderParams(), ViewportResource, ESlateDrawPrimitive::TriangleList, ShaderType, InDrawEffects, DrawFlags, DrawElement);
 
 	// Tag this batch as requiring vsync if the viewport requires it.
-	if( ViewportResource != nullptr && !DrawElementPayload.bAllowViewportScaling )
+	if( ViewportResource != nullptr && !DrawElement.bAllowViewportScaling )
 	{
-		bRequiresVsync |= DrawElementPayload.bRequiresVSync;
+		bRequiresVsync |= DrawElement.bRequiresVSync;
 	}
 
 	// Do pixel snapping
@@ -2738,7 +2629,7 @@ void FSlateElementBatcher::AddViewportElement( const FSlateDrawElement& DrawElem
 	FVector2f BotRight(LocalSize);
 
 	// If the viewport disallows scaling, force size to current texture size.
-	if (ViewportResource != nullptr && !DrawElementPayload.bAllowViewportScaling)
+	if (ViewportResource != nullptr && !DrawElement.bAllowViewportScaling)
 	{
 		const float ElementScale = DrawElement.GetScale();
 		BotRight = FVector2f(ViewportResource->GetWidth() / ElementScale, ViewportResource->GetHeight() / ElementScale);
@@ -2792,9 +2683,8 @@ void FSlateElementBatcher::AddViewportElement( const FSlateDrawElement& DrawElem
 }
 
 template<ESlateVertexRounding Rounding>
-void FSlateElementBatcher::AddBorderElement( const FSlateDrawElement& DrawElement )
+void FSlateElementBatcher::AddBorderElement( const FSlateBoxElement& DrawElement )
 {
-	const FSlateBoxPayload& DrawElementPayload = DrawElement.GetDataPayload<FSlateBoxPayload>();
 	const FSlateRenderTransform& RenderTransform = DrawElement.GetRenderTransform();
 	const FVector2f LocalSize = DrawElement.GetLocalSize();
 	const ESlateDrawEffect InDrawEffects = DrawElement.GetDrawEffects();
@@ -2807,7 +2697,7 @@ void FSlateElementBatcher::AddBorderElement( const FSlateDrawElement& DrawElemen
 	uint32 TextureHeight = 1;
 
 	// Currently borders are not atlased because they are tiled.  So we just assume the texture proxy holds the actual texture
-	const FSlateShaderResourceProxy* ResourceProxy = DrawElementPayload.GetResourceProxy();
+	const FSlateShaderResourceProxy* ResourceProxy = DrawElement.GetResourceProxy();
 	FSlateShaderResource* Resource = ResourceProxy ? ResourceProxy->Resource : nullptr;
 	if( Resource )
 	{
@@ -2822,7 +2712,7 @@ void FSlateElementBatcher::AddBorderElement( const FSlateDrawElement& DrawElemen
 	const FVector2f StartUV = HalfTexel;
 	const FVector2f EndUV = FVector2f( 1.0f, 1.0f ) + HalfTexel;
 
-	FMargin Margin = DrawElementPayload.GetBrushMargin();
+	FMargin Margin = DrawElement.GetBrushMargin();
 
 	// Do pixel snapping
 	FVector2f TopLeft(0.f, 0.f);
@@ -2887,7 +2777,7 @@ void FSlateElementBatcher::AddBorderElement( const FSlateDrawElement& DrawElemen
 	FShaderParams ShaderParams = FShaderParams::MakePixelShaderParams(FVector4f(LeftMarginU,RightMarginU,TopMarginV,BottomMarginV) );
 
 	// The tint color applies to all brushes and is passed per vertex
-	const FColor Tint = PackVertexColor(DrawElementPayload.GetTint());
+	const FColor Tint = PackVertexColor(DrawElement.GetTint());
 
 	// Pass the tiling information as a flag so we can pick the correct texture addressing mode
 	ESlateBatchDrawFlag DrawFlags = (ESlateBatchDrawFlag::TileU|ESlateBatchDrawFlag::TileV);
@@ -3022,26 +2912,25 @@ void FSlateElementBatcher::AddBorderElement( const FSlateDrawElement& DrawElemen
 	RenderBatch.IndexOffset = 0;
 }
 
-void FSlateElementBatcher::AddCustomElement( const FSlateDrawElement& DrawElement )
+void FSlateElementBatcher::AddCustomElement( const FSlateCustomDrawerElement& DrawElement )
 {
 	const int32 Layer = DrawElement.GetLayer();
 
 	FSlateRenderBatch& RenderBatch = CreateRenderBatch(Layer, FShaderParams(), nullptr, ESlateDrawPrimitive::None, ESlateShader::Default, ESlateDrawEffect::None, ESlateBatchDrawFlag::None, DrawElement);
-	RenderBatch.CustomDrawer = DrawElement.GetDataPayload<FSlateCustomDrawerPayload>().CustomDrawer.Pin().Get();
+	RenderBatch.CustomDrawer = DrawElement.CustomDrawer.Pin().Get();
 	RenderBatch.bIsMergable = false;
 }
 
-void FSlateElementBatcher::AddCustomVerts(const FSlateDrawElement& DrawElement)
+void FSlateElementBatcher::AddCustomVerts(const FSlateCustomVertsElement& DrawElement)
 {
-	const FSlateCustomVertsPayload& InPayload = DrawElement.GetDataPayload<FSlateCustomVertsPayload>();
 	const int32 Layer = DrawElement.GetLayer();
 
-	if (InPayload.Vertices.Num() > 0)
+	if (DrawElement.Vertices.Num() > 0)
 	{
 		FSlateRenderBatch& RenderBatch = CreateRenderBatch(
 			Layer, 
 			FShaderParams(), 
-			InPayload.ResourceProxy != nullptr ? InPayload.ResourceProxy->Resource : nullptr, 
+			DrawElement.ResourceProxy != nullptr ? DrawElement.ResourceProxy->Resource : nullptr,
 			ESlateDrawPrimitive::TriangleList,
 			ESlateShader::Custom, 
 			DrawElement.GetDrawEffects(), 
@@ -3049,12 +2938,12 @@ void FSlateElementBatcher::AddCustomVerts(const FSlateDrawElement& DrawElement)
 			DrawElement);
 
 		RenderBatch.bIsMergable = false;
-		RenderBatch.InstanceCount = InPayload.NumInstances;
-		RenderBatch.InstanceOffset = InPayload.InstanceOffset;
-		RenderBatch.InstanceData = InPayload.InstanceData;
+		RenderBatch.InstanceCount = DrawElement.NumInstances;
+		RenderBatch.InstanceOffset = DrawElement.InstanceOffset;
+		RenderBatch.InstanceData = DrawElement.InstanceData;
 
-		RenderBatch.AddVertices(InPayload.Vertices);
-		RenderBatch.AddIndices(InPayload.Indices);
+		RenderBatch.AddVertices(DrawElement.Vertices);
+		RenderBatch.AddIndices(DrawElement.Indices);
 
 	}
 	/*FElementBatchMap& LayerToElementBatches = CurrentDrawLayer->GetElementBatchMap();
@@ -3103,14 +2992,12 @@ void FSlateElementBatcher::AddCustomVerts(const FSlateDrawElement& DrawElement)
 	}*/
 }
 
-void FSlateElementBatcher::AddPostProcessPass(const FSlateDrawElement& DrawElement, FVector2f WindowSize)
+void FSlateElementBatcher::AddPostProcessPass(const FSlatePostProcessElement& DrawElement, FVector2f WindowSize)
 {
 	++NumPostProcessPasses;
 
 	const FSlateRenderTransform& RenderTransform = DrawElement.GetRenderTransform();
 	const FVector2f LocalSize = DrawElement.GetLocalSize();
-	
-	const FSlatePostProcessPayload& Payload = DrawElement.GetDataPayload<FSlatePostProcessPayload>();
 
 	//@todo doesn't work with rotated or skewed objects yet
 	const FVector2f Position = DrawElement.GetPosition();
@@ -3137,8 +3024,8 @@ void FSlateElementBatcher::AddPostProcessPass(const FSlateDrawElement& DrawEleme
 	{
 		FShaderParams Params = FShaderParams::MakePixelShaderParams(
 			FVector4f(WorldTopLeft, WorldBotRight),
-			FVector4f(Payload.PostProcessData.X, Payload.PostProcessData.Y, (float)Payload.DownsampleAmount, 0.f),
-			FVector4f(Payload.CornerRadius));
+			FVector4f(DrawElement.PostProcessData.X, DrawElement.PostProcessData.Y, (float)DrawElement.DownsampleAmount, 0.f),
+			FVector4f(DrawElement.CornerRadius));
 
 		CreateRenderBatch(Layer, Params, nullptr, ESlateDrawPrimitive::TriangleList, ESlateShader::PostProcess, ESlateDrawEffect::None, ESlateBatchDrawFlag::None, DrawElement);
 	}
diff --git a/Engine/Source/Runtime/SlateCore/Private/Rendering/SlateRenderBatch.cpp b/Engine/Source/Runtime/SlateCore/Private/Rendering/SlateRenderBatch.cpp
--- a/Engine/Source/Runtime/SlateCore/Private/Rendering/SlateRenderBatch.cpp
+++ b/Engine/Source/Runtime/SlateCore/Private/Rendering/SlateRenderBatch.cpp
@@ -2,6 +2,7 @@
 
 #include "Rendering/SlateRenderBatch.h"
 #include "Rendering/DrawElements.h"
+#include "Rendering/DrawElementPayloads.h"
 #include "Rendering/ElementBatcher.h"
 #include "Textures/SlateShaderResource.h"
 
diff --git a/Engine/Source/Runtime/SlateCore/Public/Rendering/DrawElementCoreTypes.h b/Engine/Source/Runtime/SlateCore/Public/Rendering/DrawElementCoreTypes.h
new file mode 100644
--- /dev/null
+++ b/Engine/Source/Runtime/SlateCore/Public/Rendering/DrawElementCoreTypes.h
@@ -0,0 +1,200 @@
+// Copyright Epic Games, Inc. All Rights Reserved.
+
+#pragma once
+
+#include "CoreMinimal.h"
+#include "Stats/Stats.h"
+#include "SlateGlobals.h"
+
+/**
+ * Note: For types needed even by RenderingCommon.h
+ */
+
+struct FSlateBoxElement;
+struct FSlateBoxElement; // ET_DebugQuad
+struct FSlateTextElement;
+struct FSlateShapedTextElement;
+struct FSlateSplineElement;
+struct FSlateLineElement;
+struct FSlateGradientElement;
+struct FSlateViewportElement;
+struct FSlateBoxElement; // ET_Border
+struct FSlateCustomDrawerElement;
+struct FSlateCustomVertsElement;
+struct FSlatePostProcessElement;
+struct FSlateRoundedBoxElement;
+class FSlateDrawElement;
+
+DECLARE_MEMORY_STAT_EXTERN(TEXT("Vertex/Index Buffer Pool Memory (CPU)"), STAT_SlateBufferPoolMemory, STATGROUP_SlateMemory, SLATECORE_API);
+DECLARE_MEMORY_STAT_EXTERN(TEXT("Cached Draw Element Memory (CPU)"), STAT_SlateCachedDrawElementMemory, STATGROUP_SlateMemory, SLATECORE_API);
+
+/**
+ * Slate higher level abstract element types
+ */
+enum class EElementType : uint8
+{
+	ET_Box,
+	ET_DebugQuad,
+	ET_Text,
+	ET_ShapedText,
+	ET_Spline,
+	ET_Line,
+	ET_Gradient,
+	ET_Viewport,
+	ET_Border,
+	ET_Custom,
+	ET_CustomVerts,
+	ET_PostProcessPass,
+	ET_RoundedBox,
+	/**
+	 * We map draw elements by type on add for better cache coherency if possible,
+	 * this type is used when that grouping is not possible.
+	 * Grouping is also planned to be used for bulk element type processing.
+	 */
+	ET_NonMapped,
+	/** Total number of draw commands */
+	ET_Count,
+};
+
+#if STATS
+
+struct FRenderingBufferStatTracker
+{
+	static void MemoryAllocated(int32 SizeBytes)
+	{
+		INC_DWORD_STAT_BY(STAT_SlateBufferPoolMemory, SizeBytes);
+	}
+
+	static void MemoryFreed(int32 SizeBytes)
+	{
+		DEC_DWORD_STAT_BY(STAT_SlateBufferPoolMemory, SizeBytes);
+	}
+};
+
+struct FDrawElementStatTracker
+{
+	static void MemoryAllocated(int32 SizeBytes)
+	{
+		INC_DWORD_STAT_BY(STAT_SlateCachedDrawElementMemory, SizeBytes);
+	}
+
+	static void MemoryFreed(int32 SizeBytes)
+	{
+		DEC_DWORD_STAT_BY(STAT_SlateCachedDrawElementMemory, SizeBytes);
+	}
+};
+
+template<typename StatTracker>
+class FSlateStatTrackingMemoryAllocator : public FDefaultAllocator
+{
+public:
+	typedef FDefaultAllocator Super;
+
+	class ForAnyElementType : public FDefaultAllocator::ForAnyElementType
+	{
+	public:
+		typedef FDefaultAllocator::ForAnyElementType Super;
+
+		ForAnyElementType()
+			: AllocatedSize(0)
+		{
+
+		}
+
+		/**
+		* Moves the state of another allocator into this one.
+		* Assumes that the allocator is currently empty, i.e. memory may be allocated but any existing elements have already been destructed (if necessary).
+		* @param Other - The allocator to move the state from.  This allocator should be left in a valid empty state.
+		*/
+		FORCEINLINE void MoveToEmpty(ForAnyElementType& Other)
+		{
+			Super::MoveToEmpty(Other);
+
+			AllocatedSize = Other.AllocatedSize;
+			Other.AllocatedSize = 0;
+		}
+
+		/** Destructor. */
+		~ForAnyElementType()
+		{
+			if (AllocatedSize)
+			{
+				StatTracker::MemoryFreed(AllocatedSize);
+			}
+		}
+
+		void ResizeAllocation(int32 PreviousNumElements, int32 NumElements, int32 NumBytesPerElement)
+		{
+			const int32 NewSize = NumElements * NumBytesPerElement;
+			StatTracker::MemoryAllocated(NewSize - AllocatedSize);
+
+			AllocatedSize = NewSize;
+
+			Super::ResizeAllocation(PreviousNumElements, NumElements, NumBytesPerElement);
+		}
+
+	private:
+		ForAnyElementType(const ForAnyElementType&);
+		ForAnyElementType& operator=(const ForAnyElementType&);
+	private:
+		int32 AllocatedSize;
+	};
+};
+
+template <typename T>
+struct TAllocatorTraits<FSlateStatTrackingMemoryAllocator<T>> : TAllocatorTraitsBase<FSlateStatTrackingMemoryAllocator<T>>
+{
+	enum { IsZeroConstruct = TAllocatorTraits<FDefaultAllocator>::IsZeroConstruct };
+};
+
+#endif // STATS
+
+#if STATS
+
+template <typename DrawElementType>
+using FSlateDrawElementArray = TArray<DrawElementType, FSlateStatTrackingMemoryAllocator<FDrawElementStatTracker>>;
+
+/** Tuple / List of element arrays for each element type */
+using FSlateDrawElementMap = TTuple<
+	FSlateDrawElementArray<FSlateBoxElement>
+	, FSlateDrawElementArray<FSlateBoxElement>
+	, FSlateDrawElementArray<FSlateTextElement>
+	, FSlateDrawElementArray<FSlateShapedTextElement>
+	, FSlateDrawElementArray<FSlateSplineElement>
+	, FSlateDrawElementArray<FSlateLineElement>
+	, FSlateDrawElementArray<FSlateGradientElement>
+	, FSlateDrawElementArray<FSlateViewportElement>
+	, FSlateDrawElementArray<FSlateBoxElement>
+	, FSlateDrawElementArray<FSlateCustomDrawerElement>
+	, FSlateDrawElementArray<FSlateCustomVertsElement>
+	, FSlateDrawElementArray<FSlatePostProcessElement>
+	, FSlateDrawElementArray<FSlateRoundedBoxElement>
+	, FSlateDrawElementArray<FSlateDrawElement>>;
+
+#else
+
+template <typename DrawElementType>
+using FSlateDrawElementArray = TArray<DrawElementType>;
+
+/** Tuple / List of element arrays for each element type */
+using FSlateDrawElementMap = TTuple<
+	FSlateDrawElementArray<FSlateBoxElement>
+	, FSlateDrawElementArray<FSlateBoxElement>
+	, FSlateDrawElementArray<FSlateTextElement>
+	, FSlateDrawElementArray<FSlateShapedTextElement>
+	, FSlateDrawElementArray<FSlateSplineElement>
+	, FSlateDrawElementArray<FSlateLineElement>
+	, FSlateDrawElementArray<FSlateGradientElement>
+	, FSlateDrawElementArray<FSlateViewportElement>
+	, FSlateDrawElementArray<FSlateBoxElement>
+	, FSlateDrawElementArray<FSlateCustomDrawerElement>
+	, FSlateDrawElementArray<FSlateCustomVertsElement>
+	, FSlateDrawElementArray<FSlatePostProcessElement>
+	, FSlateDrawElementArray<FSlateRoundedBoxElement>
+	, FSlateDrawElementArray<FSlateDrawElement>>;
+
+#endif // STATS
+
+/** Helper template to get draw element type in 'FSlateDrawElementMap' for a given element type enum, used to avoid 'auto' return values */
+template<EElementType ElementType>
+using TSlateDrawElement = typename TTupleElement<(uint8)ElementType, FSlateDrawElementMap>::Type::ElementType;
\ No newline at end of file
diff --git a/Engine/Source/Runtime/SlateCore/Public/Rendering/DrawElementPayloads.h b/Engine/Source/Runtime/SlateCore/Public/Rendering/DrawElementPayloads.h
--- a/Engine/Source/Runtime/SlateCore/Public/Rendering/DrawElementPayloads.h
+++ b/Engine/Source/Runtime/SlateCore/Public/Rendering/DrawElementPayloads.h
@@ -9,43 +9,22 @@
 #include "Styling/SlateBrush.h"
 #include "Types/SlateVector2.h"
 
-class FSlateShaderResourceProxy;
+#include "Rendering/RenderingCommon.h"
+#include "Rendering/DrawElementTypes.h"
 
-struct FSlateGradientStop
-{
-	FVector2f Position;
-	FLinearColor Color;
-
-	/**
-	 * Construct a Gradient Stop from a Position and a Color.
-	 * @param InPosition - The position in widget space for this stop. Both X and Y are used for a single-axis gradient.
-						  A two stop gradient should go from (0,0), to (Width,Height).
-	 * @param InColor	- The color to lerp towards at this stop.
-	 */
-	template<typename VectorType>
-	FSlateGradientStop(const VectorType& InPosition, const FLinearColor InColor)
-		: Position(UE::Slate::CastToVector2f(InPosition))
-		, Color(InColor)
-	{
-	}
-};
-template <> struct TIsPODType<FSlateGradientStop> { enum { Value = true }; };
+//////////////////////////////////////////////////////////////////////////
+// Deprecated payloads
+//////////////////////////////////////////////////////////////////////////
 
-struct FSlateDataPayload
+PRAGMA_DISABLE_DEPRECATION_WARNINGS
+
+struct UE_DEPRECATED(5.3, "Draw Element Payloads are no longer used, instead use the equivalent FSlateDrawElement subclass") FSlateDataPayload
 {
 	virtual ~FSlateDataPayload() {}
 	virtual void AddReferencedObjects(FReferenceCollector& Collector) {}
 };
 
-struct FSlateTintableElement
-{
-	FLinearColor Tint;
-
-	FORCEINLINE void SetTint(const FLinearColor& InTint) { Tint = InTint; }
-	FORCEINLINE FLinearColor GetTint() const { return Tint; }
-};
-
-struct FSlateBoxPayload : public FSlateDataPayload, public FSlateTintableElement
+struct UE_DEPRECATED(5.3, "Draw Element Payloads are no longer used, instead use the equivalent FSlateDrawElement subclass") FSlateBoxPayload : public FSlateDataPayload, public FSlateTintableElement
 {
 	FMargin Margin;
 	FBox2f UVRegion;
@@ -88,7 +67,7 @@ struct FSlateBoxPayload : public FSlateDataPayload, public FSlateTintableElement
 	}
 };
 
-struct FSlateRoundedBoxPayload : public FSlateBoxPayload
+struct UE_DEPRECATED(5.3, "Draw Element Payloads are no longer used, instead use the equivalent FSlateDrawElement subclass") FSlateRoundedBoxPayload : public FSlateBoxPayload
 {
 	FLinearColor OutlineColor;
 	FVector4f Radius;
@@ -102,7 +81,7 @@ struct FSlateRoundedBoxPayload : public FSlateBoxPayload
 	FORCEINLINE float GetOutlineWeight() const { return OutlineWeight; }
 };
 
-struct FSlateTextPayload : public FSlateDataPayload, public FSlateTintableElement
+struct UE_DEPRECATED(5.3, "Draw Element Payloads are no longer used, instead use the equivalent FSlateDrawElement subclass") FSlateTextPayload : public FSlateDataPayload, public FSlateTintableElement
 {
 	// The font to use when rendering
 	FSlateFontInfo FontInfo;
@@ -138,9 +117,8 @@ struct FSlateTextPayload : public FSlateDataPayload, public FSlateTintableElemen
 	}
 };
 
-
-struct FSlateShapedTextPayload : public FSlateDataPayload, public FSlateTintableElement
-{ 
+struct UE_DEPRECATED(5.3, "Draw Element Payloads are no longer used, instead use the equivalent FSlateDrawElement subclass") FSlateShapedTextPayload : public FSlateDataPayload, public FSlateTintableElement
+{
 	// Shaped text data
 	FShapedGlyphSequencePtr ShapedGlyphSequence;
 
@@ -177,7 +155,7 @@ struct FSlateShapedTextPayload : public FSlateDataPayload, public FSlateTintable
 	}
 };
 
-struct FSlateGradientPayload : public FSlateDataPayload
+struct UE_DEPRECATED(5.3, "Draw Element Payloads are no longer used, instead use the equivalent FSlateDrawElement subclass") FSlateGradientPayload : public FSlateDataPayload
 {
 	TArray<FSlateGradientStop> GradientStops;
 	EOrientation GradientType;
@@ -191,7 +169,7 @@ struct FSlateGradientPayload : public FSlateDataPayload
 	}
 };
 
-struct FSlateSplinePayload : public FSlateDataPayload, public FSlateTintableElement
+struct UE_DEPRECATED(5.3, "Draw Element Payloads are no longer used, instead use the equivalent FSlateDrawElement subclass") FSlateSplinePayload : public FSlateDataPayload, public FSlateTintableElement
 {
 	TArray<FSlateGradientStop> GradientStops;
 	// Bezier Spline Data points. E.g.
@@ -244,8 +222,8 @@ struct FSlateSplinePayload : public FSlateDataPayload, public FSlateTintableElem
 };
 
 
-struct FSlateLinePayload : public FSlateDataPayload, public FSlateTintableElement
-{ 
+struct UE_DEPRECATED(5.3, "Draw Element Payloads are no longer used, instead use the equivalent FSlateDrawElement subclass") FSlateLinePayload : public FSlateDataPayload, public FSlateTintableElement
+{
 	TArray<FVector2f> Points;
 	TArray<FLinearColor> PointColors;
 	float Thickness;
@@ -297,7 +275,7 @@ struct FSlateLinePayload : public FSlateDataPayload, public FSlateTintableElemen
 	}
 };
 
-struct FSlateViewportPayload : public FSlateDataPayload, public FSlateTintableElement
+struct UE_DEPRECATED(5.3, "Draw Element Payloads are no longer used, instead use the equivalent FSlateDrawElement subclass") FSlateViewportPayload : public FSlateDataPayload, public FSlateTintableElement
 {
 	FSlateShaderResource* RenderTargetResource;
 	uint8 bAllowViewportScaling : 1;
@@ -314,7 +292,7 @@ struct FSlateViewportPayload : public FSlateDataPayload, public FSlateTintableEl
 	}
 };
 
-struct FSlateCustomDrawerPayload : public FSlateDataPayload
+struct UE_DEPRECATED(5.3, "Draw Element Payloads are no longer used, instead use the equivalent FSlateDrawElement subclass") FSlateCustomDrawerPayload : public FSlateDataPayload
 {
 	// Custom drawer data
 	TWeakPtr<ICustomSlateElement, ESPMode::ThreadSafe> CustomDrawer;
@@ -325,7 +303,7 @@ struct FSlateCustomDrawerPayload : public FSlateDataPayload
 	}
 };
 
-struct FSlateLayerPayload : public FSlateDataPayload
+struct UE_DEPRECATED(5.3, "Draw Element Payloads are no longer used, instead use the equivalent FSlateDrawElement subclass") FSlateLayerPayload : public FSlateDataPayload
 {
 	class FSlateDrawLayerHandle* LayerHandle;
 
@@ -337,7 +315,7 @@ struct FSlateLayerPayload : public FSlateDataPayload
 
 };
 
-struct FSlateCachedBufferPayload : public FSlateDataPayload
+struct UE_DEPRECATED(5.3, "Draw Element Payloads are no longer used, instead use the equivalent FSlateDrawElement subclass") FSlateCachedBufferPayload : public FSlateDataPayload
 {
 	// Cached render data
 	class FSlateRenderDataHandle* CachedRenderData;
@@ -353,7 +331,7 @@ struct FSlateCachedBufferPayload : public FSlateDataPayload
 	}
 };
 
-struct FSlateCustomVertsPayload : public FSlateDataPayload
+struct UE_DEPRECATED(5.3, "Draw Element Payloads are no longer used, instead use the equivalent FSlateDrawElement subclass") FSlateCustomVertsPayload : public FSlateDataPayload
 {
 	const FSlateShaderResourceProxy* ResourceProxy;
 
@@ -378,10 +356,12 @@ struct FSlateCustomVertsPayload : public FSlateDataPayload
 	}
 };
 
-struct FSlatePostProcessPayload : public FSlateDataPayload
+struct UE_DEPRECATED(5.3, "Draw Element Payloads are no longer used, instead use the equivalent FSlateDrawElement subclass") FSlatePostProcessPayload : public FSlateDataPayload
 {
 	// Post Process Data
 	FVector4f PostProcessData;
 	FVector4f CornerRadius;
 	int32 DownsampleAmount;
-};
\ No newline at end of file
+};
+
+PRAGMA_ENABLE_DEPRECATION_WARNINGS
\ No newline at end of file
diff --git a/Engine/Source/Runtime/SlateCore/Public/Rendering/DrawElementTypes.h b/Engine/Source/Runtime/SlateCore/Public/Rendering/DrawElementTypes.h
--- a/Engine/Source/Runtime/SlateCore/Public/Rendering/DrawElementTypes.h
+++ b/Engine/Source/Runtime/SlateCore/Public/Rendering/DrawElementTypes.h
@@ -3,31 +3,692 @@
 #pragma once
 
 #include "CoreMinimal.h"
+#include "DrawElementTextOverflowArgs.h"
+#include "Fonts/FontCache.h"
+#include "Fonts/ShapedTextFwd.h"
+#include "SlateRenderBatch.h"
+#include "Styling/SlateBrush.h"
+#include "Types/SlateVector2.h"
+
+#include "UObject/GCObject.h"
+#include "Fonts/ShapedTextFwd.h"
+#include "Misc/MemStack.h"
+#include "Styling/WidgetStyle.h"
+#include "Fonts/SlateFontInfo.h"
+#include "Layout/SlateRect.h"
+#include "Layout/Clipping.h"
+#include "Types/PaintArgs.h"
+#include "Types/SlateVector2.h"
+#include "Layout/Geometry.h"
+#include "Rendering/RenderingCommon.h"
+#include "Debugging/SlateDebugging.h"
+#include "Rendering/SlateRenderBatch.h"
+#include "DrawElementTextOverflowArgs.h"
+#include "ElementBatcher.h"
+#include "Widgets/WidgetPixelSnapping.h"
+#include "Types/SlateVector2.h"
+
+#if UE_ENABLE_INCLUDE_ORDER_DEPRECATED_IN_5_2
+#include "Fonts/FontCache.h"
+#include "Rendering/ShaderResourceManager.h"
+#endif
+
+class FSlateRenderBatch;
+class FSlateDrawLayerHandle;
+class FSlateResourceHandle;
+class FSlateWindowElementList;
+class SWidget;
+class SWindow;
+struct FSlateBrush;
+struct FSlateDataPayload;
+struct FSlateGradientStop;
 
 /**
- * Slate higher level abstract element types
+ * FSlateDrawElement is the building block for Slate's rendering interface.
+ * Slate describes its visual output as an ordered list of FSlateDrawElement s
  */
-enum class EElementType : uint8
-{
-	ET_Box,
-	ET_DebugQuad,
-	ET_Text,
-	ET_ShapedText,
-	ET_Spline,
-	ET_Line,
-	ET_Gradient,
-	ET_Viewport,
-	ET_Border,
-	ET_Custom,
-	ET_CustomVerts,
-	ET_PostProcessPass,
-	ET_RoundedBox,
-	/**
-	 * We map draw elements by type on add for better cache coherency if possible,
-	 * this type is used when that grouping is not possible.
-	 * Grouping is also planned to be used for bulk element type processing.
-	 */
-	ET_NonMapped,
-	/** Total number of draw commands */
-	ET_Count,
-};
\ No newline at end of file
+class FSlateDrawElement
+{
+	friend class FSlateWindowElementList;
+public:
+	
+	enum ERotationSpace
+	{
+		/** Relative to the element.  (0,0) is the upper left corner of the element */
+		RelativeToElement,
+		/** Relative to the alloted paint geometry.  (0,0) is the upper left corner of the paint geometry */
+		RelativeToWorld,
+	};
+
+	/**
+	 * Creates a wireframe quad for debug purposes
+	 *
+	 * @param ElementList			The list in which to add elements
+	 * @param InLayer				The layer to draw the element on
+	 * @param PaintGeometry         DrawSpace position and dimensions; see FPaintGeometry
+	 */
+	SLATECORE_API static void MakeDebugQuad( FSlateWindowElementList& ElementList, uint32 InLayer, const FPaintGeometry& PaintGeometry, FLinearColor Tint = FLinearColor::White);
+
+	/**
+	 * Creates a box element based on the following diagram.  Allows for this element to be resized while maintain the border of the image
+	 * If there are no margins the resulting box is simply a quad
+	 *     ___LeftMargin    ___RightMargin
+	 *    /                /
+	 *  +--+-------------+--+
+	 *  |  |c1           |c2| ___TopMargin
+	 *  +--o-------------o--+
+	 *  |  |             |  |
+	 *  |  |c3           |c4|
+	 *  +--o-------------o--+
+	 *  |  |             |  | ___BottomMargin
+	 *  +--+-------------+--+
+	 *
+	 * @param ElementList			The list in which to add elements
+	 * @param InLayer               The layer to draw the element on
+	 * @param PaintGeometry         DrawSpace position and dimensions; see FPaintGeometry
+	 * @param InBrush               Brush to apply to this element
+	 * @param InDrawEffects         Optional draw effects to apply
+	 * @param InTint                Color to tint the element
+	 */
+	SLATECORE_API static void MakeBox( 
+		FSlateWindowElementList& ElementList,
+		uint32 InLayer,
+		const FPaintGeometry& PaintGeometry,
+		const FSlateBrush* InBrush,
+		ESlateDrawEffect InDrawEffects = ESlateDrawEffect::None,
+		const FLinearColor& InTint = FLinearColor::White );
+
+	SLATECORE_API static void MakeRotatedBox(
+		FSlateWindowElementList& ElementList,
+		uint32 InLayer, 
+		const FPaintGeometry& PaintGeometry, 
+		const FSlateBrush* InBrush, 
+		ESlateDrawEffect,
+		float Angle,
+		UE::Slate::FDeprecateOptionalVector2DParameter InRotationPoint = TOptional<FVector2f>(),
+		ERotationSpace RotationSpace = RelativeToElement,
+		const FLinearColor& InTint = FLinearColor::White );
+
+	/**
+	 * Creates a text element which displays a string of a rendered in a certain font on the screen
+	 *
+	 * @param ElementList			The list in which to add elements
+	 * @param InLayer               The layer to draw the element on
+	 * @param PaintGeometry         DrawSpace position and dimensions; see FPaintGeometry
+	 * @param InText                The string to draw
+	 * @param StartIndex            Inclusive index to start rendering from on the specified text
+	 * @param EndIndex				Exclusive index to stop rendering on the specified text
+	 * @param InFontInfo            The font to draw the string with
+	 * @param InDrawEffects         Optional draw effects to apply
+	 * @param InTint                Color to tint the element
+	 */
+	SLATECORE_API static void MakeText( FSlateWindowElementList& ElementList, uint32 InLayer, const FPaintGeometry& PaintGeometry, const FString& InText, const int32 StartIndex, const int32 EndIndex, const FSlateFontInfo& InFontInfo, ESlateDrawEffect InDrawEffects = ESlateDrawEffect::None, const FLinearColor& InTint = FLinearColor::White );
+	
+	SLATECORE_API static void MakeText( FSlateWindowElementList& ElementList, uint32 InLayer, const FPaintGeometry& PaintGeometry, const FString& InText, const FSlateFontInfo& InFontInfo, ESlateDrawEffect InDrawEffects = ESlateDrawEffect::None, const FLinearColor& InTint = FLinearColor::White );
+
+	FORCEINLINE static void MakeText(FSlateWindowElementList& ElementList, uint32 InLayer, const FPaintGeometry& PaintGeometry, const FText& InText, const FSlateFontInfo& InFontInfo, ESlateDrawEffect InDrawEffects = ESlateDrawEffect::None, const FLinearColor& InTint = FLinearColor::White)
+	{
+		MakeText(ElementList, InLayer, PaintGeometry, InText.ToString(), InFontInfo, InDrawEffects, InTint);
+	}
+
+	/**
+	 * Creates a text element which displays a series of shaped glyphs on the screen
+	 *
+	 * @param ElementList			The list in which to add elements
+	 * @param InLayer               The layer to draw the element on
+	 * @param PaintGeometry         DrawSpace position and dimensions; see FPaintGeometry
+	 * @param InShapedGlyphSequence The shaped glyph sequence to draw
+	 * @param InDrawEffects         Optional draw effects to apply
+	 * @param InTint                Color to tint the element
+	 */
+	SLATECORE_API static void MakeShapedText( FSlateWindowElementList& ElementList, uint32 InLayer, const FPaintGeometry& PaintGeometry, const FShapedGlyphSequenceRef& InShapedGlyphSequence, ESlateDrawEffect InDrawEffects, const FLinearColor& BaseTint, const FLinearColor& OutlineTint, FTextOverflowArgs TextOverflowArgs = FTextOverflowArgs());
+
+	/**
+	 * Creates a gradient element
+	 *
+	 * @param ElementList			   The list in which to add elements
+	 * @param InLayer                  The layer to draw the element on
+	 * @param PaintGeometry            DrawSpace position and dimensions; see FPaintGeometry
+	 * @param InGradientStops          List of gradient stops which define the element
+	 * @param InGradientType           The type of gradient (I.E Horizontal, vertical)
+	 * @param InDrawEffects            Optional draw effects to apply
+	 * @param CornerRadius			   Rounds the corners of the box created by the gradient by the specified radius
+	 */
+	SLATECORE_API static void MakeGradient( FSlateWindowElementList& ElementList, uint32 InLayer, const FPaintGeometry& PaintGeometry, TArray<FSlateGradientStop> InGradientStops, EOrientation InGradientType, ESlateDrawEffect InDrawEffects = ESlateDrawEffect::None, FVector4f CornerRadius = FVector4f(0.0f) );
+
+	/**
+	 * Creates a Hermite Spline element
+	 *
+	 * @param ElementList			The list in which to add elements
+	 * @param InLayer               The layer to draw the element on
+	 * @param PaintGeometry         DrawSpace position and dimensions; see FPaintGeometry
+	 * @param InStart               The start point of the spline (local space)
+	 * @param InStartDir            The direction of the spline from the start point
+	 * @param InEnd                 The end point of the spline (local space)
+	 * @param InEndDir              The direction of the spline to the end point
+	 * @param InDrawEffects         Optional draw effects to apply
+	 * @param InTint                Color to tint the element
+	 */
+	SLATECORE_API static void MakeSpline(FSlateWindowElementList& ElementList, uint32 InLayer, const FPaintGeometry& PaintGeometry, const UE::Slate::FDeprecateVector2DParameter InStart, const UE::Slate::FDeprecateVector2DParameter InStartDir, const UE::Slate::FDeprecateVector2DParameter InEnd, const UE::Slate::FDeprecateVector2DParameter InEndDir, float InThickness = 0.0f, ESlateDrawEffect InDrawEffects = ESlateDrawEffect::None, const FLinearColor& InTint=FLinearColor::White);
+
+	/**
+	 * Creates a Bezier Spline element
+	 *
+	 * @param ElementList			The list in which to add elements
+	 * @param InLayer               The layer to draw the element on
+	 * @param PaintGeometry         DrawSpace position and dimensions; see FPaintGeometry
+	 * @param InStart               The start point of the spline (local space)
+	 * @param InStartDir            The direction of the spline from the start point
+	 * @param InEnd                 The end point of the spline (local space)
+	 * @param InEndDir              The direction of the spline to the end point
+	 * @param InDrawEffects         Optional draw effects to apply
+	 * @param InTint                Color to tint the element
+	 */
+	SLATECORE_API static void MakeCubicBezierSpline(FSlateWindowElementList& ElementList, uint32 InLayer, const FPaintGeometry& PaintGeometry, const UE::Slate::FDeprecateVector2DParameter P0, const UE::Slate::FDeprecateVector2DParameter P1, const UE::Slate::FDeprecateVector2DParameter P2, const UE::Slate::FDeprecateVector2DParameter P3, float InThickness = 0.0f, ESlateDrawEffect InDrawEffects = ESlateDrawEffect::None, const FLinearColor& InTint = FLinearColor::White);
+
+	/** Just like MakeSpline but in draw-space coordinates. This is useful for connecting already-transformed widgets together. */
+	SLATECORE_API static void MakeDrawSpaceSpline(FSlateWindowElementList& ElementList, uint32 InLayer, const UE::Slate::FDeprecateVector2DParameter InStart, const UE::Slate::FDeprecateVector2DParameter InStartDir, const UE::Slate::FDeprecateVector2DParameter InEnd, const UE::Slate::FDeprecateVector2DParameter InEndDir, float InThickness = 0.0f, ESlateDrawEffect InDrawEffects = ESlateDrawEffect::None, const FLinearColor& InTint=FLinearColor::White);
+
+	/**
+	 * Creates a line defined by the provided points
+	 *
+	 * @param ElementList              The list in which to add elements
+	 * @param InLayer                  The layer to draw the element on
+	 * @param PaintGeometry            DrawSpace position and dimensions; see FPaintGeometry
+	 * @param Points                   Points that make up the lines.  The points are joined together. I.E if Points has A,B,C there the line is A-B-C.  To draw non-joining line segments call MakeLines multiple times
+	 * @param InDrawEffects            Optional draw effects to apply
+	 * @param InTint                   Color to tint the element
+	 * @param bAntialias               Should antialiasing be applied to the line?
+	 * @param Thickness                The thickness of the line
+	 */
+#if UE_ENABLE_SLATE_VECTOR_DEPRECATION_MECHANISMS
+	SLATECORE_API static void MakeLines(FSlateWindowElementList& ElementList, uint32 InLayer, const FPaintGeometry& PaintGeometry, const TArray<FVector2d>& Points, ESlateDrawEffect InDrawEffects = ESlateDrawEffect::None, const FLinearColor& InTint=FLinearColor::White, bool bAntialias = true, float Thickness = 1.0f);
+#endif
+	SLATECORE_API static void MakeLines(FSlateWindowElementList& ElementList, uint32 InLayer, const FPaintGeometry& PaintGeometry, TArray<FVector2f> Points, ESlateDrawEffect InDrawEffects = ESlateDrawEffect::None, const FLinearColor& InTint = FLinearColor::White, bool bAntialias = true, float Thickness = 1.0f);
+
+
+	/**
+	 * Creates a line defined by the provided points
+	 *
+	 * @param ElementList              The list in which to add elements
+	 * @param InLayer                  The layer to draw the element on
+	 * @param PaintGeometry            DrawSpace position and dimensions; see FPaintGeometry
+	 * @param Points                   Points that make up the lines.  The points are joined together. I.E if Points has A,B,C there the line is A-B-C.  To draw non-joining line segments call MakeLines multiple times
+	 * @param PointColors              Vertex Color for each defined points
+	 * @param InDrawEffects            Optional draw effects to apply
+	 * @param InTint                   Color to tint the element
+	 * @param bAntialias               Should antialiasing be applied to the line?
+	 * @param Thickness                The thickness of the line
+	 */
+#if UE_ENABLE_SLATE_VECTOR_DEPRECATION_MECHANISMS
+	SLATECORE_API static void MakeLines(FSlateWindowElementList& ElementList, uint32 InLayer, const FPaintGeometry& PaintGeometry, const TArray<FVector2d>& Points, const TArray<FLinearColor>& PointColors, ESlateDrawEffect InDrawEffects = ESlateDrawEffect::None, const FLinearColor& InTint=FLinearColor::White, bool bAntialias = true, float Thickness = 1.0f);
+#endif
+	SLATECORE_API static void MakeLines(FSlateWindowElementList& ElementList, uint32 InLayer, const FPaintGeometry& PaintGeometry, TArray<FVector2f> Points, TArray<FLinearColor> PointColors, ESlateDrawEffect InDrawEffects = ESlateDrawEffect::None, const FLinearColor& InTint=FLinearColor::White, bool bAntialias = true, float Thickness = 1.0f);
+
+	/**
+	 * Creates a viewport element which is useful for rendering custom data in a texture into Slate
+	 *
+	 * @param ElementList		   The list in which to add elements
+	 * @param InLayer                  The layer to draw the element on
+	 * @param PaintGeometry            DrawSpace position and dimensions; see FPaintGeometry
+	 * @param Viewport                 Interface for drawing the viewport
+	 * @param InScale                  Draw scale to apply to the entire element
+	 * @param InDrawEffects            Optional draw effects to apply
+	 * @param InTint                   Color to tint the element
+	 */
+	SLATECORE_API static void MakeViewport( FSlateWindowElementList& ElementList, uint32 InLayer, const FPaintGeometry& PaintGeometry, TSharedPtr<const ISlateViewport> Viewport, ESlateDrawEffect InDrawEffects = ESlateDrawEffect::None, const FLinearColor& InTint=FLinearColor::White );
+
+	/**
+	 * Creates a custom element which can be used to manually draw into the Slate render target with graphics API calls rather than Slate elements
+	 *
+	 * @param ElementList		   The list in which to add elements
+	 * @param InLayer                  The layer to draw the element on
+	 * @param PaintGeometry            DrawSpace position and dimensions; see FPaintGeometry
+	 * @param CustomDrawer		   Interface to a drawer which will be called when Slate renders this element
+	 */
+	SLATECORE_API static void MakeCustom( FSlateWindowElementList& ElementList, uint32 InLayer, TSharedPtr<ICustomSlateElement, ESPMode::ThreadSafe> CustomDrawer );
+	
+	SLATECORE_API static void MakeCustomVerts(FSlateWindowElementList& ElementList, uint32 InLayer, const FSlateResourceHandle& InRenderResourceHandle, const TArray<FSlateVertex>& InVerts, const TArray<SlateIndex>& InIndexes, ISlateUpdatableInstanceBuffer* InInstanceData, uint32 InInstanceOffset, uint32 InNumInstances, ESlateDrawEffect InDrawEffects = ESlateDrawEffect::None);
+
+	SLATECORE_API static void MakePostProcessPass(FSlateWindowElementList& ElementList, uint32 InLayer, const FPaintGeometry& PaintGeometry, const FVector4f& Params, int32 DownsampleAmount, const FVector4f CornerRadius = FVector4f(0.0f));
+
+	FSlateDrawElement();
+
+	FORCEINLINE int32 GetLayer() const { return LayerId; }
+
+	FORCEINLINE const FSlateRenderTransform& GetRenderTransform() const { return RenderTransform; }
+	FORCEINLINE void SetRenderTransform(const FSlateRenderTransform& InRenderTransform) { RenderTransform = InRenderTransform; }
+
+	FORCEINLINE UE::Slate::FDeprecateVector2DResult GetPosition() const
+	{
+		return UE::Slate::FDeprecateVector2DResult(Position);
+	}
+
+	FORCEINLINE void SetPosition(UE::Slate::FDeprecateVector2DParameter InPosition)
+	{
+		Position = InPosition;
+	}
+
+	FORCEINLINE UE::Slate::FDeprecateVector2DResult GetLocalSize() const
+	{
+		return UE::Slate::FDeprecateVector2DResult(LocalSize);
+	}
+
+	FORCEINLINE float GetScale() const { return Scale; }
+	FORCEINLINE ESlateDrawEffect GetDrawEffects() const { return DrawEffects; }
+	FORCEINLINE ESlateBatchDrawFlag GetBatchFlags() const { return BatchFlags; }
+	FORCEINLINE bool IsPixelSnapped() const { return !EnumHasAllFlags(DrawEffects, ESlateDrawEffect::NoPixelSnapping); }
+
+	FORCEINLINE int32 GetPrecachedClippingIndex() const { return ClipStateHandle.GetPrecachedClipIndex(); }
+	FORCEINLINE void SetPrecachedClippingIndex(int32 InClippingIndex) { ClipStateHandle.SetPreCachedClipIndex(InClippingIndex); }
+
+	FORCEINLINE void SetCachedClippingState(const FSlateClippingState* CachedState) { ClipStateHandle.SetCachedClipState(CachedState); }
+	FORCEINLINE const FClipStateHandle& GetClippingHandle() const { return ClipStateHandle; }
+	FORCEINLINE const int8 GetSceneIndex() const { return SceneIndex; }
+
+	FORCEINLINE void SetIsCached(bool bInIsCached) { bIsCached = bInIsCached; }
+	FORCEINLINE bool IsCached() const { return bIsCached; }
+
+	FORCEINLINE FSlateLayoutTransform GetInverseLayoutTransform() const
+	{
+		return Inverse(FSlateLayoutTransform(Scale, Position));
+	}
+
+	/**
+	 * Update element cached position with an arbitrary offset
+	 *
+	 * @param Element		   Element to update
+	 * @param InOffset         Absolute translation delta
+	 */
+	void ApplyPositionOffset(UE::Slate::FDeprecateVector2DParameter InOffset);
+
+	UE_DEPRECATED(4.23, "GetClippingIndex has been deprecated.  If you were using this please use GetPrecachedClippingIndex instead.")
+	FORCEINLINE const int32 GetClippingIndex() const { return GetPrecachedClippingIndex(); }
+
+	UE_DEPRECATED(4.23, "SetClippingIndex has been deprecated.  If you were using this please use SetPrecachedClippingIndex instead.")
+	FORCEINLINE void SetClippingIndex(const int32 InClippingIndex) { SetPrecachedClippingIndex(InClippingIndex); }
+
+private:
+	void Init(FSlateWindowElementList& ElementList, EElementType InElementType, uint32 InLayer, const FPaintGeometry& PaintGeometry, ESlateDrawEffect InDrawEffects);
+
+	static FVector2f GetRotationPoint( const FPaintGeometry& PaintGeometry, const TOptional<FVector2f>& UserRotationPoint, ERotationSpace RotationSpace );
+	static FSlateDrawElement& MakeBoxInternal(FSlateWindowElementList& ElementList, uint32 InLayer, const FPaintGeometry& PaintGeometry, const FSlateBrush* InBrush, ESlateDrawEffect InDrawEffects, const FLinearColor& InTint);
+private:
+	FSlateRenderTransform RenderTransform;
+	FVector2f Position;
+	FVector2f LocalSize;
+	int32 LayerId;
+	FClipStateHandle ClipStateHandle;
+	float Scale;
+	int8 SceneIndex;
+	ESlateDrawEffect DrawEffects;
+	EElementType ElementType;
+	// Misc data
+	ESlateBatchDrawFlag BatchFlags;
+	uint8 bIsCached : 1;
+};
+
+class FSlateShaderResourceProxy;
+
+struct FSlateGradientStop
+{
+	FVector2f Position;
+	FLinearColor Color;
+
+	/**
+	 * Construct a Gradient Stop from a Position and a Color.
+	 * @param InPosition - The position in widget space for this stop. Both X and Y are used for a single-axis gradient.
+						  A two stop gradient should go from (0,0), to (Width,Height).
+	 * @param InColor	- The color to lerp towards at this stop.
+	 */
+	template<typename VectorType>
+	FSlateGradientStop(const VectorType& InPosition, const FLinearColor InColor)
+		: Position(UE::Slate::CastToVector2f(InPosition))
+		, Color(InColor)
+	{
+	}
+};
+template <> struct TIsPODType<FSlateGradientStop> { enum { Value = true }; };
+
+struct FSlateTintableElement
+{
+	FLinearColor Tint;
+
+	FORCEINLINE void SetTint(const FLinearColor& InTint) { Tint = InTint; }
+	FORCEINLINE FLinearColor GetTint() const { return Tint; }
+};
+
+struct FSlateBoxElement : public FSlateDrawElement, public FSlateTintableElement
+{
+	FMargin Margin;
+	FBox2f UVRegion;
+	const FSlateShaderResourceProxy* ResourceProxy;
+	ESlateBrushTileType::Type Tiling;
+	ESlateBrushMirrorType::Type Mirroring;
+	ESlateBrushDrawType::Type DrawType;
+
+	const FMargin& GetBrushMargin() const { return Margin; }
+	const FBox2f& GetBrushUVRegion() const { return UVRegion; }
+	ESlateBrushTileType::Type GetBrushTiling() const { return Tiling; }
+	ESlateBrushMirrorType::Type GetBrushMirroring() const { return Mirroring; }
+	ESlateBrushDrawType::Type GetBrushDrawType() const { return DrawType; }
+	const FSlateShaderResourceProxy* GetResourceProxy() const { return ResourceProxy; }
+
+	void SetBrush(const FSlateBrush* InBrush, UE::Slate::FDeprecateVector2DParameter InLocalSize, float DrawScale)
+	{
+		check(InBrush);
+		ensureMsgf(InBrush->GetDrawType() != ESlateBrushDrawType::NoDrawType, TEXT("This should have been filtered out earlier in the Make... call."));
+
+		// Note: Do not store the brush.  It is possible brushes are destroyed after an element is enqueued for rendering
+		Margin = InBrush->GetMargin();
+		UVRegion = InBrush->GetUVRegion();
+		Tiling = InBrush->GetTiling();
+		Mirroring = InBrush->GetMirroring();
+		DrawType = InBrush->GetDrawType();
+		const FSlateResourceHandle& Handle = InBrush->GetRenderingResource(InLocalSize, DrawScale);
+		if (Handle.IsValid())
+		{
+			ResourceProxy = Handle.GetResourceProxy();
+		}
+		else
+		{
+			ResourceProxy = nullptr;
+		}
+	}
+};
+
+struct FSlateRoundedBoxElement : public FSlateBoxElement
+{
+	FLinearColor OutlineColor;
+	FVector4f Radius;
+	float OutlineWeight;
+
+	FORCEINLINE void SetRadius(FVector4f InRadius) { Radius = InRadius; }
+	FORCEINLINE FVector4f GetRadius() const { return Radius; }
+
+	FORCEINLINE void SetOutline(const FLinearColor& InOutlineColor, float InOutlineWeight) { OutlineColor = InOutlineColor; OutlineWeight = InOutlineWeight; }
+	FORCEINLINE FLinearColor GetOutlineColor() const { return OutlineColor; }
+	FORCEINLINE float GetOutlineWeight() const { return OutlineWeight; }
+};
+
+struct FSlateTextElement : public FSlateDrawElement, public FSlateTintableElement
+{
+	// The font to use when rendering
+	FSlateFontInfo FontInfo;
+	// Basic text data
+	FString ImmutableText;
+
+	const FSlateFontInfo& GetFontInfo() const { return FontInfo; }
+	const TCHAR* GetText() const { return *ImmutableText; }
+	int32 GetTextLength() const { return ImmutableText.Len(); }
+
+	void SetText(const FString& InText, const FSlateFontInfo& InFontInfo, int32 InStartIndex, int32 InEndIndex)
+	{
+		FontInfo = InFontInfo;
+		const int32 StartIndex = FMath::Min<int32>(InStartIndex, InText.Len());
+		const int32 EndIndex = FMath::Min<int32>(InEndIndex, InText.Len());
+		const int32 TextLength = (EndIndex > StartIndex) ? EndIndex - StartIndex : 0;
+		if (TextLength > 0)
+		{
+			ImmutableText = InText.Mid(StartIndex, TextLength);
+		}
+	}
+
+	void SetText(const FString& InText, const FSlateFontInfo& InFontInfo)
+	{
+		FontInfo = InFontInfo;
+		ImmutableText = InText;
+	}
+
+
+	void AddReferencedObjects(FReferenceCollector& Collector)
+	{
+		FontInfo.AddReferencedObjects(Collector);
+	}
+};
+
+struct FSlateShapedTextElement : public FSlateDrawElement, public FSlateTintableElement
+{ 
+	// Shaped text data
+	FShapedGlyphSequencePtr ShapedGlyphSequence;
+
+	FLinearColor OutlineTint;
+
+	FTextOverflowArgs OverflowArgs;
+
+	const FShapedGlyphSequencePtr& GetShapedGlyphSequence() const { return ShapedGlyphSequence; }
+	FLinearColor GetOutlineTint() const { return OutlineTint; }
+
+	void SetShapedText(FSlateWindowElementList& ElementList, const FShapedGlyphSequencePtr& InShapedGlyphSequence, const FLinearColor& InOutlineTint)
+	{
+		ShapedGlyphSequence = InShapedGlyphSequence;
+		OutlineTint = InOutlineTint;
+	}
+
+	void SetOverflowArgs(const FTextOverflowArgs& InArgs)
+	{
+		OverflowArgs = InArgs;
+		check(InArgs.OverflowDirection == ETextOverflowDirection::NoOverflow || InArgs.OverflowTextPtr.IsValid());
+	}
+
+	void AddReferencedObjects(FReferenceCollector& Collector)
+	{
+		if (ShapedGlyphSequence.IsValid())
+		{
+			const_cast<FShapedGlyphSequence*>(ShapedGlyphSequence.Get())->AddReferencedObjects(Collector);
+		}
+
+		if (OverflowArgs.OverflowTextPtr.IsValid())
+		{
+			const_cast<FShapedGlyphSequence*>(OverflowArgs.OverflowTextPtr.Get())->AddReferencedObjects(Collector);
+		}
+	}
+};
+
+struct FSlateGradientElement : public FSlateDrawElement
+{
+	TArray<FSlateGradientStop> GradientStops;
+	EOrientation GradientType;
+	FVector4f CornerRadius;
+
+	void SetGradient(TArray<FSlateGradientStop> InGradientStops, EOrientation InGradientType, FVector4f InCornerRadius)
+	{
+		GradientStops = MoveTemp(InGradientStops);
+		GradientType = InGradientType;
+		CornerRadius = InCornerRadius;
+	}
+};
+
+struct FSlateSplineElement : public FSlateDrawElement, public FSlateTintableElement
+{
+	TArray<FSlateGradientStop> GradientStops;
+	// Bezier Spline Data points. E.g.
+	//
+	//       P1 + - - - - + P2                P1 +
+	//         /           \                    / \
+	//     P0 *             * P3            P0 *   \   * P3
+	//                                              \ /
+	//                                               + P2	
+	FVector2f P0;
+	FVector2f P1;
+	FVector2f P2;
+	FVector2f P3;
+
+	float Thickness;
+
+	// Thickness
+	void SetThickness(float InThickness) { Thickness = InThickness; }
+	float GetThickness() const { return Thickness; }
+
+	void SetCubicBezier(const UE::Slate::FDeprecateVector2DParameter InP0, const UE::Slate::FDeprecateVector2DParameter InP1, const UE::Slate::FDeprecateVector2DParameter InP2, const UE::Slate::FDeprecateVector2DParameter InP3, float InThickness, const FLinearColor InTint)
+	{
+		Tint = InTint;
+		P0 = InP0;
+		P1 = InP1;
+		P2 = InP2;
+		P3 = InP3;
+		Thickness = InThickness;
+	}
+
+	void SetHermiteSpline(const UE::Slate::FDeprecateVector2DParameter InStart, const UE::Slate::FDeprecateVector2DParameter InStartDir, const UE::Slate::FDeprecateVector2DParameter InEnd, const UE::Slate::FDeprecateVector2DParameter InEndDir, float InThickness, const FLinearColor InTint)
+	{
+		Tint = InTint;
+		P0 = InStart;
+		P1 = InStart + InStartDir / 3.0f;
+		P2 = InEnd - InEndDir / 3.0f;
+		P3 = InEnd;
+		Thickness = InThickness;
+	}
+
+	void SetGradientHermiteSpline(const UE::Slate::FDeprecateVector2DParameter InStart, const UE::Slate::FDeprecateVector2DParameter InStartDir, const UE::Slate::FDeprecateVector2DParameter InEnd, const UE::Slate::FDeprecateVector2DParameter InEndDir, float InThickness, TArray<FSlateGradientStop> InGradientStops)
+	{
+		P0 = InStart;
+		P1 = InStart + InStartDir / 3.0f;
+		P2 = InEnd - InEndDir / 3.0f;
+		P3 = InEnd;
+		Thickness = InThickness;
+		GradientStops = MoveTemp(InGradientStops);
+	}
+};
+
+struct FSlateLineElement : public FSlateDrawElement, public FSlateTintableElement
+{ 
+	TArray<FVector2f> Points;
+	TArray<FLinearColor> PointColors;
+	float Thickness;
+
+	bool bAntialias;
+
+	bool IsAntialiased() const { return bAntialias; }
+	const TArray<FVector2f>& GetPoints() const { return Points; }
+	const TArray<FLinearColor>& GetPointColors() const { return PointColors; }
+	float GetThickness() const { return Thickness; }
+
+	void SetThickness(float InThickness)
+	{
+		Thickness = InThickness;
+	}
+
+#if UE_ENABLE_SLATE_VECTOR_DEPRECATION_MECHANISMS
+	void SetLines(const TArray<FVector2D>& InPoints, bool bInAntialias, const TArray<FLinearColor>* InPointColors = nullptr)
+	{
+		TArray<FVector2f> NewPoints;
+		NewPoints.Reserve(InPoints.Num());
+		for (FVector2D Vect : InPoints)
+		{
+			NewPoints.Add(UE::Slate::CastToVector2f(Vect));
+		}
+		if (InPointColors)
+		{
+			SetLines(MoveTemp(NewPoints), bInAntialias, *InPointColors);
+		}
+		else
+		{
+			SetLines(MoveTemp(NewPoints), bInAntialias);
+		}
+	}
+#endif
+
+	void SetLines(TArray<FVector2f> InPoints, bool bInAntialias)
+	{
+		bAntialias = bInAntialias;
+		Points = MoveTemp(InPoints);
+		PointColors.Reset();
+	}
+
+	void SetLines(TArray<FVector2f> InPoints, bool bInAntialias, TArray<FLinearColor> InPointColors)
+	{
+		bAntialias = bInAntialias;
+		Points = MoveTemp(InPoints);
+		PointColors = MoveTemp(InPointColors);
+	}
+};
+
+struct FSlateViewportElement : public FSlateDrawElement, public FSlateTintableElement
+{
+	FSlateShaderResource* RenderTargetResource;
+	uint8 bAllowViewportScaling : 1;
+	uint8 bViewportTextureAlphaOnly : 1;
+	uint8 bRequiresVSync : 1;
+
+	void SetViewport(const TSharedPtr<const ISlateViewport>& InViewport, const FLinearColor& InTint)
+	{
+		Tint = InTint;
+		RenderTargetResource = InViewport->GetViewportRenderTargetTexture();
+		bAllowViewportScaling = InViewport->AllowScaling();
+		bViewportTextureAlphaOnly = InViewport->IsViewportTextureAlphaOnly();
+		bRequiresVSync = InViewport->RequiresVsync();
+	}
+};
+
+struct FSlateCustomDrawerElement : public FSlateDrawElement
+{
+	// Custom drawer data
+	TWeakPtr<ICustomSlateElement, ESPMode::ThreadSafe> CustomDrawer;
+
+	void SetCustomDrawer(const TSharedPtr<ICustomSlateElement, ESPMode::ThreadSafe>& InCustomDrawer)
+	{
+		CustomDrawer = InCustomDrawer;
+	}
+};
+
+struct FSlateLayerElement : public FSlateDrawElement
+{
+	class FSlateDrawLayerHandle* LayerHandle;
+
+	void SetLayer(FSlateDrawLayerHandle* InLayerHandle)
+	{
+		LayerHandle = InLayerHandle;
+		checkSlow(LayerHandle);
+	}
+
+};
+
+struct FSlateCachedBufferElement : public FSlateDrawElement
+{
+	// Cached render data
+	class FSlateRenderDataHandle* CachedRenderData;
+	FVector2f CachedRenderDataOffset;
+
+	// Cached Buffers
+	void SetCachedBuffer(FSlateRenderDataHandle* InRenderDataHandle, const UE::Slate::FDeprecateVector2DParameter Offset)
+	{
+		check(InRenderDataHandle);
+
+		CachedRenderData = InRenderDataHandle;
+		CachedRenderDataOffset = Offset;
+	}
+};
+
+struct FSlateCustomVertsElement : public FSlateDrawElement
+{
+	const FSlateShaderResourceProxy* ResourceProxy;
+
+	TArray<FSlateVertex> Vertices;
+	TArray<SlateIndex> Indices;
+
+	// Instancing support
+	ISlateUpdatableInstanceBufferRenderProxy* InstanceData;
+	uint32 InstanceOffset;
+	uint32 NumInstances;
+
+	void SetCustomVerts(const FSlateShaderResourceProxy* InRenderProxy, TArray<FSlateVertex> InVerts, TArray<SlateIndex> InIndices, ISlateUpdatableInstanceBufferRenderProxy* InInstanceData, uint32 InInstanceOffset, uint32 InNumInstances)
+	{
+		ResourceProxy = InRenderProxy;
+
+		Vertices = MoveTemp(InVerts);
+		Indices = MoveTemp(InIndices);
+
+		InstanceData = InInstanceData;
+		InstanceOffset = InInstanceOffset;
+		NumInstances = InNumInstances;
+	}
+};
+
+struct FSlatePostProcessElement : public FSlateDrawElement
+{
+	// Post Process Data
+	FVector4f PostProcessData;
+	FVector4f CornerRadius;
+	int32 DownsampleAmount;
+};
diff --git a/Engine/Source/Runtime/SlateCore/Public/Rendering/DrawElements.h b/Engine/Source/Runtime/SlateCore/Public/Rendering/DrawElements.h
--- a/Engine/Source/Runtime/SlateCore/Public/Rendering/DrawElements.h
+++ b/Engine/Source/Runtime/SlateCore/Public/Rendering/DrawElements.h
@@ -3,24 +3,9 @@
 #pragma once
 
 #include "CoreMinimal.h"
-#include "UObject/GCObject.h"
-#include "Fonts/ShapedTextFwd.h"
-#include "Stats/Stats.h"
-#include "Misc/MemStack.h"
-#include "Styling/WidgetStyle.h"
-#include "Fonts/SlateFontInfo.h"
-#include "Layout/SlateRect.h"
-#include "Layout/Clipping.h"
-#include "Types/PaintArgs.h"
-#include "Types/SlateVector2.h"
-#include "Layout/Geometry.h"
-#include "Rendering/RenderingCommon.h"
-#include "Debugging/SlateDebugging.h"
-#include "Rendering/SlateRenderBatch.h"
-#include "DrawElementTextOverflowArgs.h"
-#include "ElementBatcher.h"
-#include "Widgets/WidgetPixelSnapping.h"
-#include "Types/SlateVector2.h"
+#include "Rendering/DrawElementCoreTypes.h"
+#include "Rendering/DrawElementTypes.h"
+#include "Rendering/DrawElementPayloads.h"
 
 #if UE_ENABLE_INCLUDE_ORDER_DEPRECATED_IN_5_2
 #include "Fonts/FontCache.h"
@@ -38,304 +23,9 @@ struct FSlateBrush;
 struct FSlateDataPayload;
 struct FSlateGradientStop;
 
-
-
 /**
- * FSlateDrawElement is the building block for Slate's rendering interface.
- * Slate describes its visual output as an ordered list of FSlateDrawElement s
+ * Note: FSlateDrawElement has been moved to DrawElementTypes.h
  */
-class FSlateDrawElement
-{
-	friend class FSlateWindowElementList;
-public:
-	
-	enum ERotationSpace
-	{
-		/** Relative to the element.  (0,0) is the upper left corner of the element */
-		RelativeToElement,
-		/** Relative to the alloted paint geometry.  (0,0) is the upper left corner of the paint geometry */
-		RelativeToWorld,
-	};
-
-	/**
-	 * Creates a wireframe quad for debug purposes
-	 *
-	 * @param ElementList			The list in which to add elements
-	 * @param InLayer				The layer to draw the element on
-	 * @param PaintGeometry         DrawSpace position and dimensions; see FPaintGeometry
-	 */
-	SLATECORE_API static void MakeDebugQuad( FSlateWindowElementList& ElementList, uint32 InLayer, const FPaintGeometry& PaintGeometry, FLinearColor Tint = FLinearColor::White);
-
-	/**
-	 * Creates a box element based on the following diagram.  Allows for this element to be resized while maintain the border of the image
-	 * If there are no margins the resulting box is simply a quad
-	 *     ___LeftMargin    ___RightMargin
-	 *    /                /
-	 *  +--+-------------+--+
-	 *  |  |c1           |c2| ___TopMargin
-	 *  +--o-------------o--+
-	 *  |  |             |  |
-	 *  |  |c3           |c4|
-	 *  +--o-------------o--+
-	 *  |  |             |  | ___BottomMargin
-	 *  +--+-------------+--+
-	 *
-	 * @param ElementList			The list in which to add elements
-	 * @param InLayer               The layer to draw the element on
-	 * @param PaintGeometry         DrawSpace position and dimensions; see FPaintGeometry
-	 * @param InBrush               Brush to apply to this element
-	 * @param InDrawEffects         Optional draw effects to apply
-	 * @param InTint                Color to tint the element
-	 */
-	SLATECORE_API static void MakeBox( 
-		FSlateWindowElementList& ElementList,
-		uint32 InLayer,
-		const FPaintGeometry& PaintGeometry,
-		const FSlateBrush* InBrush,
-		ESlateDrawEffect InDrawEffects = ESlateDrawEffect::None,
-		const FLinearColor& InTint = FLinearColor::White );
-
-	SLATECORE_API static void MakeRotatedBox(
-		FSlateWindowElementList& ElementList,
-		uint32 InLayer, 
-		const FPaintGeometry& PaintGeometry, 
-		const FSlateBrush* InBrush, 
-		ESlateDrawEffect,
-		float Angle,
-		UE::Slate::FDeprecateOptionalVector2DParameter InRotationPoint = TOptional<FVector2f>(),
-		ERotationSpace RotationSpace = RelativeToElement,
-		const FLinearColor& InTint = FLinearColor::White );
-
-	/**
-	 * Creates a text element which displays a string of a rendered in a certain font on the screen
-	 *
-	 * @param ElementList			The list in which to add elements
-	 * @param InLayer               The layer to draw the element on
-	 * @param PaintGeometry         DrawSpace position and dimensions; see FPaintGeometry
-	 * @param InText                The string to draw
-	 * @param StartIndex            Inclusive index to start rendering from on the specified text
-	 * @param EndIndex				Exclusive index to stop rendering on the specified text
-	 * @param InFontInfo            The font to draw the string with
-	 * @param InDrawEffects         Optional draw effects to apply
-	 * @param InTint                Color to tint the element
-	 */
-	SLATECORE_API static void MakeText( FSlateWindowElementList& ElementList, uint32 InLayer, const FPaintGeometry& PaintGeometry, const FString& InText, const int32 StartIndex, const int32 EndIndex, const FSlateFontInfo& InFontInfo, ESlateDrawEffect InDrawEffects = ESlateDrawEffect::None, const FLinearColor& InTint = FLinearColor::White );
-	
-	SLATECORE_API static void MakeText( FSlateWindowElementList& ElementList, uint32 InLayer, const FPaintGeometry& PaintGeometry, const FString& InText, const FSlateFontInfo& InFontInfo, ESlateDrawEffect InDrawEffects = ESlateDrawEffect::None, const FLinearColor& InTint = FLinearColor::White );
-
-	FORCEINLINE static void MakeText(FSlateWindowElementList& ElementList, uint32 InLayer, const FPaintGeometry& PaintGeometry, const FText& InText, const FSlateFontInfo& InFontInfo, ESlateDrawEffect InDrawEffects = ESlateDrawEffect::None, const FLinearColor& InTint = FLinearColor::White)
-	{
-		MakeText(ElementList, InLayer, PaintGeometry, InText.ToString(), InFontInfo, InDrawEffects, InTint);
-	}
-
-	/**
-	 * Creates a text element which displays a series of shaped glyphs on the screen
-	 *
-	 * @param ElementList			The list in which to add elements
-	 * @param InLayer               The layer to draw the element on
-	 * @param PaintGeometry         DrawSpace position and dimensions; see FPaintGeometry
-	 * @param InShapedGlyphSequence The shaped glyph sequence to draw
-	 * @param InDrawEffects         Optional draw effects to apply
-	 * @param InTint                Color to tint the element
-	 */
-	SLATECORE_API static void MakeShapedText( FSlateWindowElementList& ElementList, uint32 InLayer, const FPaintGeometry& PaintGeometry, const FShapedGlyphSequenceRef& InShapedGlyphSequence, ESlateDrawEffect InDrawEffects, const FLinearColor& BaseTint, const FLinearColor& OutlineTint, FTextOverflowArgs TextOverflowArgs = FTextOverflowArgs());
-
-	/**
-	 * Creates a gradient element
-	 *
-	 * @param ElementList			   The list in which to add elements
-	 * @param InLayer                  The layer to draw the element on
-	 * @param PaintGeometry            DrawSpace position and dimensions; see FPaintGeometry
-	 * @param InGradientStops          List of gradient stops which define the element
-	 * @param InGradientType           The type of gradient (I.E Horizontal, vertical)
-	 * @param InDrawEffects            Optional draw effects to apply
-	 * @param CornerRadius			   Rounds the corners of the box created by the gradient by the specified radius
-	 */
-	SLATECORE_API static void MakeGradient( FSlateWindowElementList& ElementList, uint32 InLayer, const FPaintGeometry& PaintGeometry, TArray<FSlateGradientStop> InGradientStops, EOrientation InGradientType, ESlateDrawEffect InDrawEffects = ESlateDrawEffect::None, FVector4f CornerRadius = FVector4f(0.0f) );
-
-	/**
-	 * Creates a Hermite Spline element
-	 *
-	 * @param ElementList			The list in which to add elements
-	 * @param InLayer               The layer to draw the element on
-	 * @param PaintGeometry         DrawSpace position and dimensions; see FPaintGeometry
-	 * @param InStart               The start point of the spline (local space)
-	 * @param InStartDir            The direction of the spline from the start point
-	 * @param InEnd                 The end point of the spline (local space)
-	 * @param InEndDir              The direction of the spline to the end point
-	 * @param InDrawEffects         Optional draw effects to apply
-	 * @param InTint                Color to tint the element
-	 */
-	SLATECORE_API static void MakeSpline(FSlateWindowElementList& ElementList, uint32 InLayer, const FPaintGeometry& PaintGeometry, const UE::Slate::FDeprecateVector2DParameter InStart, const UE::Slate::FDeprecateVector2DParameter InStartDir, const UE::Slate::FDeprecateVector2DParameter InEnd, const UE::Slate::FDeprecateVector2DParameter InEndDir, float InThickness = 0.0f, ESlateDrawEffect InDrawEffects = ESlateDrawEffect::None, const FLinearColor& InTint=FLinearColor::White);
-
-	/**
-	 * Creates a Bezier Spline element
-	 *
-	 * @param ElementList			The list in which to add elements
-	 * @param InLayer               The layer to draw the element on
-	 * @param PaintGeometry         DrawSpace position and dimensions; see FPaintGeometry
-	 * @param InStart               The start point of the spline (local space)
-	 * @param InStartDir            The direction of the spline from the start point
-	 * @param InEnd                 The end point of the spline (local space)
-	 * @param InEndDir              The direction of the spline to the end point
-	 * @param InDrawEffects         Optional draw effects to apply
-	 * @param InTint                Color to tint the element
-	 */
-	SLATECORE_API static void MakeCubicBezierSpline(FSlateWindowElementList& ElementList, uint32 InLayer, const FPaintGeometry& PaintGeometry, const UE::Slate::FDeprecateVector2DParameter P0, const UE::Slate::FDeprecateVector2DParameter P1, const UE::Slate::FDeprecateVector2DParameter P2, const UE::Slate::FDeprecateVector2DParameter P3, float InThickness = 0.0f, ESlateDrawEffect InDrawEffects = ESlateDrawEffect::None, const FLinearColor& InTint = FLinearColor::White);
-
-	/** Just like MakeSpline but in draw-space coordinates. This is useful for connecting already-transformed widgets together. */
-	SLATECORE_API static void MakeDrawSpaceSpline(FSlateWindowElementList& ElementList, uint32 InLayer, const UE::Slate::FDeprecateVector2DParameter InStart, const UE::Slate::FDeprecateVector2DParameter InStartDir, const UE::Slate::FDeprecateVector2DParameter InEnd, const UE::Slate::FDeprecateVector2DParameter InEndDir, float InThickness = 0.0f, ESlateDrawEffect InDrawEffects = ESlateDrawEffect::None, const FLinearColor& InTint=FLinearColor::White);
-
-	/**
-	 * Creates a line defined by the provided points
-	 *
-	 * @param ElementList              The list in which to add elements
-	 * @param InLayer                  The layer to draw the element on
-	 * @param PaintGeometry            DrawSpace position and dimensions; see FPaintGeometry
-	 * @param Points                   Points that make up the lines.  The points are joined together. I.E if Points has A,B,C there the line is A-B-C.  To draw non-joining line segments call MakeLines multiple times
-	 * @param InDrawEffects            Optional draw effects to apply
-	 * @param InTint                   Color to tint the element
-	 * @param bAntialias               Should antialiasing be applied to the line?
-	 * @param Thickness                The thickness of the line
-	 */
-#if UE_ENABLE_SLATE_VECTOR_DEPRECATION_MECHANISMS
-	SLATECORE_API static void MakeLines(FSlateWindowElementList& ElementList, uint32 InLayer, const FPaintGeometry& PaintGeometry, const TArray<FVector2d>& Points, ESlateDrawEffect InDrawEffects = ESlateDrawEffect::None, const FLinearColor& InTint=FLinearColor::White, bool bAntialias = true, float Thickness = 1.0f);
-#endif
-	SLATECORE_API static void MakeLines(FSlateWindowElementList& ElementList, uint32 InLayer, const FPaintGeometry& PaintGeometry, TArray<FVector2f> Points, ESlateDrawEffect InDrawEffects = ESlateDrawEffect::None, const FLinearColor& InTint = FLinearColor::White, bool bAntialias = true, float Thickness = 1.0f);
-
-
-	/**
-	 * Creates a line defined by the provided points
-	 *
-	 * @param ElementList              The list in which to add elements
-	 * @param InLayer                  The layer to draw the element on
-	 * @param PaintGeometry            DrawSpace position and dimensions; see FPaintGeometry
-	 * @param Points                   Points that make up the lines.  The points are joined together. I.E if Points has A,B,C there the line is A-B-C.  To draw non-joining line segments call MakeLines multiple times
-	 * @param PointColors              Vertex Color for each defined points
-	 * @param InDrawEffects            Optional draw effects to apply
-	 * @param InTint                   Color to tint the element
-	 * @param bAntialias               Should antialiasing be applied to the line?
-	 * @param Thickness                The thickness of the line
-	 */
-#if UE_ENABLE_SLATE_VECTOR_DEPRECATION_MECHANISMS
-	SLATECORE_API static void MakeLines(FSlateWindowElementList& ElementList, uint32 InLayer, const FPaintGeometry& PaintGeometry, const TArray<FVector2d>& Points, const TArray<FLinearColor>& PointColors, ESlateDrawEffect InDrawEffects = ESlateDrawEffect::None, const FLinearColor& InTint=FLinearColor::White, bool bAntialias = true, float Thickness = 1.0f);
-#endif
-	SLATECORE_API static void MakeLines(FSlateWindowElementList& ElementList, uint32 InLayer, const FPaintGeometry& PaintGeometry, TArray<FVector2f> Points, TArray<FLinearColor> PointColors, ESlateDrawEffect InDrawEffects = ESlateDrawEffect::None, const FLinearColor& InTint=FLinearColor::White, bool bAntialias = true, float Thickness = 1.0f);
-
-	/**
-	 * Creates a viewport element which is useful for rendering custom data in a texture into Slate
-	 *
-	 * @param ElementList		   The list in which to add elements
-	 * @param InLayer                  The layer to draw the element on
-	 * @param PaintGeometry            DrawSpace position and dimensions; see FPaintGeometry
-	 * @param Viewport                 Interface for drawing the viewport
-	 * @param InScale                  Draw scale to apply to the entire element
-	 * @param InDrawEffects            Optional draw effects to apply
-	 * @param InTint                   Color to tint the element
-	 */
-	SLATECORE_API static void MakeViewport( FSlateWindowElementList& ElementList, uint32 InLayer, const FPaintGeometry& PaintGeometry, TSharedPtr<const ISlateViewport> Viewport, ESlateDrawEffect InDrawEffects = ESlateDrawEffect::None, const FLinearColor& InTint=FLinearColor::White );
-
-	/**
-	 * Creates a custom element which can be used to manually draw into the Slate render target with graphics API calls rather than Slate elements
-	 *
-	 * @param ElementList		   The list in which to add elements
-	 * @param InLayer                  The layer to draw the element on
-	 * @param PaintGeometry            DrawSpace position and dimensions; see FPaintGeometry
-	 * @param CustomDrawer		   Interface to a drawer which will be called when Slate renders this element
-	 */
-	SLATECORE_API static void MakeCustom( FSlateWindowElementList& ElementList, uint32 InLayer, TSharedPtr<ICustomSlateElement, ESPMode::ThreadSafe> CustomDrawer );
-	
-	SLATECORE_API static void MakeCustomVerts(FSlateWindowElementList& ElementList, uint32 InLayer, const FSlateResourceHandle& InRenderResourceHandle, const TArray<FSlateVertex>& InVerts, const TArray<SlateIndex>& InIndexes, ISlateUpdatableInstanceBuffer* InInstanceData, uint32 InInstanceOffset, uint32 InNumInstances, ESlateDrawEffect InDrawEffects = ESlateDrawEffect::None);
-
-	SLATECORE_API static void MakePostProcessPass(FSlateWindowElementList& ElementList, uint32 InLayer, const FPaintGeometry& PaintGeometry, const FVector4f& Params, int32 DownsampleAmount, const FVector4f CornerRadius = FVector4f(0.0f));
-
-	FSlateDrawElement();
-	SLATECORE_API ~FSlateDrawElement();
-
-	FORCEINLINE int32 GetLayer() const { return LayerId; }
-
-	FORCEINLINE EElementType GetElementType() const { return ElementType; }
-
-	template<typename PayloadType>
-	FORCEINLINE const PayloadType& GetDataPayload() const { return *(PayloadType*)DataPayload; }
-
-	template<typename PayloadType>
-	FORCEINLINE PayloadType& GetDataPayload() { return *(PayloadType*)DataPayload; }
-
-	FORCEINLINE const FSlateRenderTransform& GetRenderTransform() const { return RenderTransform; }
-	FORCEINLINE void SetRenderTransform(const FSlateRenderTransform& InRenderTransform) { RenderTransform = InRenderTransform; }
-
-	FORCEINLINE UE::Slate::FDeprecateVector2DResult GetPosition() const
-	{
-		return UE::Slate::FDeprecateVector2DResult(Position);
-	}
-
-	FORCEINLINE void SetPosition(UE::Slate::FDeprecateVector2DParameter InPosition)
-	{
-		Position = InPosition;
-	}
-
-	FORCEINLINE UE::Slate::FDeprecateVector2DResult GetLocalSize() const
-	{
-		return UE::Slate::FDeprecateVector2DResult(LocalSize);
-	}
-
-	FORCEINLINE float GetScale() const { return Scale; }
-	FORCEINLINE ESlateDrawEffect GetDrawEffects() const { return DrawEffects; }
-	FORCEINLINE ESlateBatchDrawFlag GetBatchFlags() const { return BatchFlags; }
-	FORCEINLINE bool IsPixelSnapped() const { return !EnumHasAllFlags(DrawEffects, ESlateDrawEffect::NoPixelSnapping); }
-
-	FORCEINLINE int32 GetPrecachedClippingIndex() const { return ClipStateHandle.GetPrecachedClipIndex(); }
-	FORCEINLINE void SetPrecachedClippingIndex(int32 InClippingIndex) { ClipStateHandle.SetPreCachedClipIndex(InClippingIndex); }
-
-	FORCEINLINE void SetCachedClippingState(const FSlateClippingState* CachedState) { ClipStateHandle.SetCachedClipState(CachedState); }
-	FORCEINLINE const FClipStateHandle& GetClippingHandle() const { return ClipStateHandle; }
-	FORCEINLINE const int8 GetSceneIndex() const { return SceneIndex; }
-
-	FORCEINLINE void SetIsCached(bool bInIsCached) { bIsCached = bInIsCached; }
-	FORCEINLINE bool IsCached() const { return bIsCached; }
-
-	FORCEINLINE FSlateLayoutTransform GetInverseLayoutTransform() const
-	{
-		return Inverse(FSlateLayoutTransform(Scale, Position));
-	}
-
-	void AddReferencedObjects(FReferenceCollector& Collector);
-
-	/**
-	 * Update element cached position with an arbitrary offset
-	 *
-	 * @param Element		   Element to update
-	 * @param InOffset         Absolute translation delta
-	 */
-	void ApplyPositionOffset(UE::Slate::FDeprecateVector2DParameter InOffset);
-
-	UE_DEPRECATED(4.23, "GetClippingIndex has been deprecated.  If you were using this please use GetPrecachedClippingIndex instead.")
-	FORCEINLINE const int32 GetClippingIndex() const { return GetPrecachedClippingIndex(); }
-
-	UE_DEPRECATED(4.23, "SetClippingIndex has been deprecated.  If you were using this please use SetPrecachedClippingIndex instead.")
-	FORCEINLINE void SetClippingIndex(const int32 InClippingIndex) { SetPrecachedClippingIndex(InClippingIndex); }
-
-private:
-	void Init(FSlateWindowElementList& ElementList, EElementType InElementType, uint32 InLayer, const FPaintGeometry& PaintGeometry, ESlateDrawEffect InDrawEffects);
-
-	static FVector2f GetRotationPoint( const FPaintGeometry& PaintGeometry, const TOptional<FVector2f>& UserRotationPoint, ERotationSpace RotationSpace );
-	static FSlateDrawElement& MakeBoxInternal(FSlateWindowElementList& ElementList, uint32 InLayer, const FPaintGeometry& PaintGeometry, const FSlateBrush* InBrush, ESlateDrawEffect InDrawEffects, const FLinearColor& InTint);
-private:
-	FSlateDataPayload* DataPayload;
-	FSlateRenderTransform RenderTransform;
-	FVector2f Position;
-	FVector2f LocalSize;
-	int32 LayerId;
-	FClipStateHandle ClipStateHandle;
-	float Scale;
-	int8 SceneIndex;
-	ESlateDrawEffect DrawEffects;
-	EElementType ElementType;
-	// Misc data
-	ESlateBatchDrawFlag BatchFlags;
-	uint8 bIsCached : 1;
-};
 
 /**
  * Data holder for info used in cached renderbatches. Primarily clipping state, vertices and indices
@@ -394,7 +84,6 @@ struct FSlateCachedElementList
 	 * Returns number of elements in all containers summed
 	 */
 	SLATECORE_API int32 NumElements();
-	
 private:
 	SLATECORE_API void DestroyCachedData();
 
@@ -465,19 +154,21 @@ struct FSlateCachedElementData
 
 	static const FSlateClippingState* GetClipStateFromParent(const FSlateClippingManager& ParentClipManager);
 
-	FORCEINLINE void ValidateWidgetOwner(TSharedPtr<FSlateCachedElementList> List, const SWidget* CurrentWidget);
+	SLATECORE_API FORCEINLINE void ValidateWidgetOwner(TSharedPtr<FSlateCachedElementList> List, const SWidget* CurrentWidget);
 
 	template<EElementType ElementType>
-	FSlateDrawElement& AddCachedElement(FSlateCachedElementsHandle& CacheHandle, const FSlateClippingManager& ParentClipManager, const SWidget* CurrentWidget)
+	TSlateDrawElement<ElementType>& AddCachedElement(FSlateCachedElementsHandle& CacheHandle, const FSlateClippingManager& ParentClipManager, const SWidget* CurrentWidget)
 	{
+		using FSlateElementType = TSlateDrawElement<ElementType>;
+
 		TSharedPtr<FSlateCachedElementList> List = CacheHandle.Ptr.Pin(); 
 
 #if WITH_SLATE_DEBUGGING
 		ValidateWidgetOwner(List, CurrentWidget);
 #endif
 
-		FSlateDrawElementContainer& Container = List->DrawElements.Get<(uint8)ElementType>();
-		FSlateDrawElement& NewElement = Container.Elements.AddDefaulted_GetRef();
+		FSlateDrawElementArray<FSlateElementType>& Container = List->DrawElements.Get<(uint8)ElementType>();
+		FSlateElementType& NewElement = Container.AddDefaulted_GetRef();
 		NewElement.SetIsCached(true);
 
 		// Check if slow vs checking a flag on the list to see if it contains new data.
@@ -579,33 +270,16 @@ public:
 	{
 		return UE::Slate::FDeprecateVector2DResult(WindowSize);
 	}
-
-	template<typename PayloadType>
-	PayloadType& CreatePayload(FSlateDrawElement& DrawElement)
-	{
-		// Elements are responsible for deleting their own payloads
-		PayloadType* Payload;
-		if (DrawElement.IsCached())
-		{
-			Payload = new PayloadType;
-		}
-		else
-		{
-			void* PayloadMem = MemManager.Alloc(sizeof(PayloadType), alignof(PayloadType));
-			Payload = new(PayloadMem) PayloadType();
-		}
-
-		DrawElement.DataPayload = Payload;
-		return *Payload;
-	}
-
+	
 	/**
 	 * Creates an uninitialized draw element if using caching will create a new cached draw list
 	 * if needed (Whenever a top level draw widget's cache handle doesn't match the current cached handle).
 	 */
 	template<EElementType ElementType = EElementType::ET_NonMapped>
-	FSlateDrawElement& AddUninitialized()
+	TSlateDrawElement<ElementType>& AddUninitialized()
 	{
+		using FSlateElementType = TSlateDrawElement<ElementType>;
+
 		const bool bAllowCache = CachedElementDataListStack.Num() > 0 && WidgetDrawStack.Num() && !WidgetDrawStack.Top().bIsVolatile;
 
 		if (bAllowCache)
@@ -616,15 +290,15 @@ public:
 		else
 		{
 			FSlateDrawElementMap& Elements = UncachedDrawElements;
-			FSlateDrawElementContainer& Container = Elements.Get<(uint8)ElementType>();
-			const int32 InsertIdx = Container.Elements.AddDefaulted();
+			FSlateDrawElementArray<FSlateElementType>& Container = Elements.Get<(uint8)ElementType>();
+			const int32 InsertIdx = Container.AddDefaulted();
 
 #if WITH_SLATE_DEBUGGING
 			FSlateDebuggingElementTypeAddedEventArgs ElementTypeAddedArgs{*this, InsertIdx, ElementType};
 			FSlateDebugging::ElementTypeAdded.Broadcast(ElementTypeAddedArgs);
 #endif
 
-			FSlateDrawElement& NewElement = Container.Elements[InsertIdx];
+			FSlateElementType& NewElement = Container[InsertIdx];
 			return NewElement;
 		}
 	}
@@ -753,7 +427,7 @@ public:
 private:
 	/** Adds a cached element, generating a new cached list for the widget at the top of the cache if needed */
 	template<EElementType ElementType>
-	FSlateDrawElement& AddCachedElement()
+	TSlateDrawElement<ElementType>& AddCachedElement()
 	{
 		FSlateCachedElementData* CurrentCachedElementData = GetCurrentCachedElementData();
 		check(CurrentCachedElementData);
diff --git a/Engine/Source/Runtime/SlateCore/Public/Rendering/ElementBatcher.h b/Engine/Source/Runtime/SlateCore/Public/Rendering/ElementBatcher.h
--- a/Engine/Source/Runtime/SlateCore/Public/Rendering/ElementBatcher.h
+++ b/Engine/Source/Runtime/SlateCore/Public/Rendering/ElementBatcher.h
@@ -283,8 +283,8 @@ private:
 	 * @param InElementBatchParamCreator - Functor that generates batch params given a slate draw element, used during batch-reuse
 	 * @param InElementBatchReserver - Functor that reserves vertexes and indicies given an element range and list of elements
 	 */
-	template<typename ElementAdder, typename ElementBatchParamCreator, typename ElementBatchReserver>
-	FORCEINLINE void GenerateIndexedVertexBatches(const FSlateDrawElementArray& DrawElements
+	template<typename ElementType, typename ElementAdder, typename ElementBatchParamCreator, typename ElementBatchReserver>
+	FORCEINLINE void GenerateIndexedVertexBatches(const FSlateDrawElementArray<ElementType>& DrawElements
 		, ElementAdder&& InElementAdder
 		, ElementBatchParamCreator&& InElementBatchParamCreator
 		, ElementBatchReserver&& InElementBatchReserver);
@@ -293,58 +293,59 @@ private:
 	 * Creates vertices necessary to draw a Quad element 
 	 */
 	template<ESlateVertexRounding Rounding>
-	void AddDebugQuadElement( const FSlateDrawElement& DrawElement);
+	void AddDebugQuadElement( const FSlateBoxElement& DrawElement);
 
 	/** 
 	 * Creates vertices necessary to draw multiple 3x3 elements
 	 */
-	void AddBoxElements( const FSlateDrawElementArray& DrawElement );
+	template<typename ElementType>
+	void AddBoxElements( const FSlateDrawElementArray<ElementType>& DrawElement );
 
 	/** 
 	 * Creates vertices necessary to draw a string (one quad per character)
 	 */
 	template<ESlateVertexRounding Rounding>
-	void AddTextElement( const FSlateDrawElement& DrawElement );
+	void AddTextElement( const FSlateTextElement& DrawElement );
 
 	/** 
 	 * Creates vertices necessary to draw a shaped glyph sequence (one quad per glyph)
 	 */
 	template<ESlateVertexRounding Rounding>
-	void AddShapedTextElement( const FSlateDrawElement& DrawElement );
+	void AddShapedTextElement( const FSlateShapedTextElement& DrawElement );
 
 	/** 
 	 * Creates vertices necessary to draw a gradient box (horizontal or vertical)
 	 */
 	template<ESlateVertexRounding Rounding>
-	void AddGradientElement( const FSlateDrawElement& DrawElement );
+	void AddGradientElement( const FSlateGradientElement& DrawElement );
 
 	/** 
 	 * Creates vertices necessary to draw a spline (Bezier curve)
 	 */
-	void AddSplineElement( const FSlateDrawElement& DrawElement );
+	void AddSplineElement( const FSlateSplineElement& DrawElement );
 
 	/** 
 	 * Creates vertices necessary to draw a multiple attached line segments
 	 */
-	void AddLineElements(const FSlateDrawElementArray& DrawElements);
+	void AddLineElements(const FSlateDrawElementArray<FSlateLineElement>& DrawElements);
 	
 	/** 
 	 * Creates vertices necessary to draw a viewport (just a textured quad)
 	 */
 	template<ESlateVertexRounding Rounding>
-	void AddViewportElement( const FSlateDrawElement& DrawElement );
+	void AddViewportElement( const FSlateViewportElement& DrawElement );
 
 	/** 
 	 * Creates vertices necessary to draw a border element
 	 */
 	template<ESlateVertexRounding Rounding>
-	void AddBorderElement( const FSlateDrawElement& DrawElement );
+	void AddBorderElement( const FSlateBoxElement& DrawElement );
 
-	void AddCustomElement( const FSlateDrawElement& DrawElement );
+	void AddCustomElement( const FSlateCustomDrawerElement& DrawElement );
 
-	void AddCustomVerts( const FSlateDrawElement& DrawElement );
+	void AddCustomVerts( const FSlateCustomVertsElement& DrawElement );
 
-	void AddPostProcessPass(const FSlateDrawElement& DrawElement, FVector2f WindowSize);
+	void AddPostProcessPass(const FSlatePostProcessElement& DrawElement, FVector2f WindowSize);
 
 	FSlateRenderBatch& CreateRenderBatch(
 		int32 Layer,
diff --git a/Engine/Source/Runtime/SlateCore/Public/Rendering/RenderingCommon.h b/Engine/Source/Runtime/SlateCore/Public/Rendering/RenderingCommon.h
--- a/Engine/Source/Runtime/SlateCore/Public/Rendering/RenderingCommon.h
+++ b/Engine/Source/Runtime/SlateCore/Public/Rendering/RenderingCommon.h
@@ -10,7 +10,7 @@
 #include "Input/Reply.h"
 #include "Input/NavigationReply.h"
 #include "Input/PopupMethodReply.h"
-#include "Rendering/DrawElementTypes.h"
+#include "Rendering/DrawElementCoreTypes.h"
 #include "SlateGlobals.h"
 #include <utility>
 
@@ -20,10 +20,6 @@ class FSlateInstanceBufferUpdate;
 class FWidgetStyle;
 class SWidget;
 
-
-DECLARE_MEMORY_STAT_EXTERN(TEXT("Vertex/Index Buffer Pool Memory (CPU)"), STAT_SlateBufferPoolMemory, STATGROUP_SlateMemory, SLATECORE_API);
-DECLARE_MEMORY_STAT_EXTERN(TEXT("Cached Draw Element Memory (CPU)"), STAT_SlateCachedDrawElementMemory, STATGROUP_SlateMemory, SLATECORE_API);
-
 DECLARE_DWORD_COUNTER_STAT_EXTERN(TEXT("Num Cached Element Lists"), STAT_SlateNumCachedElementLists, STATGROUP_Slate, SLATECORE_API);
 DECLARE_DWORD_COUNTER_STAT_EXTERN(TEXT("Num Cached Elements"), STAT_SlateNumCachedElements, STATGROUP_Slate, SLATECORE_API);
 
@@ -419,21 +415,6 @@ static_assert(TIsTriviallyDestructible<FShortRect>::Value == true, "FShortRect s
 
 namespace UE::Slate
 {
-	template<std::size_t, typename T>
-	using TypeRepeater = T;
-
-	template<typename ElementType, std::size_t... Indices>
-	auto MakeTupleWithNumInner(std::index_sequence<Indices...>)
-	{
-		return TTuple<TypeRepeater<Indices, ElementType>... >();
-	};
-
-	template<typename ElementType, std::size_t Num, typename Indices = std::make_index_sequence<Num>>
-	auto MakeTupleWithNum()
-	{
-		return MakeTupleWithNumInner<ElementType>(Indices{});
-	}
-
 	template<typename IndexType, IndexType...Indices>
 	auto MakeTupleIndiciesInner(std::integer_sequence<IndexType, Indices...>)
 	{
@@ -447,151 +428,21 @@ namespace UE::Slate
 	};
 };
 
+/**
+ * Note: FRenderingBufferStatTracker & FSlateDrawElementArray have been moved to DrawElementCoreTypes.h
+ */
 
 #if STATS
 
-struct FRenderingBufferStatTracker
-{
-	static void MemoryAllocated(int32 SizeBytes)
-	{
-		INC_DWORD_STAT_BY(STAT_SlateBufferPoolMemory, SizeBytes);
-	}
-
-	static void MemoryFreed(int32 SizeBytes)
-	{
-		DEC_DWORD_STAT_BY(STAT_SlateBufferPoolMemory, SizeBytes);
-	}
-};
-
-struct FDrawElementStatTracker
-{
-	static void MemoryAllocated(int32 SizeBytes)
-	{
-		INC_DWORD_STAT_BY(STAT_SlateCachedDrawElementMemory, SizeBytes);
-	}
-
-	static void MemoryFreed(int32 SizeBytes)
-	{
-		DEC_DWORD_STAT_BY(STAT_SlateCachedDrawElementMemory, SizeBytes);
-	}
-};
-
-template<typename StatTracker>
-class FSlateStatTrackingMemoryAllocator : public FDefaultAllocator
-{
-public:
-	typedef FDefaultAllocator Super;
-
-	class ForAnyElementType : public FDefaultAllocator::ForAnyElementType
-	{
-	public:
-		typedef FDefaultAllocator::ForAnyElementType Super;
-
-		ForAnyElementType()
-			: AllocatedSize(0)
-		{
-
-		}
-
-		/**
-		* Moves the state of another allocator into this one.
-		* Assumes that the allocator is currently empty, i.e. memory may be allocated but any existing elements have already been destructed (if necessary).
-		* @param Other - The allocator to move the state from.  This allocator should be left in a valid empty state.
-		*/
-		FORCEINLINE void MoveToEmpty(ForAnyElementType& Other)
-		{
-			Super::MoveToEmpty(Other);
-
-			AllocatedSize = Other.AllocatedSize;
-			Other.AllocatedSize = 0;
-		}
-
-		/** Destructor. */
-		~ForAnyElementType()
-		{
-			if (AllocatedSize)
-			{
-				StatTracker::MemoryFreed(AllocatedSize);
-			}
-		}
-
-		void ResizeAllocation(int32 PreviousNumElements, int32 NumElements, int32 NumBytesPerElement)
-		{
-			const int32 NewSize = NumElements * NumBytesPerElement;
-			StatTracker::MemoryAllocated(NewSize - AllocatedSize);
-
-			AllocatedSize = NewSize;
-
-			Super::ResizeAllocation(PreviousNumElements, NumElements, NumBytesPerElement);
-		}
-
-	private:
-		ForAnyElementType(const ForAnyElementType&);
-		ForAnyElementType& operator=(const ForAnyElementType&);
-	private:
-		int32 AllocatedSize;
-	};
-};
-
-template <typename T>
-struct TAllocatorTraits<FSlateStatTrackingMemoryAllocator<T>> : TAllocatorTraitsBase<FSlateStatTrackingMemoryAllocator<T>>
-{
-	enum { IsZeroConstruct = TAllocatorTraits<FDefaultAllocator>::IsZeroConstruct };
-};
-
 typedef TArray<FSlateVertex, FSlateStatTrackingMemoryAllocator<FRenderingBufferStatTracker>> FSlateVertexArray;
 typedef TArray<SlateIndex, FSlateStatTrackingMemoryAllocator<FRenderingBufferStatTracker>> FSlateIndexArray;
-typedef TArray<FSlateDrawElement, FSlateStatTrackingMemoryAllocator<FDrawElementStatTracker>> FSlateDrawElementArray;
-
-struct FSlateDrawElementContainer
-{
-	FSlateDrawElementArray Elements;
-};
-
-using FSlateDrawElementMap = TTuple<
-	FSlateDrawElementContainer
-	, FSlateDrawElementContainer
-	, FSlateDrawElementContainer
-	, FSlateDrawElementContainer
-	, FSlateDrawElementContainer
-	, FSlateDrawElementContainer
-	, FSlateDrawElementContainer
-	, FSlateDrawElementContainer
-	, FSlateDrawElementContainer
-	, FSlateDrawElementContainer
-	, FSlateDrawElementContainer
-	, FSlateDrawElementContainer
-	, FSlateDrawElementContainer
-	, FSlateDrawElementContainer>;
 
 #else
 
 typedef TArray<FSlateVertex> FSlateVertexArray;
 typedef TArray<SlateIndex> FSlateIndexArray;
-typedef TArray<FSlateDrawElement> FSlateDrawElementArray;
 
-struct FSlateDrawElementContainer
-{
-	FSlateDrawElementArray Elements;
-};
-
-using FSlateDrawElementMap = TTuple<
-	FSlateDrawElementContainer
-	, FSlateDrawElementContainer
-	, FSlateDrawElementContainer
-	, FSlateDrawElementContainer
-	, FSlateDrawElementContainer
-	, FSlateDrawElementContainer
-	, FSlateDrawElementContainer
-	, FSlateDrawElementContainer
-	, FSlateDrawElementContainer
-	, FSlateDrawElementContainer
-	, FSlateDrawElementContainer
-	, FSlateDrawElementContainer
-	, FSlateDrawElementContainer
-	, FSlateDrawElementContainer>;
-
-#endif
+#endif // STATS
 
 /**
  * Viewport implementation interface that is used by SViewport when it needs to draw and processes input.                   
