Avoid element stats on empty sets.

diff --git a/Engine/Source/Runtime/SlateCore/Private/Rendering/ElementBatcher.cpp b/Engine/Source/Runtime/SlateCore/Private/Rendering/ElementBatcher.cpp
--- a/Engine/Source/Runtime/SlateCore/Private/Rendering/ElementBatcher.cpp
+++ b/Engine/Source/Runtime/SlateCore/Private/Rendering/ElementBatcher.cpp
@@ -356,127 +356,139 @@ void FSlateElementBatcher::AddElements(FSlateWindowElementList& WindowElementLis
 #endif
 }
 
-
 void FSlateElementBatcher::AddElementsInternal(const FSlateDrawElementMap& DrawElements, FVector2f ViewportSize)
 {
+	const FSlateDrawElementArray<FSlateBoxElement>& BoxElements = DrawElements.Get<(uint8)EElementType::ET_Box>();
+	if (BoxElements.Num() > 0)
 	{
 		SCOPED_NAMED_EVENT_TEXT("Slate::AddBoxElements", FColor::Magenta);
-		const FSlateDrawElementArray<FSlateBoxElement>& Container = DrawElements.Get<(uint8)EElementType::ET_Box>();
-		STAT(ElementStat_Other += Container.Num());
-		AddBoxElements(Container);
+		STAT(ElementStat_Boxes += BoxElements.Num());
+		AddBoxElements(BoxElements);
 	}
 
+	const FSlateDrawElementArray<FSlateRoundedBoxElement>& RoundedBoxElements = DrawElements.Get<(uint8)EElementType::ET_RoundedBox>();
+	if (RoundedBoxElements.Num() > 0)
 	{
 		SCOPED_NAMED_EVENT_TEXT("Slate::AddRoundedBoxElements", FColor::Magenta);
-		const FSlateDrawElementArray<FSlateRoundedBoxElement>& Container = DrawElements.Get<(uint8)EElementType::ET_RoundedBox>();
-		STAT(ElementStat_Other += Container.Num());
-		AddBoxElements(Container);
+		STAT(ElementStat_Boxes += RoundedBoxElements.Num());
+		AddBoxElements(RoundedBoxElements);
 	}
 
+	const FSlateDrawElementArray<FSlateBoxElement>& BorderElements = DrawElements.Get<(uint8)EElementType::ET_Border>();
+	if (BorderElements.Num() > 0)
 	{
 		SCOPED_NAMED_EVENT_TEXT("Slate::AddBorderElement", FColor::Magenta);
-		const FSlateDrawElementArray<FSlateBoxElement>& Container = DrawElements.Get<(uint8)EElementType::ET_Border>();
-		STAT(ElementStat_Other += Container.Num());
-		for (const FSlateBoxElement& DrawElement : Container)
+		STAT(ElementStat_Borders += BorderElements.Num());
+		for (const FSlateBoxElement& DrawElement : BorderElements)
 		{
 			DrawElement.IsPixelSnapped() ? AddBorderElement<ESlateVertexRounding::Enabled>(DrawElement) : AddBorderElement<ESlateVertexRounding::Disabled>(DrawElement);
 		}
 	}
 
+	const FSlateDrawElementArray<FSlateTextElement>& TextElements = DrawElements.Get<(uint8)EElementType::ET_Text>();
+	if (TextElements.Num() > 0)
 	{
 		SCOPED_NAMED_EVENT_TEXT("Slate::AddTextElement", FColor::Magenta);
-		const FSlateDrawElementArray<FSlateTextElement>& Container = DrawElements.Get<(uint8)EElementType::ET_Text>();
-		STAT(ElementStat_Other += Container.Num());
-		for (const FSlateTextElement& DrawElement : Container)
+		STAT(ElementStat_Text += TextElements.Num());
+		for (const FSlateTextElement& DrawElement : TextElements)
 		{
 			DrawElement.IsPixelSnapped() ? AddTextElement<ESlateVertexRounding::Enabled>(DrawElement) : AddTextElement<ESlateVertexRounding::Disabled>(DrawElement);
 		}
 	}
 
+	const FSlateDrawElementArray<FSlateShapedTextElement>& ShapedTextElements = DrawElements.Get<(uint8)EElementType::ET_ShapedText>();
+	if (ShapedTextElements.Num() > 0)
 	{
 		SCOPED_NAMED_EVENT_TEXT("Slate::AddShapedTextElement", FColor::Magenta);
-		const FSlateDrawElementArray<FSlateShapedTextElement>& Container = DrawElements.Get<(uint8)EElementType::ET_ShapedText>();
-		STAT(ElementStat_Other += Container.Num());
-		for (const FSlateShapedTextElement& DrawElement : Container)
+		STAT(ElementStat_ShapedText += ShapedTextElements.Num());
+		for (const FSlateShapedTextElement& DrawElement : ShapedTextElements)
 		{
 			DrawElement.IsPixelSnapped() ? AddShapedTextElement<ESlateVertexRounding::Enabled>(DrawElement) : AddShapedTextElement<ESlateVertexRounding::Disabled>(DrawElement);
 		}
 	}
 
+	const FSlateDrawElementArray<FSlateLineElement>& LineElements = DrawElements.Get<(uint8)EElementType::ET_Line>();
+	if (LineElements.Num() > 0)
 	{
 		SCOPED_NAMED_EVENT_TEXT("Slate::AddLineElements", FColor::Magenta);
-		const FSlateDrawElementArray<FSlateLineElement>& Container = DrawElements.Get<(uint8)EElementType::ET_Line>();
-		STAT(ElementStat_Other += Container.Num());
-		AddLineElements(Container);
+		STAT(ElementStat_Line += LineElements.Num());
+		AddLineElements(LineElements);
 	}
 
+	const FSlateDrawElementArray<FSlateBoxElement>& DebugQuadElements = DrawElements.Get<(uint8)EElementType::ET_DebugQuad>();
+	if (DebugQuadElements.Num() > 0)
 	{
 		SCOPED_NAMED_EVENT_TEXT("Slate::AddDebugQuadElement", FColor::Magenta);
-		const FSlateDrawElementArray<FSlateBoxElement>& Container = DrawElements.Get<(uint8)EElementType::ET_DebugQuad>();
-		STAT(ElementStat_Other += Container.Num());
-		for (const FSlateBoxElement& DrawElement : Container)
+		STAT(ElementStat_Other += DebugQuadElements.Num());
+		for (const FSlateBoxElement& DrawElement : DebugQuadElements)
 		{
 			DrawElement.IsPixelSnapped() ? AddDebugQuadElement<ESlateVertexRounding::Enabled>(DrawElement) : AddDebugQuadElement<ESlateVertexRounding::Disabled>(DrawElement);
 		}
 	}
 
+	const FSlateDrawElementArray<FSlateSplineElement>& SplineElements = DrawElements.Get<(uint8)EElementType::ET_Spline>();
+	if (SplineElements.Num() > 0)
 	{
 		SCOPED_NAMED_EVENT_TEXT("Slate::AddSplineElement", FColor::Magenta);
-		const FSlateDrawElementArray<FSlateSplineElement>& Container = DrawElements.Get<(uint8)EElementType::ET_Spline>();
 
 		// Note that we ignore pixel snapping here; see implementation for more info.
-		STAT(ElementStat_Other += Container.Num());
-		for (const FSlateSplineElement& DrawElement : Container)
+		STAT(ElementStat_Other += SplineElements.Num());
+		for (const FSlateSplineElement& DrawElement : SplineElements)
 		{
 			AddSplineElement(DrawElement);
 		}
 	}
 
+	const FSlateDrawElementArray<FSlateGradientElement>& GradientElements = DrawElements.Get<(uint8)EElementType::ET_Gradient>();
+	if (GradientElements.Num() > 0)
 	{
 		SCOPED_NAMED_EVENT_TEXT("Slate::AddGradientElement", FColor::Magenta);
-		const FSlateDrawElementArray<FSlateGradientElement>& Container = DrawElements.Get<(uint8)EElementType::ET_Gradient>();
-		STAT(ElementStat_Other += Container.Num());
-		for (const FSlateGradientElement& DrawElement : Container)
+		STAT(ElementStat_Other += GradientElements.Num());
+		for (const FSlateGradientElement& DrawElement : GradientElements)
 		{
 			DrawElement.IsPixelSnapped() ? AddGradientElement<ESlateVertexRounding::Enabled>(DrawElement) : AddGradientElement<ESlateVertexRounding::Disabled>(DrawElement);
 		}
 	}
 
+	const FSlateDrawElementArray<FSlateViewportElement>& ViewportElements = DrawElements.Get<(uint8)EElementType::ET_Viewport>();
+	if (ViewportElements.Num() > 0)
 	{
 		SCOPED_NAMED_EVENT_TEXT("Slate::AddViewportElement", FColor::Magenta);
-		const FSlateDrawElementArray<FSlateViewportElement>& Container = DrawElements.Get<(uint8)EElementType::ET_Viewport>();
-		STAT(ElementStat_Other += Container.Num());
-		for (const FSlateViewportElement& DrawElement : Container)
+		STAT(ElementStat_Other += ViewportElements.Num());
+		for (const FSlateViewportElement& DrawElement : ViewportElements)
 		{
 			DrawElement.IsPixelSnapped() ? AddViewportElement<ESlateVertexRounding::Enabled>(DrawElement) : AddViewportElement<ESlateVertexRounding::Disabled>(DrawElement);
 		}
 	}
 
+	const FSlateDrawElementArray<FSlateCustomDrawerElement>& CustomElements = DrawElements.Get<(uint8)EElementType::ET_Custom>();
+	if (CustomElements.Num() > 0)
 	{
 		SCOPED_NAMED_EVENT_TEXT("Slate::AddCustomElement", FColor::Magenta);
-		const FSlateDrawElementArray<FSlateCustomDrawerElement>& Container = DrawElements.Get<(uint8)EElementType::ET_Custom>();
-		STAT(ElementStat_Other += Container.Num());
-		for (const FSlateCustomDrawerElement& DrawElement : Container)
+		STAT(ElementStat_Other += CustomElements.Num());
+		for (const FSlateCustomDrawerElement& DrawElement : CustomElements)
 		{
 			AddCustomElement(DrawElement);
 		}
 	}
 
+	const FSlateDrawElementArray<FSlateCustomVertsElement>& CustomVertElements = DrawElements.Get<(uint8)EElementType::ET_CustomVerts>();
+	if (CustomVertElements.Num() > 0)
 	{
 		SCOPED_NAMED_EVENT_TEXT("Slate::AddCustomVertsElement", FColor::Magenta);
-		const FSlateDrawElementArray<FSlateCustomVertsElement>& Container = DrawElements.Get<(uint8)EElementType::ET_CustomVerts>();
-		STAT(ElementStat_Other += Container.Num());
-		for (const FSlateCustomVertsElement& DrawElement : Container)
+		STAT(ElementStat_Other += CustomVertElements.Num());
+		for (const FSlateCustomVertsElement& DrawElement : CustomVertElements)
 		{
 			AddCustomVerts(DrawElement);
 		}
 	}
 
+	const FSlateDrawElementArray<FSlatePostProcessElement>& PostProcessElements = DrawElements.Get<(uint8)EElementType::ET_PostProcessPass>();
+	if (PostProcessElements.Num() > 0)
 	{
 		SCOPED_NAMED_EVENT_TEXT("Slate::AddPostProcessElement", FColor::Magenta);
-		const FSlateDrawElementArray<FSlatePostProcessElement>& Container = DrawElements.Get<(uint8)EElementType::ET_PostProcessPass>();
-		STAT(ElementStat_Other += Container.Num());
-		for (const FSlatePostProcessElement& DrawElement : Container)
+		STAT(ElementStat_Other += PostProcessElements.Num());
+		for (const FSlatePostProcessElement& DrawElement : PostProcessElements)
 		{
 			AddPostProcessPass(DrawElement, ViewportSize);
 		}
