Fixed SlateTextureAtlas creating free nodes in the empty slot list that have dimensions smaller than what we feasibly expect can be used by glyphs.

The code presently creates two nodes as long as one of them exceeds the minimum dimension, which in practice this has resulted in us eventually having to iterate over 94+ useless elements every time we're adding a glyph. We should only create the smaller node if it exceeds the minimum.

diff --git a/Engine/Source/Runtime/SlateCore/Private/Textures/TextureAtlas.cpp b/Engine/Source/Runtime/SlateCore/Private/Textures/TextureAtlas.cpp
--- a/Engine/Source/Runtime/SlateCore/Private/Textures/TextureAtlas.cpp
+++ b/Engine/Source/Runtime/SlateCore/Private/Textures/TextureAtlas.cpp
@@ -308,8 +308,8 @@ const FAtlasedTextureSlot* FSlateTextureAtlas::FindSlotForTexture(uint32 InWidth
 		// Split the remaining area around this slot into two children.
 		if (RemainingHeight >= MinSlotDim || RemainingWidth >= MinSlotDim)
 		{
-			FAtlasedTextureSlot* LeftSlot = NULL;
-			FAtlasedTextureSlot* RightSlot = NULL;
+			FAtlasedTextureSlot* LeftSlot = nullptr;
+			FAtlasedTextureSlot* RightSlot = nullptr;
 
 			if (RemainingHeight <= RemainingWidth)
 			{
@@ -323,7 +323,10 @@ const FAtlasedTextureSlot* FSlateTextureAtlas::FindSlotForTexture(uint32 InWidth
 				// |  Left |       |
 				// |       |       |
 				// - - - - - - - - -
-				LeftSlot = new FAtlasedTextureSlot(ReturnVal->X, ReturnVal->Y + PaddedHeight, PaddedWidth, RemainingHeight, Padding);
+				if (RemainingHeight >= MinSlotDim)
+				{
+					LeftSlot = new FAtlasedTextureSlot(ReturnVal->X, ReturnVal->Y + PaddedHeight, PaddedWidth, RemainingHeight, Padding);
+				}
 				RightSlot = new FAtlasedTextureSlot(ReturnVal->X + PaddedWidth, ReturnVal->Y, RemainingWidth, ReturnVal->Height, Padding);
 			}
 			else
@@ -338,13 +341,23 @@ const FAtlasedTextureSlot* FSlateTextureAtlas::FindSlotForTexture(uint32 InWidth
 				// |     Right     |
 				// |               |
 				// - - - - - - - - -
-				LeftSlot = new FAtlasedTextureSlot(ReturnVal->X + PaddedWidth, ReturnVal->Y, RemainingWidth, PaddedHeight, Padding);
+				if (RemainingWidth >= MinSlotDim)
+				{
+					LeftSlot = new FAtlasedTextureSlot(ReturnVal->X + PaddedWidth, ReturnVal->Y, RemainingWidth, PaddedHeight, Padding);
+				}
 				RightSlot = new FAtlasedTextureSlot(ReturnVal->X, ReturnVal->Y + PaddedHeight, ReturnVal->Width, RemainingHeight, Padding);
 			}
 
 			// Replace the old slot within AtlasEmptySlots, with the new Left and Right slot, then add the old slot to AtlasUsedSlots
-			LeftSlot->LinkReplace(ReturnVal);
-			RightSlot->LinkAfter(LeftSlot);
+			if (LeftSlot)
+			{
+				LeftSlot->LinkReplace(ReturnVal);
+				RightSlot->LinkAfter(LeftSlot);
+			}
+			else
+			{
+				RightSlot->LinkReplace(ReturnVal);
+			}
 
 			ReturnVal->LinkHead(AtlasUsedSlots);
 		}
