Fix ZOrder of AddToViewport not being taken into consideration

diff --git a/Engine/Source/Runtime/SlateCore/Private/Widgets/SOverlay.cpp b/Engine/Source/Runtime/SlateCore/Private/Widgets/SOverlay.cpp
--- a/Engine/Source/Runtime/SlateCore/Private/Widgets/SOverlay.cpp
+++ b/Engine/Source/Runtime/SlateCore/Private/Widgets/SOverlay.cpp
@@ -16,7 +16,7 @@ void SOverlay::PrivateRegisterAttributes(FSlateAttributeInitializer& AttributeIn
 PRAGMA_DISABLE_DEPRECATION_WARNINGS
 void SOverlay::FOverlaySlot::Construct(const FChildren& SlotOwner, FSlotArguments&& InArgs)
 {
-	ZOrder = InArgs._ZOrder.Get(0);
+	ZOrder = InArgs._ZOrder.Get(ZOrder);
 	TBasicLayoutWidgetSlot<FOverlaySlot>::Construct(SlotOwner, MoveTemp(InArgs));
 }
 
