PR #10567: Add OutlineSettings to FSlateBursh == operator

diff --git a/Engine/Source/Runtime/SlateCore/Public/Styling/SlateBrush.h b/Engine/Source/Runtime/SlateCore/Public/Styling/SlateBrush.h
--- a/Engine/Source/Runtime/SlateCore/Public/Styling/SlateBrush.h
+++ b/Engine/Source/Runtime/SlateCore/Public/Styling/SlateBrush.h
@@ -182,6 +182,22 @@ struct FSlateBrushOutlineSettings
 		, bUseBrushTransparency(false)
 	{}
 
+	/**
+	 * Compares these outline settings with another for equality.
+	 *
+	 * @param Other The other outline settings.
+	 *
+	 * @return true if settings are equal, false otherwise.
+	 */
+	bool operator==(const FSlateBrushOutlineSettings& Other) const
+	{
+		return CornerRadii == Other.CornerRadii
+			&& Color == Other.Color
+			&& Width == Other.Width
+			&& RoundingType == Other.RoundingType
+			&& bUseBrushTransparency == Other.bUseBrushTransparency;
+	}
+
 	/** Radius in Slate Units applied to the outline at each corner. X = Top Left, Y = Top Right, Z = Bottom Right, W = Bottom Left */
 	UPROPERTY(EditAnywhere, BlueprintReadWrite, Category=Brush)
 	FVector4 CornerRadii;
@@ -409,7 +425,8 @@ public:
 			&& ResourceObject == Other.ResourceObject
 			&& ResourceName == Other.ResourceName
 			&& bIsDynamicallyLoaded == Other.bIsDynamicallyLoaded
-			&& UVRegion == Other.UVRegion;
+			&& UVRegion == Other.UVRegion
+			&& (DrawAs != ESlateBrushDrawType::RoundedBox || OutlineSettings == Other.OutlineSettings); // Compare outline settings for equality only if we have a rounded box brush.
 	}
 
 	/**
