PR #10490: GetDynamicImageBrush returning null upon successful brush creation in FSlateStyleSet

diff --git a/Engine/Source/Runtime/SlateCore/Private/Styling/SlateStyleSet.cpp b/Engine/Source/Runtime/SlateCore/Private/Styling/SlateStyleSet.cpp
--- a/Engine/Source/Runtime/SlateCore/Private/Styling/SlateStyleSet.cpp
+++ b/Engine/Source/Runtime/SlateCore/Private/Styling/SlateStyleSet.cpp
@@ -383,7 +383,7 @@ const TSharedPtr< FSlateDynamicImageBrush > FSlateStyleSet::GetDynamicImageBrush
 	if ( !ReturnBrush.IsValid() )
 	{
 		const ISlateStyle* ReportingStyle = RequestingStyle != nullptr ? RequestingStyle : this;
-		ReportingStyle->MakeDynamicImageBrush(BrushTemplate, TextureResource, TextureName);
+		ReturnBrush = ReportingStyle->MakeDynamicImageBrush(BrushTemplate, TextureResource, TextureName);
 	}
 
 	return ReturnBrush;
