Fix regression on flow direction inheritence when using cvar GSlateFlowDirectionShouldFollowCultureByDefault (was ignored), introduced when fixing another flow issue in CL 23445460

#test Did repro steps from the associated UDN, and from the UDN associated to the CL that introduced the regression. Both are fixed.

diff --git a/Engine/Source/Runtime/SlateCore/Private/FastUpdate/SlateInvalidationRoot.cpp b/Engine/Source/Runtime/SlateCore/Private/FastUpdate/SlateInvalidationRoot.cpp
--- a/Engine/Source/Runtime/SlateCore/Private/FastUpdate/SlateInvalidationRoot.cpp
+++ b/Engine/Source/Runtime/SlateCore/Private/FastUpdate/SlateInvalidationRoot.cpp
@@ -375,7 +375,12 @@ FSlateInvalidationResult FSlateInvalidationRoot::PaintInvalidationRoot(const FSl
 		bNeedScreenPositionShift = false;
 	}
 
-	TGuardValue<EFlowDirection> FlowGuard(GSlateFlowDirection, RootWidget->ComputeFlowDirection());
+	EFlowDirection NewFlowDirection = GSlateFlowDirection;
+	if (RootWidget->GetFlowDirectionPreference() == EFlowDirectionPreference::Inherit)
+	{
+		NewFlowDirection = GSlateFlowDirectionShouldFollowCultureByDefault ? FLayoutLocalization::GetLocalizedLayoutDirection() : RootWidget->ComputeFlowDirection();
+	}
+	TGuardValue<EFlowDirection> FlowGuard(GSlateFlowDirection, NewFlowDirection);
 	if (!Context.bAllowFastPathUpdate || bNeedsSlowPath || GSlateIsInInvalidationSlowPath)
 	{
 		SCOPED_NAMED_EVENT(Slate_PaintSlowPath, FColor::Red);
