[SlateCore]
* Removed 7 million allocations when starting game client

diff --git a/Engine/Source/Runtime/SlateCore/Private/Input/HittestGrid.cpp b/Engine/Source/Runtime/SlateCore/Private/Input/HittestGrid.cpp
--- a/Engine/Source/Runtime/SlateCore/Private/Input/HittestGrid.cpp
+++ b/Engine/Source/Runtime/SlateCore/Private/Input/HittestGrid.cpp
@@ -145,6 +145,11 @@ void FHittestGrid::FCell::RemoveIndex(int32 WidgetIndex)
 	WidgetIndexes.RemoveSingleSwap(WidgetIndex);
 }
 
+void FHittestGrid::FCell::Reset()
+{
+	WidgetIndexes.Reset();
+}
+
 //
 // FHittestGrid
 //
@@ -270,7 +275,8 @@ void FHittestGrid::ClearInternal(int32 TotalCells)
 #if UE_SLATE_ENABLE_HITTEST_STATS
 	SCOPE_CYCLE_COUNTER(STAT_SlateHTG_Clear);
 #endif
-	Cells.Reset(TotalCells);
+	for (FCell& Cell : Cells)
+		Cell.Reset();
 	Cells.SetNumZeroed(TotalCells);
 
 	WidgetMap.Reset();
diff --git a/Engine/Source/Runtime/SlateCore/Public/Input/HittestGrid.h b/Engine/Source/Runtime/SlateCore/Public/Input/HittestGrid.h
--- a/Engine/Source/Runtime/SlateCore/Public/Input/HittestGrid.h
+++ b/Engine/Source/Runtime/SlateCore/Public/Input/HittestGrid.h
@@ -243,6 +243,7 @@ private:
 
 		void AddIndex(int32 WidgetIndex);
 		void RemoveIndex(int32 WidgetIndex);
+		void Reset();
 
 		const TArray<int32>& GetWidgetIndexes() const { return WidgetIndexes; }
 		
