Honor the requested kerning flags

#rnx

diff --git a/Engine/Source/Runtime/SlateCore/Private/Fonts/FontCacheFreeType.cpp b/Engine/Source/Runtime/SlateCore/Private/Fonts/FontCacheFreeType.cpp
--- a/Engine/Source/Runtime/SlateCore/Private/Fonts/FontCacheFreeType.cpp
+++ b/Engine/Source/Runtime/SlateCore/Private/Fonts/FontCacheFreeType.cpp
@@ -770,11 +770,11 @@ TSharedPtr<FFreeTypeKerningCache> FFreeTypeCacheDirectory::GetKerningCache(FT_Fa
 	// We also can't perform kerning between two separate font faces
 	if (InFace && FT_HAS_KERNING(InFace))
 	{
-		const FFontKey Key(InFace, FT_KERNING_DEFAULT, InFontSize, InFontScale);
+		const FFontKey Key(InFace, InKerningFlags, InFontSize, InFontScale);
 		TSharedPtr<FFreeTypeKerningCache>& Entry = KerningCacheMap.FindOrAdd(Key);
 		if (!Entry)
 		{
-			Entry = MakeShared<FFreeTypeKerningCache>(InFace, FT_KERNING_DEFAULT, InFontSize, InFontScale);
+			Entry = MakeShared<FFreeTypeKerningCache>(InFace, InKerningFlags, InFontSize, InFontScale);
 		}
 		return Entry;
 	}
