[SlateCore] Added pseudo-monospacing to text shaping

diff --git a/Engine/Source/Editor/DetailCustomizations/Private/SlateFontInfoCustomization.cpp b/Engine/Source/Editor/DetailCustomizations/Private/SlateFontInfoCustomization.cpp
--- a/Engine/Source/Editor/DetailCustomizations/Private/SlateFontInfoCustomization.cpp
+++ b/Engine/Source/Editor/DetailCustomizations/Private/SlateFontInfoCustomization.cpp
@@ -123,6 +123,24 @@ void FSlateFontInfoStructCustomization::CustomizeChildren(TSharedRef<IPropertyHa
 	
 	InStructBuilder.AddProperty(InStructPropertyHandle->GetChildHandle(GET_MEMBER_NAME_CHECKED(FSlateFontInfo, SkewAmount)).ToSharedRef());
 
+	const TSharedRef<IPropertyHandle> MonospacingHandle = InStructPropertyHandle->GetChildHandle(GET_MEMBER_NAME_CHECKED(FSlateFontInfo, bForceMonospaced)).ToSharedRef();
+	InStructBuilder.AddProperty(MonospacingHandle);
+
+	const TSharedRef<IPropertyHandle> MonospacingWidthHandle = InStructPropertyHandle->GetChildHandle(GET_MEMBER_NAME_CHECKED(FSlateFontInfo, MonospacedWidth)).ToSharedRef();
+	InStructBuilder.AddProperty(MonospacingWidthHandle);
+	
+	// Set an initial "sensible" value based on the current font size. Won't run if value is already non-default/zero 
+	MonospacingHandle->SetOnPropertyValueChanged(FSimpleDelegate::CreateLambda([this, MonospacingWidthHandle]()
+	{
+		if (!MonospacingWidthHandle->DiffersFromDefault())
+		{
+			float FontSizeValue;
+			FontSizeProperty->GetValue(FontSizeValue);
+			
+			MonospacingWidthHandle->SetValue(static_cast<int32>(FontSizeValue));
+		}
+	}));
+
 	InStructBuilder.AddProperty(InStructPropertyHandle->GetChildHandle(GET_MEMBER_NAME_CHECKED(FSlateFontInfo, FontMaterial)).ToSharedRef());
 
 	InStructBuilder.AddProperty(InStructPropertyHandle->GetChildHandle(GET_MEMBER_NAME_CHECKED(FSlateFontInfo, OutlineSettings)).ToSharedRef());
diff --git a/Engine/Source/Runtime/SlateCore/Private/Fonts/SlateTextShaper.cpp b/Engine/Source/Runtime/SlateCore/Private/Fonts/SlateTextShaper.cpp
--- a/Engine/Source/Runtime/SlateCore/Private/Fonts/SlateTextShaper.cpp
+++ b/Engine/Source/Runtime/SlateCore/Private/Fonts/SlateTextShaper.cpp
@@ -282,7 +282,6 @@ FShapedGlyphSequenceRef FSlateTextShaper::FinalizeTextShaping(TArray<FShapedGlyp
 
 		if (FaceGlyphData.FaceAndMemory.IsValid() && FaceGlyphData.FaceAndMemory->IsFaceValid())
 		{
-
 			if (FMath::IsNearlyEqual(InFontInfo.GetClampSkew(), 0.f))
 			{
 				FT_Set_Transform(FaceGlyphData.FaceAndMemory->GetFace(), nullptr, nullptr);
@@ -405,6 +404,15 @@ void FSlateTextShaper::PerformKerningOnlyTextShaping(const TCHAR* InText, const
 			ensure(LetterSpacingScaledAsFloat <= std::numeric_limits<int16>::max());
 			const int16 LetterSpacingScaled = (int16)LetterSpacingScaledAsFloat;
 
+			// Used for monospacing
+			int16 FixedAdvance = INDEX_NONE;
+			if (InFontInfo.bForceMonospaced)
+			{
+				const float MonospacingScaledAsFloat = InFontInfo.MonospacedWidth != 0 ? InFontInfo.MonospacedWidth * InFontInfo.Size * FinalFontScale : 0.f;
+				ensure(MonospacingScaledAsFloat <= std::numeric_limits<int16>::max());
+				FixedAdvance = (int16)MonospacingScaledAsFloat;
+			}
+
 			FreeTypeUtils::ApplySizeAndScale(KerningOnlyTextSequenceEntry.FaceAndMemory->GetFace(), InFontInfo.Size, FinalFontScale);
 			TSharedRef<FShapedGlyphFaceData> ShapedGlyphFaceData = MakeShared<FShapedGlyphFaceData>(KerningOnlyTextSequenceEntry.FaceAndMemory, GlyphFlags, InFontInfo.Size, FinalFontScale, InFontInfo.GetClampSkew());
 			TSharedPtr<FFreeTypeKerningCache> KerningCache = FTCacheDirectory->GetKerningCache(KerningOnlyTextSequenceEntry.FaceAndMemory->GetFace(), FT_KERNING_DEFAULT, InFontInfo.Size, FinalFontScale);
@@ -439,9 +447,9 @@ void FSlateTextShaper::PerformKerningOnlyTextShaping(const TCHAR* InText, const
 					ShapedGlyphEntry.FontFaceData = ShapedGlyphFaceData;
 					ShapedGlyphEntry.GlyphIndex = GlyphIndex;
 					ShapedGlyphEntry.SourceIndex = CurrentCharIndex;
-					ShapedGlyphEntry.XAdvance = XAdvance;
+					ShapedGlyphEntry.XAdvance = !InFontInfo.bForceMonospaced ? XAdvance : FixedAdvance;
 					ShapedGlyphEntry.YAdvance = 0;
-					ShapedGlyphEntry.XOffset = 0;
+					ShapedGlyphEntry.XOffset = !InFontInfo.bForceMonospaced ? 0 : (FixedAdvance - XAdvance) / 2;
 					ShapedGlyphEntry.YOffset = 0;
 					ShapedGlyphEntry.Kerning = 0;
 					ShapedGlyphEntry.NumCharactersInGlyph = 1;
@@ -459,7 +467,7 @@ void FSlateTextShaper::PerformKerningOnlyTextShaping(const TCHAR* InText, const
 							PreviousShapedGlyphEntry.XAdvance += LetterSpacingScaled;
 						}
 
-						if (ShapedGlyphEntry.bIsVisible)
+						if (ShapedGlyphEntry.bIsVisible && !InFontInfo.bForceMonospaced)
 						{
 							FT_Vector KerningVector;
 							if (KerningCache && KerningCache->FindOrCache(PreviousShapedGlyphEntry.GlyphIndex, ShapedGlyphEntry.GlyphIndex, KerningVector))
diff --git a/Engine/Source/Runtime/SlateCore/Public/Fonts/FontCache.h b/Engine/Source/Runtime/SlateCore/Public/Fonts/FontCache.h
--- a/Engine/Source/Runtime/SlateCore/Public/Fonts/FontCache.h
+++ b/Engine/Source/Runtime/SlateCore/Public/Fonts/FontCache.h
@@ -481,7 +481,7 @@ struct FCharacterEntry
 	uint16 VSize = 0;
 	/** The vertical distance from the baseline to the topmost border of the character */
 	int16 VerticalOffset = 0;
-	/** The vertical distance from the origin to the left most border of the character */
+	/** The horizontal distance from the origin to the left most border of the character */
 	int16 HorizontalOffset = 0;
 	/** The largest vertical distance below the baseline for any character in the font */
 	int16 GlobalDescender = 0;
diff --git a/Engine/Source/Runtime/SlateCore/Public/Fonts/FontTypes.h b/Engine/Source/Runtime/SlateCore/Public/Fonts/FontTypes.h
--- a/Engine/Source/Runtime/SlateCore/Public/Fonts/FontTypes.h
+++ b/Engine/Source/Runtime/SlateCore/Public/Fonts/FontTypes.h
@@ -49,7 +49,7 @@ public:
 
 	inline bool IsIdenticalToForCaching(const FSlateFontKey& Other) const
 	{
-		return FontInfo.IsIdentialToForCaching(Other.FontInfo)
+		return FontInfo.IsIdenticalToForCaching(Other.FontInfo)
 			&& OutlineSettings.IsIdenticalToForCaching(Other.OutlineSettings)
 			&& Scale == Other.Scale;
 	}
diff --git a/Engine/Source/Runtime/SlateCore/Public/Fonts/SlateFontInfo.h b/Engine/Source/Runtime/SlateCore/Public/Fonts/SlateFontInfo.h
--- a/Engine/Source/Runtime/SlateCore/Public/Fonts/SlateFontInfo.h
+++ b/Engine/Source/Runtime/SlateCore/Public/Fonts/SlateFontInfo.h
@@ -155,7 +155,7 @@ struct FSlateFontInfo
 	TSharedPtr<const FCompositeFont> CompositeFont;
 
 	/** The name of the font to use from the default typeface (None will use the first entry) */
-	UPROPERTY(EditAnywhere, BlueprintReadWrite, Category=SlateStyleRules, meta=(DisplayName="Typeface"))
+	UPROPERTY(EditAnywhere, BlueprintReadWrite, Category=SlateStyleRules, meta=(DisplayName="Typeface", EditCondition="FontObject"))
 	FName TypefaceFontName;
 
 	/**
@@ -177,6 +177,14 @@ struct FSlateFontInfo
 	/** The font fallback level. Runtime only, don't set on shared FSlateFontInfo, as it may change the font elsewhere (make a copy). */
 	EFontFallback FontFallback;
 
+	/** Enable pseudo-monospaced font. */
+	UPROPERTY(EditAnywhere, BlueprintReadWrite, Category=SlateStyleRules, meta=(DisplayName="Monospacing"))
+	bool bForceMonospaced = false;
+
+	/** The uniform width to apply to all characters when bForceMonospaced is enabled, proportional of the font Size. */
+	UPROPERTY(EditAnywhere, BlueprintReadWrite, Category=SlateStyleRules, meta=(ClampMin=0))
+	float MonospacedWidth = 1.0f;
+
 #if WITH_EDITORONLY_DATA
 private:
 
@@ -252,7 +260,7 @@ public:
 	SLATECORE_API FSlateFontInfo( const WIDECHAR* InFontName, float InSize, EFontHinting InHinting = EFontHinting::Default );
 
 public:
-	inline bool IsIdentialToForCaching(const FSlateFontInfo& Other) const
+	inline bool IsIdenticalToForCaching(const FSlateFontInfo& Other) const
 	{
 		// Ignore FontMaterial because it does not affect the cached glyph.
 		return FontObject == Other.FontObject
@@ -263,7 +271,7 @@ public:
 			&& SkewAmount == Other.SkewAmount;
 	}
 
-	inline bool IsIdenticalTo(const FSlateFontInfo& Other) const
+	inline  bool IsIdenticalTo(const FSlateFontInfo& Other) const
 	{
 		return FontObject == Other.FontObject
 			&& FontMaterial == Other.FontMaterial
@@ -272,7 +280,9 @@ public:
 			&& TypefaceFontName == Other.TypefaceFontName
 			&& Size == Other.Size
 			&& LetterSpacing == Other.LetterSpacing
-			&& SkewAmount == Other.SkewAmount;
+			&& SkewAmount == Other.SkewAmount
+			&& bForceMonospaced == Other.bForceMonospaced
+			&& (bForceMonospaced ? MonospacedWidth == Other.MonospacedWidth : true);
 	}
 
 	inline bool operator==(const FSlateFontInfo& Other) const
