Ascent/descent override: add a new option at the composite font level, to be able to consider or ignore the font face ascent/descent override, based on the composite font using it.
Fixed a few typo (ascend VS ascent, descend VS descent).

#test Using a font containing a font face with ascent/descent overridden, enabled or disabled the new option in the font: in font editor preview, UMG editor, and running PIE, the ascend/descent was taken into account or ignored.

diff --git a/Engine/Source/Editor/FontEditor/Private/FontEditor.cpp b/Engine/Source/Editor/FontEditor/Private/FontEditor.cpp
--- a/Engine/Source/Editor/FontEditor/Private/FontEditor.cpp
+++ b/Engine/Source/Editor/FontEditor/Private/FontEditor.cpp
@@ -1014,6 +1014,13 @@ void FFontEditor::OnPostReimport(UObject* InObject, bool bSuccess)
 
 void FFontEditor::OnObjectPropertyChanged(UObject* InObject, struct FPropertyChangedEvent& InPropertyChangedEvent)
 {
+	if (Cast<UFont>(InObject))
+	{
+		//Force all texts using a font to be refreshed.
+		FSlateApplicationBase::Get().InvalidateAllWidgets(false);
+		GSlateLayoutGeneration++;
+	}
+
 	if (Cast<UFontFace>(InObject))
 	{
 		// Refresh the composite font editor when a font face is changed as it may affect our preview
diff --git a/Engine/Source/Runtime/Engine/Classes/Engine/Font.h b/Engine/Source/Runtime/Engine/Classes/Engine/Font.h
--- a/Engine/Source/Runtime/Engine/Classes/Engine/Font.h
+++ b/Engine/Source/Runtime/Engine/Classes/Engine/Font.h
@@ -169,7 +169,7 @@ class UFont : public UObject, public IFontProviderInterface
 	FName LegacyFontName;
 
 	/** Embedded composite font data */
-	UPROPERTY()
+	UPROPERTY(EditAnywhere, Category = RuntimeFont)
 	FCompositeFont CompositeFont;
 
 public:
diff --git a/Engine/Source/Runtime/SlateCore/Private/Fonts/FontCacheCompositeFont.cpp b/Engine/Source/Runtime/SlateCore/Private/Fonts/FontCacheCompositeFont.cpp
--- a/Engine/Source/Runtime/SlateCore/Private/Fonts/FontCacheCompositeFont.cpp
+++ b/Engine/Source/Runtime/SlateCore/Private/Fonts/FontCacheCompositeFont.cpp
@@ -448,8 +448,8 @@ TSharedPtr<FFreeTypeFace> FCompositeFontCache::GetFontFace(const FFontData& InFo
 		// Got a valid font?
 		if (FaceAndMemory.IsValid())
 		{
-			FaceAndMemory->OverrideAscend(InFontData.IsAscendOverridden(), InFontData.GetAscendOverriddenValue());
-			FaceAndMemory->OverrideDescend(InFontData.IsDescendOverridden(), InFontData.GetDescendOverriddenValue());
+			FaceAndMemory->OverrideAscent(InFontData.IsAscendOverridden(), InFontData.GetAscendOverriddenValue());
+			FaceAndMemory->OverrideDescent(InFontData.IsDescendOverridden(), InFontData.GetDescendOverriddenValue());
 
 			FontFaceMap.Add(InFontData, FaceAndMemory);
 
diff --git a/Engine/Source/Runtime/SlateCore/Private/Fonts/FontCacheFreeType.h b/Engine/Source/Runtime/SlateCore/Private/Fonts/FontCacheFreeType.h
--- a/Engine/Source/Runtime/SlateCore/Private/Fonts/FontCacheFreeType.h
+++ b/Engine/Source/Runtime/SlateCore/Private/Fonts/FontCacheFreeType.h
@@ -223,28 +223,28 @@ public:
 		return FreeTypeUtils::GetHeight(FTFace, LayoutMethod);
 	}
 
-	FORCEINLINE FT_Pos GetScaledHeight() const
+	FORCEINLINE FT_Pos GetScaledHeight(bool bAllowOverride) const
 	{
-		if (IsAscendOverridden || IsDescendOverridden)
-			return GetAscender() - GetDescender();
+		if (bAllowOverride && (IsAscentOverridden || IsDescentOverridden))
+			return GetAscender(true) - GetDescender(true);
 		return FreeTypeUtils::GetScaledHeight(FTFace, LayoutMethod);
 	}
 
-	FORCEINLINE FT_Pos GetAscender() const
+	FORCEINLINE FT_Pos GetAscender(bool bAllowOverride) const
 	{
-		if (IsAscendOverridden)
+		if (bAllowOverride && IsAscentOverridden)
 		{
-			FT_F26Dot6 ScaledAscender = FT_MulFix(AscendOverrideValue, FTFace->size->metrics.y_scale);
+			FT_F26Dot6 ScaledAscender = FT_MulFix(AscentOverrideValue, FTFace->size->metrics.y_scale);
 			return (ScaledAscender + 0b111111) & ~0b111111; //(26.6 fixed point ceil). Using ceiling of scaled ascend, as recommended by Freetype, to avoid grid fitting/hinting issues.
 		}
 		return FreeTypeUtils::GetAscender(FTFace, LayoutMethod);
 	}
 
-	FORCEINLINE FT_Pos GetDescender() const
+	FORCEINLINE FT_Pos GetDescender(bool bAllowOverride) const
 	{
-		if (IsDescendOverridden)
+		if (bAllowOverride && IsDescentOverridden)
 		{
-			FT_F26Dot6 ScaledDescender =  FT_MulFix(DescendOverrideValue, FTFace->size->metrics.y_scale);
+			FT_F26Dot6 ScaledDescender =  FT_MulFix(DescentOverrideValue, FTFace->size->metrics.y_scale);
 			return ScaledDescender & ~0b111111; //(26.6 fixed point floor). Using floor of scaled descend, as recommended by Freetype, to avoid grid fitting/hinting issues.
 		}
 		return FreeTypeUtils::GetDescender(FTFace, LayoutMethod);
@@ -283,19 +283,19 @@ public:
 #endif
 	}
 
-	void OverrideAscend(bool InOverride, int32 Value = 0)
+	void OverrideAscent(bool InOverride, int32 Value = 0)
 	{
 #if WITH_FREETYPE
-		IsAscendOverridden = InOverride;
-		AscendOverrideValue = FreeTypeUtils::ConvertPixelTo26Dot6<FT_F26Dot6>(Value);
+		IsAscentOverridden = InOverride;
+		AscentOverrideValue = FreeTypeUtils::ConvertPixelTo26Dot6<FT_F26Dot6>(Value);
 #endif //WITH_FREETYPE
 	}
 
-	void OverrideDescend(bool InOverride, int32 Value = 0)
+	void OverrideDescent(bool InOverride, int32 Value = 0)
 	{
 #if WITH_FREETYPE
-		IsDescendOverridden = InOverride;
-		DescendOverrideValue = FreeTypeUtils::ConvertPixelTo26Dot6<FT_F26Dot6>(Value);
+		IsDescentOverridden = InOverride;
+		DescentOverrideValue = FreeTypeUtils::ConvertPixelTo26Dot6<FT_F26Dot6>(Value);
 #endif //WITH_FREETYPE
 	}
 
@@ -343,10 +343,10 @@ private:
 	FT_StreamRec FTStream;
 	FT_Open_Args FTFaceOpenArgs;
 
-	bool IsAscendOverridden = false;
-	bool IsDescendOverridden = false;
-	FT_F26Dot6 AscendOverrideValue = 0;
-	FT_F26Dot6 DescendOverrideValue = 0;
+	bool IsAscentOverridden = false;
+	bool IsDescentOverridden = false;
+	FT_F26Dot6 AscentOverrideValue = 0;
+	FT_F26Dot6 DescentOverrideValue = 0;
 #endif // WITH_FREETYPE
 
 	TSet<FName> Attributes;
diff --git a/Engine/Source/Runtime/SlateCore/Private/Fonts/FontCacheHarfBuzz.cpp b/Engine/Source/Runtime/SlateCore/Private/Fonts/FontCacheHarfBuzz.cpp
--- a/Engine/Source/Runtime/SlateCore/Private/Fonts/FontCacheHarfBuzz.cpp
+++ b/Engine/Source/Runtime/SlateCore/Private/Fonts/FontCacheHarfBuzz.cpp
@@ -5,6 +5,7 @@
 #include "Fonts/FontCacheFreeType.h"
 #include "Fonts/SlateFontRenderer.h"
 #include "Trace/SlateMemoryTags.h"
+#include "Fonts/FontUtils.h"
 
 #if WITH_HARFBUZZ
 
@@ -449,13 +450,13 @@ FHarfBuzzFontFactory::~FHarfBuzzFontFactory()
 
 #if WITH_HARFBUZZ
 
-hb_font_t* FHarfBuzzFontFactory::CreateFont(const FFreeTypeFace& InFace, const uint32 InGlyphFlags, const float InFontSize, const float InFontScale) const
+hb_font_t* FHarfBuzzFontFactory::CreateFont(const FFreeTypeFace& InFace, const uint32 InGlyphFlags, const FSlateFontInfo& InFontInfo, const float InFontScale) const
 {
 	hb_font_t* HarfBuzzFont = nullptr;
 
 #if WITH_FREETYPE
 	FT_Face FreeTypeFace = InFace.GetFace();
-	FreeTypeUtils::ApplySizeAndScale(FreeTypeFace, InFontSize, InFontScale);
+	FreeTypeUtils::ApplySizeAndScale(FreeTypeFace, InFontInfo.Size, InFontScale);
 
 	hb_font_extents_t HarfBuzzFontExtents;
 	FMemory::Memzero(HarfBuzzFontExtents);
@@ -472,9 +473,10 @@ hb_font_t* FHarfBuzzFontFactory::CreateFont(const FFreeTypeFace& InFace, const u
 			hb_font_get_scale(HarfBuzzFTFont, &HarfBuzzFTFontXScale, &HarfBuzzFTFontYScale);
 
 			// Cache the font extents
-			HarfBuzzFontExtents.ascender = InFace.GetAscender();
-			HarfBuzzFontExtents.descender = InFace.GetDescender();
-			HarfBuzzFontExtents.line_gap = InFace.GetScaledHeight() - (HarfBuzzFontExtents.ascender - HarfBuzzFontExtents.descender);
+			const bool IsAscentDescentOverridenEnabled = UE::Slate::FontUtils::IsAscentDescentOverrideEnabled(InFontInfo.FontObject);
+			HarfBuzzFontExtents.ascender = InFace.GetAscender(IsAscentDescentOverridenEnabled);
+			HarfBuzzFontExtents.descender = InFace.GetDescender(IsAscentDescentOverridenEnabled);
+			HarfBuzzFontExtents.line_gap = InFace.GetScaledHeight(IsAscentDescentOverridenEnabled) - (HarfBuzzFontExtents.ascender - HarfBuzzFontExtents.descender);
 			if (HarfBuzzFTFontYScale < 0)
 			{
 				HarfBuzzFontExtents.ascender = -HarfBuzzFontExtents.ascender;
@@ -493,7 +495,7 @@ hb_font_t* FHarfBuzzFontFactory::CreateFont(const FFreeTypeFace& InFace, const u
 	hb_font_set_user_data(
 		HarfBuzzFont, 
 		&HarfBuzzFontFunctions::UserDataKey, 
-		HarfBuzzFontFunctions::CreateUserData(InFontSize, InFontScale, FTCacheDirectory, HarfBuzzFontExtents),
+		HarfBuzzFontFunctions::CreateUserData(InFontInfo.Size, InFontScale, FTCacheDirectory, HarfBuzzFontExtents),
 		&HarfBuzzFontFunctions::DestroyUserData, 
 		true
 		);
diff --git a/Engine/Source/Runtime/SlateCore/Private/Fonts/FontCacheHarfBuzz.h b/Engine/Source/Runtime/SlateCore/Private/Fonts/FontCacheHarfBuzz.h
--- a/Engine/Source/Runtime/SlateCore/Private/Fonts/FontCacheHarfBuzz.h
+++ b/Engine/Source/Runtime/SlateCore/Private/Fonts/FontCacheHarfBuzz.h
@@ -7,6 +7,7 @@
 
 class FFreeTypeFace;
 class FFreeTypeCacheDirectory;
+struct FSlateFontInfo;
 
 #ifndef WITH_HARFBUZZ
 	#define WITH_HARFBUZZ 0
@@ -54,7 +55,7 @@ public:
 
 #if WITH_HARFBUZZ
 	/** Create a HarfBuzz font from the given face - must be destroyed with hb_font_destroy when done */
-	hb_font_t* CreateFont(const FFreeTypeFace& InFace, const uint32 InGlyphFlags, const float InFontSize, const float InFontScale) const;
+	hb_font_t* CreateFont(const FFreeTypeFace& InFace, const uint32 InGlyphFlags, const FSlateFontInfo& InFontInfo, const float InFontScale) const;
 #endif // WITH_HARFBUZZ
 
 private:
diff --git a/Engine/Source/Runtime/SlateCore/Private/Fonts/FontUtils.cpp b/Engine/Source/Runtime/SlateCore/Private/Fonts/FontUtils.cpp
new file mode 100644
--- /dev/null
+++ b/Engine/Source/Runtime/SlateCore/Private/Fonts/FontUtils.cpp
@@ -0,0 +1,24 @@
+// Copyright Epic Games, Inc. All Rights Reserved.
+
+#pragma once
+
+#include "Fonts/FontUtils.h"
+#include "Fonts/FontProviderInterface.h"
+#include "Fonts/CompositeFont.h"
+
+namespace UE
+{
+	namespace Slate
+	{
+		bool FontUtils::IsAscentDescentOverrideEnabled(const TObjectPtr<UObject const> FontObject)
+		{
+			const IFontProviderInterface* FontProvider = Cast<const IFontProviderInterface>(FontObject);
+			if (FontProvider && FontProvider->GetCompositeFont())
+			{
+				return FontProvider->GetCompositeFont()->IsAscentDescentOverrideEnabled();
+			}
+
+			return true; //Default is enabled, unless we have find an explicit disabling earlier.
+		}
+	}
+}
\ No newline at end of file
diff --git a/Engine/Source/Runtime/SlateCore/Private/Fonts/FontUtils.h b/Engine/Source/Runtime/SlateCore/Private/Fonts/FontUtils.h
new file mode 100644
--- /dev/null
+++ b/Engine/Source/Runtime/SlateCore/Private/Fonts/FontUtils.h
@@ -0,0 +1,17 @@
+// Copyright Epic Games, Inc. All Rights Reserved.
+
+#pragma once
+
+#include "CoreMinimal.h"
+
+namespace UE
+{
+	namespace Slate
+	{
+		class FontUtils
+		{
+		public:
+			static bool IsAscentDescentOverrideEnabled(const TObjectPtr<UObject const> FontObject);
+		};
+	}
+}
\ No newline at end of file
diff --git a/Engine/Source/Runtime/SlateCore/Private/Fonts/SlateFontRenderer.cpp b/Engine/Source/Runtime/SlateCore/Private/Fonts/SlateFontRenderer.cpp
--- a/Engine/Source/Runtime/SlateCore/Private/Fonts/SlateFontRenderer.cpp
+++ b/Engine/Source/Runtime/SlateCore/Private/Fonts/SlateFontRenderer.cpp
@@ -8,6 +8,7 @@
 #include "SlateGlobals.h"
 
 #include <limits>
+#include "Fonts/FontUtils.h"
 
 DECLARE_CYCLE_STAT(TEXT("Freetype Render Glyph"), STAT_FreetypeRenderGlyph, STATGROUP_Slate);
 
@@ -197,7 +198,7 @@ uint16 FSlateFontRenderer::GetMaxHeight(const FSlateFontInfo& InFontInfo, const
 
 		// Adjust the height by the size of the outline that was applied.  
 		const float HeightAdjustment = InFontInfo.OutlineSettings.OutlineSize * InScale;
-		return static_cast<uint16>(FreeTypeUtils::Convert26Dot6ToRoundedPixel<int32>(FaceGlyphData.FaceAndMemory->GetScaledHeight()) + HeightAdjustment);
+		return static_cast<uint16>(FreeTypeUtils::Convert26Dot6ToRoundedPixel<int32>(FaceGlyphData.FaceAndMemory->GetScaledHeight(UE::Slate::FontUtils::IsAscentDescentOverrideEnabled(InFontInfo.FontObject))) + HeightAdjustment); 
 	}
 
 	return 0;
@@ -218,7 +219,7 @@ int16 FSlateFontRenderer::GetBaseline(const FSlateFontInfo& InFontInfo, const fl
 	{
 		FreeTypeUtils::ApplySizeAndScale(FaceGlyphData.FaceAndMemory->GetFace(), InFontInfo.Size, InScale);
 
-		return FreeTypeUtils::Convert26Dot6ToRoundedPixel<int16>(FaceGlyphData.FaceAndMemory->GetDescender());
+		return FreeTypeUtils::Convert26Dot6ToRoundedPixel<int16>(FaceGlyphData.FaceAndMemory->GetDescender(UE::Slate::FontUtils::IsAscentDescentOverrideEnabled(InFontInfo.FontObject)));
 	}
 
 	return 0;
diff --git a/Engine/Source/Runtime/SlateCore/Private/Fonts/SlateTextShaper.cpp b/Engine/Source/Runtime/SlateCore/Private/Fonts/SlateTextShaper.cpp
--- a/Engine/Source/Runtime/SlateCore/Private/Fonts/SlateTextShaper.cpp
+++ b/Engine/Source/Runtime/SlateCore/Private/Fonts/SlateTextShaper.cpp
@@ -7,6 +7,7 @@
 #include "SlateGlobals.h"
 
 #include <limits>
+#include "Fonts/FontUtils.h"
 
 DECLARE_CYCLE_STAT(TEXT("Shape Bidirectional Text"), STAT_SlateShapeBidirectionalText, STATGROUP_Slate);
 DECLARE_CYCLE_STAT(TEXT("Shape Unidirectional Text"), STAT_SlateShapeUnidirectionalText, STATGROUP_Slate);
@@ -298,9 +299,10 @@ FShapedGlyphSequenceRef FSlateTextShaper::FinalizeTextShaping(TArray<FShapedGlyp
 			}
 
 			FreeTypeUtils::ApplySizeAndScale(FaceGlyphData.FaceAndMemory->GetFace(), InFontInfo.Size, InFontScale);
-			
-			TextBaseline = FreeTypeUtils::Convert26Dot6ToRoundedPixel<int16>(FaceGlyphData.FaceAndMemory->GetDescender());
-			MaxHeight = FreeTypeUtils::Convert26Dot6ToRoundedPixel<uint16>(FaceGlyphData.FaceAndMemory->GetScaledHeight());
+
+			const bool IsAscentDescentOverridenEnabled = UE::Slate::FontUtils::IsAscentDescentOverrideEnabled(InFontInfo.FontObject);
+			TextBaseline = FreeTypeUtils::Convert26Dot6ToRoundedPixel<int16>(FaceGlyphData.FaceAndMemory->GetDescender(IsAscentDescentOverridenEnabled));
+			MaxHeight = FreeTypeUtils::Convert26Dot6ToRoundedPixel<uint16>(FaceGlyphData.FaceAndMemory->GetScaledHeight(IsAscentDescentOverridenEnabled));
 		}
 	}
 #endif // WITH_FREETYPE
@@ -692,7 +694,7 @@ void FSlateTextShaper::PerformHarfBuzzTextShaping(const TCHAR* InText, const int
 			SlateFontRendererUtils::AppendGlyphFlags(*HarfBuzzTextSequenceEntry.FaceAndMemory, *HarfBuzzTextSequenceEntry.FontDataPtr, GlyphFlags);
 			const float FinalFontScale = InFontScale * HarfBuzzTextSequenceEntry.SubFontScalingFactor;
 
-			hb_font_t* HarfBuzzFont = HarfBuzzFontFactory.CreateFont(*HarfBuzzTextSequenceEntry.FaceAndMemory, GlyphFlags, InFontInfo.Size, FinalFontScale);
+			hb_font_t* HarfBuzzFont = HarfBuzzFontFactory.CreateFont(*HarfBuzzTextSequenceEntry.FaceAndMemory, GlyphFlags, InFontInfo, FinalFontScale);
 			TSharedRef<FShapedGlyphFaceData> ShapedGlyphFaceData = MakeShared<FShapedGlyphFaceData>(HarfBuzzTextSequenceEntry.FaceAndMemory, GlyphFlags, InFontInfo.Size, FinalFontScale, InFontInfo.GetClampSkew());
 			TSharedPtr<FFreeTypeKerningCache> KerningCache = FTCacheDirectory->GetKerningCache(HarfBuzzTextSequenceEntry.FaceAndMemory->GetFace(), FT_KERNING_DEFAULT, InFontInfo.Size, FinalFontScale);
 			TSharedRef<FFreeTypeAdvanceCache> AdvanceCache = FTCacheDirectory->GetAdvanceCache(HarfBuzzTextSequenceEntry.FaceAndMemory->GetFace(), ShapedGlyphFaceData->GlyphFlags, InFontInfo.Size, FinalFontScale);
diff --git a/Engine/Source/Runtime/SlateCore/Public/Fonts/CompositeFont.h b/Engine/Source/Runtime/SlateCore/Public/Fonts/CompositeFont.h
--- a/Engine/Source/Runtime/SlateCore/Public/Fonts/CompositeFont.h
+++ b/Engine/Source/Runtime/SlateCore/Public/Fonts/CompositeFont.h
@@ -454,6 +454,10 @@ struct FCompositeFont
 	UPROPERTY()
 	TArray<FCompositeSubFont> SubTypefaces;
 
+	/** If set to false, the ascent and descent override specified in a Font face will be ignored, and the value from the font source file will be used instead. */
+	UPROPERTY(EditAnywhere, Category = RuntimeFont)
+	bool bEnableAscentDescentOverride = true;
+
 #if WITH_EDITORONLY_DATA
 	/** 
 	 * Transient value containing the current history ID of this composite font
@@ -461,6 +465,12 @@ struct FCompositeFont
 	 */
 	int32 HistoryRevision;
 #endif	// WITH_EDITORONLY_DATA
+
+	/** Returns the enabled state of ascent and descent override. */
+	SLATECORE_API bool IsAscentDescentOverrideEnabled() const
+	{
+		return bEnableAscentDescentOverride;
+	}
 };
 
 /**
