Fixed compilation tests issue in a new file

#test compiled the code.

diff --git a/Engine/Source/Runtime/SlateCore/Private/Fonts/FontUtils.cpp b/Engine/Source/Runtime/SlateCore/Private/Fonts/FontUtils.cpp
--- a/Engine/Source/Runtime/SlateCore/Private/Fonts/FontUtils.cpp
+++ b/Engine/Source/Runtime/SlateCore/Private/Fonts/FontUtils.cpp
@@ -1,7 +1,5 @@
 // Copyright Epic Games, Inc. All Rights Reserved.
 
-#pragma once
-
 #include "Fonts/FontUtils.h"
 #include "Fonts/FontProviderInterface.h"
 #include "Fonts/CompositeFont.h"
