Make OnBackBufferReadyToPresentDelegate ThreadSafe

#rnx

diff --git a/Engine/Source/Runtime/SlateCore/Public/Rendering/SlateRenderer.h b/Engine/Source/Runtime/SlateCore/Public/Rendering/SlateRenderer.h
--- a/Engine/Source/Runtime/SlateCore/Public/Rendering/SlateRenderer.h
+++ b/Engine/Source/Runtime/SlateCore/Public/Rendering/SlateRenderer.h
@@ -297,7 +297,7 @@ public:
 	FOnPostResizeWindowBackbuffer& OnPostResizeWindowBackBuffer() { return PostResizeBackBufferDelegate; }
 
 	/** Callback on the render thread after slate rendering finishes and right before present is called */
-	DECLARE_MULTICAST_DELEGATE_TwoParams(FOnBackBufferReadyToPresent, SWindow&, const FTexture2DRHIRef&);
+	DECLARE_TS_MULTICAST_DELEGATE_TwoParams(FOnBackBufferReadyToPresent, SWindow&, const FTexture2DRHIRef&);
 	FOnBackBufferReadyToPresent& OnBackBufferReadyToPresent() { return OnBackBufferReadyToPresentDelegate; }
 
 	/** 
