StringView-ifying FSlateFontMeasure and FLRUStringCache

diff --git a/Engine/Source/Runtime/SlateCore/Private/Fonts/FontCacheUtils.h b/Engine/Source/Runtime/SlateCore/Private/Fonts/FontCacheUtils.h
--- a/Engine/Source/Runtime/SlateCore/Private/Fonts/FontCacheUtils.h
+++ b/Engine/Source/Runtime/SlateCore/Private/Fonts/FontCacheUtils.h
@@ -6,9 +6,6 @@
 #include "Stats/Stats.h"
 #include "SlateGlobals.h"
 
-typedef FString KeyType;
-typedef FVector2f ValueType;
-
 DECLARE_MEMORY_STAT_EXTERN(TEXT("Font Measure Memory"), STAT_SlateFontMeasureCacheMemory, STATGROUP_SlateMemory, SLATECORE_API);
 
 /**
@@ -32,9 +29,9 @@ public:
 	/**
 	 * Accesses an item in the cache.  
 	 */
-	FORCEINLINE const ValueType* AccessItem(const KeyType& Key)
+	FORCEINLINE const FVector2f* AccessItem(FStringView Key)
 	{
-		CacheEntry** Entry = LookupSet.Find( Key );
+		FCacheEntry** Entry = LookupSet.FindByHash(FLookupSet::KeyFuncsType::GetKeyHash(Key), Key);
 		if( Entry )
 		{
 			MarkAsRecent( *Entry );
@@ -44,9 +41,9 @@ public:
 		return nullptr;
 	}
 
-	void Add( const KeyType& Key, const ValueType& Value )
+	void Add( FStringView Key, const FVector2f& Value )
 	{
-		CacheEntry** Entry = LookupSet.Find( Key );
+		FCacheEntry** Entry = LookupSet.FindByHash(FLookupSet::KeyFuncsType::GetKeyHash(Key), Key);
 	
 		// Make a new link
 		if( !Entry )
@@ -57,7 +54,7 @@ public:
 				checkf( LookupSet.Num() < MaxNumElements, TEXT("Could not eject item from the LRU: (%d of %d), %s"), LookupSet.Num(), MaxNumElements, LeastRecent ? *LeastRecent->Key : TEXT("NULL"));
 			}
 
-			CacheEntry* NewEntry = new CacheEntry( Key, Value );
+			FCacheEntry* NewEntry = new FCacheEntry( FString(Key), Value );
 
 			// Link before the most recent so that we become the most recent
 			NewEntry->Link(MostRecent);
@@ -76,11 +73,11 @@ public:
 		else
 		{
 			// Trying to add an existing value 
-			CacheEntry* EntryPtr = *Entry;
+			FCacheEntry* EntryPtr = *Entry;
 			
 			// Update the value
 			EntryPtr->Value = Value;
-			checkSlow( EntryPtr->Key == Key );
+			checkSlow(FLookupSet::KeyFuncsType::Matches(EntryPtr->Key, Key));
 			// Mark as the most recent
 			MarkAsRecent( EntryPtr );
 		}
@@ -90,9 +87,9 @@ public:
 	{
 		DEC_MEMORY_STAT_BY( STAT_SlateFontMeasureCacheMemory, LookupSet.GetAllocatedSize() );
 
-		for( TSet<CacheEntry*, FCaseSensitiveStringKeyFuncs >::TIterator It(LookupSet); It; ++It )
+		for( TSet<FCacheEntry*, FCaseSensitiveStringKeyFuncs >::TIterator It(LookupSet); It; ++It )
 		{
-			CacheEntry* Entry = *It;
+			FCacheEntry* Entry = *It;
 			// Note no need to unlink anything here. we are emptying the entire list
 			delete Entry;
 		}
@@ -101,15 +98,15 @@ public:
 		MostRecent = LeastRecent = nullptr;
 	}
 private:
-	struct CacheEntry
+	struct FCacheEntry
 	{
-		KeyType Key;
-		ValueType Value;
-		CacheEntry* Next;
-		CacheEntry* Prev;
+		FString Key;
+		FVector2f Value;
+		FCacheEntry* Next;
+		FCacheEntry* Prev;
 
-		CacheEntry( const KeyType& InKey, const ValueType& InValue )
-			: Key( InKey )
+		FCacheEntry( FString InKey, const FVector2f& InValue )
+			: Key( MoveTemp(InKey) )
 			, Value( InValue )
 			, Next( nullptr )
 			, Prev( nullptr )
@@ -117,12 +114,12 @@ private:
 			INC_MEMORY_STAT_BY( STAT_SlateFontMeasureCacheMemory, Key.GetAllocatedSize()+sizeof(Value)+sizeof(Next)+sizeof(Prev) );
 		}
 
-		~CacheEntry()
+		~FCacheEntry()
 		{
 			DEC_MEMORY_STAT_BY( STAT_SlateFontMeasureCacheMemory, Key.GetAllocatedSize()+sizeof(Value)+sizeof(Next)+sizeof(Prev) );
 		}
 
-		FORCEINLINE void Link( CacheEntry* Before )
+		FORCEINLINE void Link( FCacheEntry* Before )
 		{
 			Next = Before;
 
@@ -149,32 +146,31 @@ private:
 		}
 	};
 
-	struct FCaseSensitiveStringKeyFuncs : BaseKeyFuncs<CacheEntry*, FString>
+	struct FCaseSensitiveStringKeyFuncs : BaseKeyFuncs<FCacheEntry*, FString>
 	{
-		FORCEINLINE static const KeyType& GetSetKey( const CacheEntry* Entry )
+		FORCEINLINE static const FString& GetSetKey( const FCacheEntry* Entry )
 		{
 			return Entry->Key;
 		}
 
-		FORCEINLINE static bool Matches(const KeyType& A,const KeyType& B)
+		FORCEINLINE static bool Matches(const FString& A, const FStringView& B)
 		{
-			return A.Equals( B, ESearchCase::CaseSensitive );
+			return FStringView(A).Equals( B, ESearchCase::CaseSensitive );
 		}
 
-		FORCEINLINE static uint32 GetKeyHash(const KeyType& Identifier)
+		FORCEINLINE static uint32 GetKeyHash(FStringView Identifier)
 		{
-			return FCrc::StrCrc32( *Identifier );
+			return CityHash32(reinterpret_cast<const char*>(Identifier.GetData()), Identifier.Len() * sizeof(TCHAR));
 		}
 	};
-
 	
 
 	/**
 	 * Marks the link as the the most recent
 	 *
-	 * @param Link	The link to mark as most recent
+	 * @param Entry	The link to mark as most recent
 	 */
-	FORCEINLINE void MarkAsRecent( CacheEntry* Entry )
+	FORCEINLINE void MarkAsRecent( FCacheEntry* Entry )
 	{
 		checkSlow( LeastRecent && MostRecent );
 
@@ -202,7 +198,7 @@ private:
 	 */
 	FORCEINLINE void Eject()
 	{
-		CacheEntry* EntryToRemove = LeastRecent;
+		FCacheEntry* EntryToRemove = LeastRecent;
 		// Eject the least recent, no more space
 		check( EntryToRemove );
 		STAT( uint32 CurrentMemUsage = LookupSet.GetAllocatedSize() );
@@ -219,11 +215,12 @@ private:
 	}
 
 private:
-	TSet< CacheEntry*, FCaseSensitiveStringKeyFuncs > LookupSet;
+	using FLookupSet = TSet< FCacheEntry*, FCaseSensitiveStringKeyFuncs >;
+	FLookupSet LookupSet;
 	/** Most recent item in the cache */
-	CacheEntry* MostRecent;
+	FCacheEntry* MostRecent;
 	/** Least recent item in the cache */
-	CacheEntry* LeastRecent;
+	FCacheEntry* LeastRecent;
 	/** The maximum number of elements in the cache */
 	int32 MaxNumElements;
 };
diff --git a/Engine/Source/Runtime/SlateCore/Private/Fonts/FontMeasure.cpp b/Engine/Source/Runtime/SlateCore/Private/Fonts/FontMeasure.cpp
--- a/Engine/Source/Runtime/SlateCore/Private/Fonts/FontMeasure.cpp
+++ b/Engine/Source/Runtime/SlateCore/Private/Fonts/FontMeasure.cpp
@@ -56,7 +56,7 @@ FSlateFontMeasure::FSlateFontMeasure( const TSharedRef<class FSlateFontCache>& I
 
 }
 
-UE::Slate::FDeprecateVector2DResult FSlateFontMeasure::Measure( const FString& Text, const FSlateFontInfo& InFontInfo, float FontScale ) const
+UE::Slate::FDeprecateVector2DResult FSlateFontMeasure::Measure( FStringView Text, const FSlateFontInfo& InFontInfo, float FontScale ) const
 {
 	int32 DummyLastCharacterIndex;
 	return MeasureStringInternal( Text, 0, Text.Len(), InFontInfo, false, FontScale, INDEX_NONE, ELastCharacterIndexFormat::Unused, DummyLastCharacterIndex);
@@ -65,17 +65,17 @@ UE::Slate::FDeprecateVector2DResult FSlateFontMeasure::Measure( const FString& T
 UE::Slate::FDeprecateVector2DResult FSlateFontMeasure::Measure( const FText& Text, const FSlateFontInfo& InFontInfo, float FontScale ) const
 {
 	int32 DummyLastCharacterIndex;
-	const FString& String = Text.ToString();
+	FStringView String = Text.ToString();
 	return MeasureStringInternal( String, 0, String.Len(), InFontInfo, false, FontScale, INDEX_NONE, ELastCharacterIndexFormat::Unused, DummyLastCharacterIndex);
 }
 
-UE::Slate::FDeprecateVector2DResult FSlateFontMeasure::Measure( const FString& Text, int32 StartIndex, int32 EndIndex, const FSlateFontInfo& InFontInfo, bool IncludeKerningWithPrecedingChar, float FontScale ) const
+UE::Slate::FDeprecateVector2DResult FSlateFontMeasure::Measure( FStringView Text, int32 StartIndex, int32 EndIndex, const FSlateFontInfo& InFontInfo, bool IncludeKerningWithPrecedingChar, float FontScale ) const
 {
 	int32 DummyLastCharacterIndex;
 	return MeasureStringInternal( Text, StartIndex, EndIndex, InFontInfo, IncludeKerningWithPrecedingChar, FontScale, INDEX_NONE, ELastCharacterIndexFormat::Unused, DummyLastCharacterIndex);
 }
 
-int32 FSlateFontMeasure::FindLastWholeCharacterIndexBeforeOffset( const FString& Text, const FSlateFontInfo& InFontInfo, const int32 HorizontalOffset, float FontScale ) const
+int32 FSlateFontMeasure::FindLastWholeCharacterIndexBeforeOffset( FStringView Text, const FSlateFontInfo& InFontInfo, const int32 HorizontalOffset, float FontScale ) const
 {
 	return FindLastWholeCharacterIndexBeforeOffset( Text, 0, Text.Len(), InFontInfo, HorizontalOffset, false, FontScale );
 }
@@ -86,7 +86,7 @@ int32 FSlateFontMeasure::FindLastWholeCharacterIndexBeforeOffset( const FText& T
 }
 
 
-int32 FSlateFontMeasure::FindLastWholeCharacterIndexBeforeOffset( const FString& Text, int32 StartIndex, int32 EndIndex, const FSlateFontInfo& InFontInfo, const int32 HorizontalOffset, bool IncludeKerningWithPrecedingChar, float FontScale ) const
+int32 FSlateFontMeasure::FindLastWholeCharacterIndexBeforeOffset( FStringView Text, int32 StartIndex, int32 EndIndex, const FSlateFontInfo& InFontInfo, const int32 HorizontalOffset, bool IncludeKerningWithPrecedingChar, float FontScale ) const
 {
 	int32 FoundLastCharacterIndex = INDEX_NONE;
 	MeasureStringInternal( Text, StartIndex, EndIndex, InFontInfo, IncludeKerningWithPrecedingChar, FontScale, HorizontalOffset, ELastCharacterIndexFormat::LastWholeCharacterBeforeOffset, FoundLastCharacterIndex );
@@ -99,12 +99,12 @@ int32 FSlateFontMeasure::FindFirstWholeCharacterIndexAfterOffset( const FText& T
 	return FindFirstWholeCharacterIndexAfterOffset( Text.ToString(), InFontInfo, HorizontalOffset, FontScale );
 }
 
-int32 FSlateFontMeasure::FindFirstWholeCharacterIndexAfterOffset( const FString& Text, const FSlateFontInfo& InFontInfo, const int32 HorizontalOffset, float FontScale ) const
+int32 FSlateFontMeasure::FindFirstWholeCharacterIndexAfterOffset( FStringView Text, const FSlateFontInfo& InFontInfo, const int32 HorizontalOffset, float FontScale ) const
 {
 	return FindFirstWholeCharacterIndexAfterOffset( Text, 0, Text.Len(), InFontInfo, HorizontalOffset, false, FontScale );
 }
 
-int32 FSlateFontMeasure::FindFirstWholeCharacterIndexAfterOffset( const FString& Text, int32 StartIndex, int32 EndIndex, const FSlateFontInfo& InFontInfo, const int32 HorizontalOffset, bool IncludeKerningWithPrecedingChar, float FontScale ) const
+int32 FSlateFontMeasure::FindFirstWholeCharacterIndexAfterOffset( FStringView Text, int32 StartIndex, int32 EndIndex, const FSlateFontInfo& InFontInfo, const int32 HorizontalOffset, bool IncludeKerningWithPrecedingChar, float FontScale ) const
 {
 	int32 FoundLastCharacterIndex = FindCharacterIndexAtOffset( Text, StartIndex, EndIndex, InFontInfo, HorizontalOffset, IncludeKerningWithPrecedingChar, FontScale );
 
@@ -120,7 +120,7 @@ int32 FSlateFontMeasure::FindFirstWholeCharacterIndexAfterOffset( const FString&
 	return FoundLastCharacterIndex;
 }
 
-int32 FSlateFontMeasure::FindCharacterIndexAtOffset( const FString& Text, const FSlateFontInfo& InFontInfo, const int32 HorizontalOffset, float FontScale ) const
+int32 FSlateFontMeasure::FindCharacterIndexAtOffset( FStringView Text, const FSlateFontInfo& InFontInfo, const int32 HorizontalOffset, float FontScale ) const
 {
 	return FindCharacterIndexAtOffset( Text, 0, Text.Len(), InFontInfo, HorizontalOffset, false, FontScale );
 }
@@ -130,14 +130,14 @@ int32 FSlateFontMeasure::FindCharacterIndexAtOffset( const FText& Text, const FS
 	return FindCharacterIndexAtOffset( Text.ToString(), InFontInfo, HorizontalOffset, FontScale );
 }
 
-int32 FSlateFontMeasure::FindCharacterIndexAtOffset( const FString& Text, int32 StartIndex, int32 EndIndex, const FSlateFontInfo& InFontInfo, const int32 HorizontalOffset, bool IncludeKerningWithPrecedingChar, float FontScale ) const
+int32 FSlateFontMeasure::FindCharacterIndexAtOffset( FStringView Text, int32 StartIndex, int32 EndIndex, const FSlateFontInfo& InFontInfo, const int32 HorizontalOffset, bool IncludeKerningWithPrecedingChar, float FontScale ) const
 {
 	int32 FoundLastCharacterIndex = 0;
 	MeasureStringInternal( Text, StartIndex, EndIndex, InFontInfo, IncludeKerningWithPrecedingChar, FontScale, HorizontalOffset, ELastCharacterIndexFormat::CharacterAtOffset, FoundLastCharacterIndex );
 	return FoundLastCharacterIndex;
 }
 
-FVector2f FSlateFontMeasure::MeasureStringInternal( const FString& Text, int32 StartIndex, int32 EndIndex, const FSlateFontInfo& InFontInfo, bool IncludeKerningWithPrecedingChar, float FontScale, int32 StopAfterHorizontalOffset, ELastCharacterIndexFormat CharIndexFormat, int32& OutLastCharacterIndex ) const
+FVector2f FSlateFontMeasure::MeasureStringInternal( FStringView Text, int32 StartIndex, int32 EndIndex, const FSlateFontInfo& InFontInfo, bool IncludeKerningWithPrecedingChar, float FontScale, int32 StopAfterHorizontalOffset, ELastCharacterIndexFormat CharIndexFormat, int32& OutLastCharacterIndex ) const
 {
 	QUICK_SCOPE_CYCLE_COUNTER(STAT_Slate_MeasureStringInternal)
 	FCharacterList& CharacterList = FontCache->GetCharacterList( InFontInfo, FontScale );
diff --git a/Engine/Source/Runtime/SlateCore/Public/Fonts/FontMeasure.h b/Engine/Source/Runtime/SlateCore/Public/Fonts/FontMeasure.h
--- a/Engine/Source/Runtime/SlateCore/Public/Fonts/FontMeasure.h
+++ b/Engine/Source/Runtime/SlateCore/Public/Fonts/FontMeasure.h
@@ -26,7 +26,7 @@ public:
 	 * @param FontScale	The scale to apply to the font
 	 * @return The width and height of the string.
 	 */
-	SLATECORE_API UE::Slate::FDeprecateVector2DResult Measure( const FString& Text, const FSlateFontInfo &InFontInfo, float FontScale = 1.0f ) const;
+	SLATECORE_API UE::Slate::FDeprecateVector2DResult Measure( FStringView Text, const FSlateFontInfo &InFontInfo, float FontScale = 1.0f ) const;
 
 	/** 
 	 * Measures the width and height of a passed in text.  The height is the maximum height of the largest character in the font/size pair
@@ -48,7 +48,7 @@ public:
 	 * @param FontScale	The scale to apply to the font
 	 * @return The width and height of the string.
 	 */
-	SLATECORE_API UE::Slate::FDeprecateVector2DResult Measure( const FString& Text, int32 StartIndex, int32 EndIndex, const FSlateFontInfo &InFontInfo, bool IncludeKerningWithPrecedingChar = true, float FontScale = 1.0f ) const;
+	SLATECORE_API UE::Slate::FDeprecateVector2DResult Measure( FStringView Text, int32 StartIndex, int32 EndIndex, const FSlateFontInfo &InFontInfo, bool IncludeKerningWithPrecedingChar = true, float FontScale = 1.0f ) const;
 
 	/** 
 	 * Finds the last whole character index before the specified position in pixels along the string horizontally
@@ -59,7 +59,7 @@ public:
 	 *
 	 * @return The index of the last whole character before the specified horizontal offset
 	 */
-	SLATECORE_API int32 FindLastWholeCharacterIndexBeforeOffset( const FString& Text, const FSlateFontInfo& InFontInfo, const int32 HorizontalOffset, float FontScale = 1.0f ) const;
+	SLATECORE_API int32 FindLastWholeCharacterIndexBeforeOffset( FStringView Text, const FSlateFontInfo& InFontInfo, const int32 HorizontalOffset, float FontScale = 1.0f ) const;
 
 	/** 
 	 * Finds the last whole character index before the specified position in pixels along the string horizontally
@@ -83,7 +83,7 @@ public:
 	 *
 	 * @return The index of the last whole character before the specified horizontal offset
 	 */
-	SLATECORE_API int32 FindLastWholeCharacterIndexBeforeOffset( const FString& Text, int32 StartIndex, int32 EndIndex, const FSlateFontInfo& InFontInfo, const int32 HorizontalOffset, bool IncludeKerningWithPrecedingChar = true, float FontScale = 1.0f ) const;
+	SLATECORE_API int32 FindLastWholeCharacterIndexBeforeOffset( FStringView Text, int32 StartIndex, int32 EndIndex, const FSlateFontInfo& InFontInfo, const int32 HorizontalOffset, bool IncludeKerningWithPrecedingChar = true, float FontScale = 1.0f ) const;
 
 	/** 
 	 * Finds the first whole character index after the specified position in pixels along the string horizontally
@@ -94,7 +94,7 @@ public:
 	 *
 	 * @return The index of the first whole character after the specified horizontal offset
 	 */
-	SLATECORE_API int32 FindFirstWholeCharacterIndexAfterOffset( const FString& Text, const FSlateFontInfo& InFontInfo, const int32 HorizontalOffset, float FontScale = 1.0f ) const;
+	SLATECORE_API int32 FindFirstWholeCharacterIndexAfterOffset( FStringView Text, const FSlateFontInfo& InFontInfo, const int32 HorizontalOffset, float FontScale = 1.0f ) const;
 
 	/** 
 	 * Finds the first whole character index after the specified position in pixels along the string horizontally
@@ -118,7 +118,7 @@ public:
 	 *
 	 * @return The index of the first whole character after the specified horizontal offset
 	 */
-	SLATECORE_API int32 FindFirstWholeCharacterIndexAfterOffset( const FString& Text, int32 StartIndex, int32 EndIndex, const FSlateFontInfo& InFontInfo, const int32 HorizontalOffset, bool IncludeKerningWithPrecedingChar = true, float FontScale = 1.0f ) const;
+	SLATECORE_API int32 FindFirstWholeCharacterIndexAfterOffset( FStringView Text, int32 StartIndex, int32 EndIndex, const FSlateFontInfo& InFontInfo, const int32 HorizontalOffset, bool IncludeKerningWithPrecedingChar = true, float FontScale = 1.0f ) const;
 
 	/** 
 	 * Finds the character index at the specified position in pixels along the string horizontally
@@ -131,7 +131,7 @@ public:
 	 *
 	 * @return The index of the character closest to the specified horizontal offset
 	 */
-	SLATECORE_API int32 FindCharacterIndexAtOffset( const FString& Text, const FSlateFontInfo &InFontInfo, const int32 HorizontalOffset, float FontScale = 1.0f ) const;
+	SLATECORE_API int32 FindCharacterIndexAtOffset( FStringView Text, const FSlateFontInfo &InFontInfo, const int32 HorizontalOffset, float FontScale = 1.0f ) const;
 
 	/** 
 	 * Finds the character index at the specified position in pixels along the text horizontally
@@ -159,7 +159,7 @@ public:
 	 *
 	 * @return The index of the character closest to the specified horizontal offset
 	 */
-	SLATECORE_API int32 FindCharacterIndexAtOffset( const FString& Text, int32 StartIndex, int32 EndIndex, const FSlateFontInfo &InFontInfo, const int32 HorizontalOffset, bool IncludeKerningWithPrecedingChar = true, float FontScale = 1.0f ) const;
+	SLATECORE_API int32 FindCharacterIndexAtOffset( FStringView Text, int32 StartIndex, int32 EndIndex, const FSlateFontInfo &InFontInfo, const int32 HorizontalOffset, bool IncludeKerningWithPrecedingChar = true, float FontScale = 1.0f ) const;
 
 	/**
 	 * Returns the height of the largest character in the font. 
@@ -225,7 +225,7 @@ private:
 	 *
 	 * @return The width and height of the string.
 	 */
-	FVector2f MeasureStringInternal( const FString& Text, int32 StartIndex, int32 EndIndex, const FSlateFontInfo& InFontInfo, bool IncludeKerningWithPrecedingChar, float FontScale, int32 StopAfterHorizontalOffset, ELastCharacterIndexFormat CharIndexFormat, int32& OutLastCharacterIndex ) const;
+	FVector2f MeasureStringInternal( FStringView Text, int32 StartIndex, int32 EndIndex, const FSlateFontInfo& InFontInfo, bool IncludeKerningWithPrecedingChar, float FontScale, int32 StopAfterHorizontalOffset, ELastCharacterIndexFormat CharIndexFormat, int32& OutLastCharacterIndex ) const;
 
 	/**
 	 * Check to see if there's an existing cached measurement, or failing that, add a new entry so that we can cache a new measurement
