Fix for SSlider not having a disabled style for Starship core style.

diff --git a/Engine/Source/Runtime/SlateCore/Private/Styling/StarshipCoreStyle.cpp b/Engine/Source/Runtime/SlateCore/Private/Styling/StarshipCoreStyle.cpp
--- a/Engine/Source/Runtime/SlateCore/Private/Styling/StarshipCoreStyle.cpp
+++ b/Engine/Source/Runtime/SlateCore/Private/Styling/StarshipCoreStyle.cpp
@@ -564,10 +564,12 @@ TSharedRef<ISlateStyle> FStarshipCoreStyle::Create()
 	// SSlider and SVolumeControl defaults...
 	{
 		FSliderStyle SliderStyle = FSliderStyle()
-			.SetNormalBarImage(   FSlateRoundedBoxBrush(FStyleColors::Input, 2.0f, FStyleColors::Input, 1.0f))
-			.SetHoveredBarImage(  FSlateRoundedBoxBrush(FStyleColors::Input, 2.0f, FStyleColors::Input, 1.0f))
-			.SetNormalThumbImage(  FSlateRoundedBoxBrush(FStyleColors::Foreground, Icon8x8) )
-			.SetHoveredThumbImage( FSlateRoundedBoxBrush(FStyleColors::ForegroundHover, Icon8x8) )
+			.SetNormalBarImage(FSlateRoundedBoxBrush(FStyleColors::Input, 2.0f, FStyleColors::Input, 1.0f))
+			.SetHoveredBarImage(FSlateRoundedBoxBrush(FStyleColors::Input, 2.0f, FStyleColors::Input, 1.0f))
+			.SetDisabledBarImage(FSlateRoundedBoxBrush(FStyleColors::Input, 2.0f, FStyleColors::Recessed, 1.0f))
+			.SetNormalThumbImage(FSlateRoundedBoxBrush(FStyleColors::Foreground, Icon8x8))
+			.SetHoveredThumbImage(FSlateRoundedBoxBrush(FStyleColors::ForegroundHover, Icon8x8))
+			.SetDisabledThumbImage(FSlateRoundedBoxBrush(FStyleColors::Foreground, Icon8x8))
 			.SetBarThickness(4.0f);
 		Style->Set("Slider", SliderStyle);
 
