Fix mixed MSDF and color glyphs in one sequence, re-enable MSDF text

#rnx <Release Notes exclude>

diff --git a/Engine/Source/Runtime/SlateCore/Private/Rendering/ElementBatcher.cpp b/Engine/Source/Runtime/SlateCore/Private/Rendering/ElementBatcher.cpp
--- a/Engine/Source/Runtime/SlateCore/Private/Rendering/ElementBatcher.cpp
+++ b/Engine/Source/Runtime/SlateCore/Private/Rendering/ElementBatcher.cpp
@@ -3206,8 +3206,9 @@ void FSlateElementBatcher::BuildShapedTextSequence(const FShapedTextBuildContext
 		{
 			// Get Sizing and atlas info
 			int8 NextAtlasDataTextureIndex = -1;
+			const bool bIsSdfGlyph = bIsSdfFont && GlyphToRender.FontFaceData && GlyphToRender.FontFaceData->bSupportsSdf;
 
-			if (bIsSdfFont)
+			if (bIsSdfGlyph)
 			{
 				const FFontSdfSettings& FontSdfSettings = GlyphSequenceToRender->GetFontSdfSettings();
 				const FSdfGlyphFontAtlasData SdfGlyphAtlasData = Context.FontCache->GetSdfGlyphFontAtlasData(GlyphToRender, *Context.OutlineSettings, FontSdfSettings);
@@ -3268,7 +3269,7 @@ void FSlateElementBatcher::BuildShapedTextSequence(const FShapedTextBuildContext
 				}
 
 				check(NextAtlasDataTextureIndex >= 0);
-				if (FontAtlasTexture == nullptr || NextAtlasDataTextureIndex != FontTextureIndex || (bIsSdfFont && (GlyphShaderParams.X != SdfShaderParams.W || GlyphShaderParams.Y != SdfShaderParams.Z)))
+				if (FontAtlasTexture == nullptr || NextAtlasDataTextureIndex != FontTextureIndex || (bIsSdfGlyph && (GlyphShaderParams.X != SdfShaderParams.W || GlyphShaderParams.Y != SdfShaderParams.Z)))
 				{
 					// Font has a new texture for this glyph or shader parameters changed. Refresh the batch we use and the index we are currently using
 					FontTextureIndex = NextAtlasDataTextureIndex;
@@ -3285,7 +3286,7 @@ void FSlateElementBatcher::BuildShapedTextSequence(const FShapedTextBuildContext
 					InvTextureSizeX = 1.0f / FontAtlasTexture->GetWidth();
 					InvTextureSizeY = 1.0f / FontAtlasTexture->GetHeight();
 
-					if (bIsSdfFont)
+					if (bIsSdfGlyph)
 					{
 						SdfShaderParams.X = InvTextureSizeX*GlyphShaderParams.X;
 						SdfShaderParams.Y = InvTextureSizeX*GlyphShaderParams.X;
@@ -3295,7 +3296,7 @@ void FSlateElementBatcher::BuildShapedTextSequence(const FShapedTextBuildContext
 
 					const ESlateFontAtlasContentType ContentType = SlateFontTexture->GetContentType();
 					Tint = ContentType == ESlateFontAtlasContentType::Color ? FColor::White : Context.FontTint;
-					check(bIsSdfFont == (ContentType == ESlateFontAtlasContentType::Msdf));
+					check(bIsSdfGlyph == (ContentType == ESlateFontAtlasContentType::Msdf));
 
 					ESlateShader ShaderType = ESlateShader::Default;
 					switch (ContentType)
@@ -3318,7 +3319,7 @@ void FSlateElementBatcher::BuildShapedTextSequence(const FShapedTextBuildContext
 					check(ShaderType != ESlateShader::Default);
 
 					RenderBatch = &CreateRenderBatch(Context.LayerId,
-						bIsSdfFont
+						bIsSdfGlyph
 						? FShaderParams::MakePixelShaderParams(SdfShaderParams)
 						: FShaderParams(),
 						FontShaderResource,
