Slate MSDF text - multiple spread categories per glyph (improves quality in some situations)

#rnx <Release Notes exclude>

diff --git a/Engine/Source/Runtime/SlateCore/Private/Fonts/FontCache.cpp b/Engine/Source/Runtime/SlateCore/Private/Fonts/FontCache.cpp
--- a/Engine/Source/Runtime/SlateCore/Private/Fonts/FontCache.cpp
+++ b/Engine/Source/Runtime/SlateCore/Private/Fonts/FontCache.cpp
@@ -137,15 +137,23 @@ FShapedGlyphEntryKey::FShapedGlyphEntryKey(const FShapedGlyphFaceData& InFontFac
 *
 */
 
-FSdfGlyphEntryKey::FSdfGlyphEntryKey(const TWeakPtr<FFreeTypeFace> InFontFace, uint32 InGlyphIndex, int32 InPpem)
+FSdfGlyphEntryKey::FSdfGlyphEntryKey(const TWeakPtr<FFreeTypeFace> InFontFace, uint32 InGlyphIndex, int32 InPpem, float InEmOuterSpread, float InEmInnerSpread)
 	: FontFace(InFontFace)
 	, GlyphIndex(InGlyphIndex)
 	, Ppem(InPpem)
+	, SpreadCategory(GetSpreadCategory(InEmOuterSpread, InEmInnerSpread))
 	, KeyHash(0)
 {
 	KeyHash = HashCombine(KeyHash, GetTypeHash(FontFace));
 	KeyHash = HashCombine(KeyHash, GetTypeHash(GlyphIndex));
 	KeyHash = HashCombine(KeyHash, GetTypeHash(Ppem));
+	KeyHash = HashCombine(KeyHash, GetTypeHash(SpreadCategory));
+}
+
+int32 FSdfGlyphEntryKey::GetSpreadCategory(float InEmOuterSpread, float InEmInnerSpread)
+{
+	(void) InEmInnerSpread;
+	return FMath::FloorToInt(FMath::Loge(InEmOuterSpread));
 }
 
 FSdfGlyphTaskKey::FSdfGlyphTaskKey(const FSdfGlyphEntryKey& InSdfGlyphEntryKey, float InEmOuterSpread, float InEmInnerSpread)
@@ -1206,8 +1214,8 @@ FSdfGlyphFontAtlasData FSlateFontCache::GetSdfGlyphFontAtlasData(const FShapedGl
 	const float Scale = TargetPpem/SdfPpem;
 	const float EmOutlineSize = FMath::RoundToFloat(InOutlineSettings.OutlineSize * InShapedGlyph.FontFaceData->FontScale)/TargetPpem;
 
-	float EmInnerSpread = FMath::Clamp(2.f/TargetPpem, 0.125f, 4.f);
-	float EmOuterSpread = FMath::Min(EmInnerSpread + EmOutlineSize, 4.f);
+	float EmInnerSpread = FMath::Clamp(2.f/TargetPpem, 0.05f, 8.f);
+	float EmOuterSpread = FMath::Min(EmInnerSpread + EmOutlineSize, 8.f);
 
 	// Has the atlas data already been cached on the glyph?
 	{
@@ -1254,7 +1262,7 @@ FSdfGlyphFontAtlasData FSlateFontCache::GetSdfGlyphFontAtlasData(const FShapedGl
 	};
 	// Not cached on the glyph, so create a key for to look up this glyph, as it may
 	// have already been cached by another shaped text sequence
-	const FSdfGlyphEntryKey GlyphKey(InShapedGlyph.FontFaceData->FontFace, InShapedGlyph.GlyphIndex, InSdfSettings.GetClampedPpem());
+	const FSdfGlyphEntryKey GlyphKey(InShapedGlyph.FontFaceData->FontFace, InShapedGlyph.GlyphIndex, InSdfSettings.GetClampedPpem(), EmOuterSpread, EmInnerSpread);
 
 	{
 		// Has the atlas data already been cached by another shaped text sequence?
@@ -1630,7 +1638,9 @@ void FSlateFontCache::UpdateCache()
 			{
 				const FSdfGlyphEntryKey GlyphKey(InRequestDescriptor.FontFace, 
 												 InRequestDescriptor.GlyphIndex, 
-												 InRequestDescriptor.Ppem);
+												 InRequestDescriptor.Ppem,
+												 InRequestDescriptor.EmOuterSpread,
+												 InRequestDescriptor.EmInnerSpread);
 				const FSdfGlyphTaskKey TaskKey(GlyphKey, InRequestDescriptor.EmOuterSpread, InRequestDescriptor.EmInnerSpread);
 				const TSharedRef<FSdfGlyphFontAtlasData>* FoundAtlasData = SdfTaskToAtlasData.Find(TaskKey);
 				if (FoundAtlasData)
diff --git a/Engine/Source/Runtime/SlateCore/Public/Fonts/FontCache.h b/Engine/Source/Runtime/SlateCore/Public/Fonts/FontCache.h
--- a/Engine/Source/Runtime/SlateCore/Public/Fonts/FontCache.h
+++ b/Engine/Source/Runtime/SlateCore/Public/Fonts/FontCache.h
@@ -251,13 +251,14 @@ private:
 struct FSdfGlyphEntryKey
 {
 public:
-	FSdfGlyphEntryKey(const TWeakPtr<FFreeTypeFace> InFontFace, uint32 InGlyphIndex, int32 InPpem);
+	FSdfGlyphEntryKey(const TWeakPtr<FFreeTypeFace> InFontFace, uint32 InGlyphIndex, int32 InPpem, float InEmOuterSpread, float InEmInnerSpread);
 
 	FORCEINLINE bool operator==(const FSdfGlyphEntryKey& Other) const
 	{
 		return FontFace == Other.FontFace
 			&& GlyphIndex == Other.GlyphIndex
-			&& Ppem == Other.Ppem;
+			&& Ppem == Other.Ppem
+			&& SpreadCategory == Other.SpreadCategory;
 	}
 
 	FORCEINLINE bool operator!=(const FSdfGlyphEntryKey& Other) const
@@ -276,10 +277,15 @@ public:
 	const uint32 GlyphIndex;
 	/** The pixel size at which the sdf glyph is generated */
 	const int32 Ppem;
+	/** The spread category. The spreads of a glyph entry can be arbitrary but similar values will share the same category and therefore glyph entry */
+	const int32 SpreadCategory;
 
 private:
 	/** Cached hash value used for map lookups */
 	uint32 KeyHash;
+
+	/** Computes the discrete spread category from real spread values specified in em */
+	static int32 GetSpreadCategory(float InEmOuterSpread, float InEmInnerSpread);
 };
 
 /** Used to lookup information about specific SDF generation tasks - unlike FSdfGlyphEntryKey is also identified by em spread */
