SlateCore: Moving TAnimatedAttribute from sandbox to SlateCore

diff --git a/Engine/Source/Runtime/SlateCore/Private/Animation/AnimatedAttributeManager.cpp b/Engine/Source/Runtime/SlateCore/Private/Animation/AnimatedAttributeManager.cpp
new file mode 100644
--- /dev/null
+++ b/Engine/Source/Runtime/SlateCore/Private/Animation/AnimatedAttributeManager.cpp
@@ -0,0 +1,93 @@
+// Copyright Epic Games, Inc. All Rights Reserved.
+
+#include "Animation/AnimatedAttributeManager.h"
+
+////////////////////////////////////////////////////////////////////////////////
+// TAnimatedAttributeBase
+////////////////////////////////////////////////////////////////////////////////
+
+TAnimatedAttributeBase::TAnimatedAttributeBase()
+	: bIsRegistered(false)
+{
+}
+
+TAnimatedAttributeBase::~TAnimatedAttributeBase()
+{
+}
+
+void TAnimatedAttributeBase::Register()
+{
+	if(!bIsRegistered)
+	{
+		FAnimatedAttributeManager::Get().RegisterAttribute(AsShared());
+		bIsRegistered = true;
+	}
+}
+
+void TAnimatedAttributeBase::Unregister()
+{
+	if(bIsRegistered)
+	{
+		FAnimatedAttributeManager::Get().UnregisterAttribute(AsShared());
+		bIsRegistered = false;
+	}
+}
+
+////////////////////////////////////////////////////////////////////////////////
+// FAnimatedAttributeManager
+////////////////////////////////////////////////////////////////////////////////
+
+FAnimatedAttributeManager::FAnimatedAttributeManager()
+{
+}
+
+FAnimatedAttributeManager::~FAnimatedAttributeManager()
+{
+}
+
+bool FAnimatedAttributeManager::Tick(float InDeltaTime)
+{
+	// remove state attributes
+	Attributes.RemoveAll([](const TWeakPtr<TAnimatedAttributeBase>& Attribute) {
+		return !Attribute.IsValid();
+	});
+
+	// at this point we expect the remaining attributes to be valid
+	for(const TWeakPtr<TAnimatedAttributeBase>& Attribute : Attributes)
+	{
+		Attribute.Pin()->Tick(InDeltaTime);
+	}
+
+	return true;
+}
+
+void FAnimatedAttributeManager::SetupTick()
+{
+	const FTickerDelegate TickDelegate = FTickerDelegate::CreateRaw(this, &FAnimatedAttributeManager::Tick);
+	TickHandle = FTSTicker::GetCoreTicker().AddTicker(TickDelegate);
+}
+
+void FAnimatedAttributeManager::TeardownTick()
+{
+	if(TickHandle.IsValid())
+	{
+		FTSTicker::GetCoreTicker().RemoveTicker(TickHandle);
+		TickHandle.Reset();
+	}
+}
+
+void FAnimatedAttributeManager::RegisterAttribute(const TSharedRef<TAnimatedAttributeBase>& InAttribute)
+{
+	Attributes.Add(InAttribute.ToWeakPtr());
+}
+
+void FAnimatedAttributeManager::UnregisterAttribute(const TSharedRef<TAnimatedAttributeBase>& InAttribute)
+{
+	Attributes.Remove(InAttribute.ToWeakPtr());
+}
+
+FAnimatedAttributeManager& FAnimatedAttributeManager::Get()
+{
+	static FAnimatedAttributeManager Manager;
+	return Manager;
+}
diff --git a/Engine/Source/Runtime/SlateCore/Private/Animation/AttributeInterpolator.cpp b/Engine/Source/Runtime/SlateCore/Private/Animation/AttributeInterpolator.cpp
new file mode 100644
--- /dev/null
+++ b/Engine/Source/Runtime/SlateCore/Private/Animation/AttributeInterpolator.cpp
@@ -0,0 +1,137 @@
+// Copyright Epic Games, Inc. All Rights Reserved.
+
+#include "Animation/AttributeInterpolator.h"
+#include "AHEasing/easing.h"
+
+float EaseInterpolatorRatio(EEasingInterpolatorType InEasingType, float InRatio)
+{
+	switch(InEasingType)
+	{
+		case EEasingInterpolatorType::Linear:
+		{
+			return FMath::Clamp<float>(InRatio, 0, 1);
+		}
+		case EEasingInterpolatorType::QuadraticEaseIn:
+		{
+			return QuadraticEaseIn(InRatio);
+		}
+		case EEasingInterpolatorType::QuadraticEaseOut:
+		{
+			return QuadraticEaseOut(InRatio);
+		}
+		case EEasingInterpolatorType::QuadraticEaseInOut:
+		{
+			return QuadraticEaseInOut(InRatio);
+		}
+		case EEasingInterpolatorType::CubicEaseIn:
+		{
+			return CubicEaseIn(InRatio);
+		}
+		case EEasingInterpolatorType::CubicEaseOut:
+		{
+			return CubicEaseOut(InRatio);
+		}
+		case EEasingInterpolatorType::CubicEaseInOut:
+		{
+			return CubicEaseInOut(InRatio);
+		}
+		case EEasingInterpolatorType::QuarticEaseIn:
+		{
+			return QuarticEaseIn(InRatio);
+		}
+		case EEasingInterpolatorType::QuarticEaseOut:
+		{
+			return QuarticEaseOut(InRatio);
+		}
+		case EEasingInterpolatorType::QuarticEaseInOut:
+		{
+			return QuarticEaseInOut(InRatio);
+		}
+		case EEasingInterpolatorType::QuinticEaseIn:
+		{
+			return QuinticEaseIn(InRatio);
+		}
+		case EEasingInterpolatorType::QuinticEaseOut:
+		{
+			return QuinticEaseOut(InRatio);
+		}
+		case EEasingInterpolatorType::QuinticEaseInOut:
+		{
+			return QuinticEaseInOut(InRatio);
+		}
+		case EEasingInterpolatorType::SineEaseIn:
+		{
+			return SineEaseIn(InRatio);
+		}
+		case EEasingInterpolatorType::SineEaseOut:
+		{
+			return SineEaseOut(InRatio);
+		}
+		case EEasingInterpolatorType::SineEaseInOut:
+		{
+			return SineEaseInOut(InRatio);
+		}
+		case EEasingInterpolatorType::CircularEaseIn:
+		{
+			return CircularEaseIn(InRatio);
+		}
+		case EEasingInterpolatorType::CircularEaseOut:
+		{
+			return CircularEaseOut(InRatio);
+		}
+		case EEasingInterpolatorType::CircularEaseInOut:
+		{
+			return CircularEaseInOut(InRatio);
+		}
+		case EEasingInterpolatorType::ExponentialEaseIn:
+		{
+			return ExponentialEaseIn(InRatio);
+		}
+		case EEasingInterpolatorType::ExponentialEaseOut:
+		{
+			return ExponentialEaseOut(InRatio);
+		}
+		case EEasingInterpolatorType::ExponentialEaseInOut:
+		{
+			return ExponentialEaseInOut(InRatio);
+		}
+		case EEasingInterpolatorType::ElasticEaseIn:
+		{
+			return ElasticEaseIn(InRatio);
+		}
+		case EEasingInterpolatorType::ElasticEaseOut:
+		{
+			return ElasticEaseOut(InRatio);
+		}
+		case EEasingInterpolatorType::ElasticEaseInOut:
+		{
+			return ElasticEaseInOut(InRatio);
+		}
+		case EEasingInterpolatorType::BackEaseIn:
+		{
+			return BackEaseIn(InRatio);
+		}
+		case EEasingInterpolatorType::BackEaseOut:
+		{
+			return BackEaseOut(InRatio);
+		}
+		case EEasingInterpolatorType::BackEaseInOut:
+		{
+			return BackEaseInOut(InRatio);
+		}
+		case EEasingInterpolatorType::BounceEaseIn:
+		{
+			return BounceEaseIn(InRatio);
+		}
+		case EEasingInterpolatorType::BounceEaseOut:
+		{
+			return BounceEaseOut(InRatio);
+		}
+		case EEasingInterpolatorType::BounceEaseInOut:
+		{
+			return BounceEaseInOut(InRatio);
+		}
+	}
+
+	return 0.f;
+}
\ No newline at end of file
diff --git a/Engine/Source/Runtime/SlateCore/Private/SlateCoreModule.cpp b/Engine/Source/Runtime/SlateCore/Private/SlateCoreModule.cpp
--- a/Engine/Source/Runtime/SlateCore/Private/SlateCoreModule.cpp
+++ b/Engine/Source/Runtime/SlateCore/Private/SlateCoreModule.cpp
@@ -11,7 +11,7 @@
 #include "Modules/ModuleManager.h"
 #include "SlateGlobals.h"
 #include "Types/SlateStructs.h"
-
+#include "Animation/AnimatedAttributeManager.h"
 
 DEFINE_LOG_CATEGORY(LogSlate);
 DEFINE_LOG_CATEGORY(LogSlateStyles);
@@ -38,6 +38,16 @@ public:
 #endif
 	}
 
+	virtual void StartupModule() override
+	{
+		FAnimatedAttributeManager::Get().SetupTick();
+	}
+
+	virtual void ShutdownModule() override
+	{
+		FAnimatedAttributeManager::Get().TeardownTick();
+	}
+
 #if WITH_SLATE_DEBUGGING
 private:
 	TUniquePtr<FConsoleSlateDebugger> SlateDebuggerEvent;
diff --git a/Engine/Source/Runtime/SlateCore/Public/Animation/AnimatedAttribute.h b/Engine/Source/Runtime/SlateCore/Public/Animation/AnimatedAttribute.h
new file mode 100644
--- /dev/null
+++ b/Engine/Source/Runtime/SlateCore/Public/Animation/AnimatedAttribute.h
@@ -0,0 +1,298 @@
+// Copyright Epic Games, Inc. All Rights Reserved.
+
+#pragma once
+
+#include "Misc/Attribute.h"
+#include "AnimatedAttributeManager.h"
+#include "AttributeInterpolator.h"
+
+/**
+ * Animated Attribute object
+ */
+template< typename NumericType >
+class TAnimatedAttribute : public TAnimatedAttributeBase
+{
+public:
+
+	/**
+	 * Attribute 'getter' delegate
+	 *
+	 * NumericType GetValue() const
+	 *
+	 * @return  The attribute's value
+	 */
+	using FGetter = TDelegate<NumericType()>;
+
+protected:
+
+	struct FPrivateToken { explicit FPrivateToken() = default; };
+	
+	/** Default constructor. */
+	TAnimatedAttribute() = delete;
+
+public:
+
+	/**
+	 * Construct only from interpolator settings
+	 * This is used only for the Create methods and only accessible from there
+	 * due to the use of a protected FPrivateToken. Please rely on the Create
+	 * methods if you are seeing a compiler error related to this. 
+	 *
+	 * @param  InSettings		The settings for the interpolator
+	 */
+	template< typename InterpolatorSettings >
+	TAnimatedAttribute( FPrivateToken, const InterpolatorSettings& InSettings )
+		: Interpolator( MakeUnique<InterpolatorSettings::InterpolatorType>(InSettings) )
+		, Getter()
+		, DesiredValue() 
+	{
+	}
+
+	/**
+	 * Static: Creates an animated attribute implicitly from an initial value
+	 *
+	 * @param  InSettings		The settings for the interpolator
+	 * @param  InInitialValue	The value for this attribute
+	 */
+	template< typename InterpolatorSettings, typename OtherType >
+	[[nodiscard]] static TSharedRef< TAnimatedAttribute > Create( const InterpolatorSettings& InSettings, OtherType&& InInitialValue )
+	{
+		TSharedRef< TAnimatedAttribute< NumericType > > Attribute = MakeShared< TAnimatedAttribute< NumericType > >(FPrivateToken{}, InSettings);
+		Attribute->Set(Forward<OtherType>(InInitialValue));
+		Attribute->Register();
+		return Attribute;
+	}
+	
+	/**
+	 * Static: Creates an animated attribute that's pre-bound to the specified 'getter' delegate
+	 *
+	 * @param  InSettings		The settings for the interpolator
+	 * @param  InGetter			Delegate to bind
+	 */
+	template< typename InterpolatorSettings >
+	[[nodiscard]] static TSharedRef< TAnimatedAttribute > CreateWithGetter( const InterpolatorSettings& InSettings, const FGetter& InGetter )
+	{
+		TSharedRef< TAnimatedAttribute< NumericType > > Attribute = MakeShared< TAnimatedAttribute< NumericType > >(FPrivateToken{}, InSettings);
+		Attribute->Getter = InGetter;
+		Attribute->Set(Attribute->Getter.Execute());
+		Attribute->Register();
+		return Attribute;
+	}
+	
+	/**
+	 * Static: Creates an animated attribute that's pre-bound to the specified 'getter' delegate
+	 *
+	 * @param  InSettings		The settings for the interpolator
+	 * @param  InGetter			Delegate to bind
+	 */
+	template< typename InterpolatorSettings >
+	[[nodiscard]] static TSharedRef< TAnimatedAttribute > CreateWithGetter( const InterpolatorSettings& InSettings, FGetter&& InGetter )
+	{
+		TSharedRef< TAnimatedAttribute< NumericType > > Attribute = MakeShared< TAnimatedAttribute< NumericType > >(FPrivateToken{}, InSettings);
+		Attribute->Getter = MoveTemp(InGetter);
+		Attribute->Set(Attribute->Getter.Execute());
+		Attribute->Register();
+		return Attribute;
+	}
+	
+	/**
+	 * Sets the attribute's value while keeping a running interpolation going.
+	 * The value will be become the new desired value and the interpolation will
+	 * continue trying to reach it. If the interpolation is currently stopped the
+	 * animation / interpolation will restart towards the new provided value.
+	 *
+	 * @param  InNewValue  The value to set the attribute to
+	 */
+	template< typename OtherType >
+	void Set( const OtherType& InNewValue )
+	{
+		Getter.Unbind();
+		DesiredValue = InNewValue;
+		Interpolator->SetValue(DesiredValue.GetValue());
+	}
+
+	/**
+	 * Sets the attribute's value while keeping a running interpolation going.
+	 * The value will be become the new desired value and the interpolation will
+	 * continue trying to reach it. If the interpolation is currently stopped the
+	 * animation / interpolation will restart towards the new provided value.
+	 * 
+	 * @param InNewValue  The value to set the attribute to
+	 */
+	void Set( NumericType&& InNewValue )
+	{
+		Getter.Unbind();
+		DesiredValue = InNewValue;
+		Interpolator->SetValue(DesiredValue.GetValue());
+	}
+
+	/**
+	 * Sets the attribute's value and disables animation. The interpolator
+	 * will be stopped and the value will become the new current value.
+	 * 
+	 * @param InNewValue  The value to set the attribute to
+	 */
+	void SetValueAndStop( NumericType&& InNewValue )
+	{
+		Set(Forward<NumericType>(InNewValue));
+		Interpolator->Reset();
+	}
+
+	/** Was this TAnimatedAttribute ever assigned? */
+	bool IsSet() const
+	{
+		return DesiredValue.IsSet() || IsBound();
+	}
+
+	/**
+	 * Gets the attribute's current value.
+	 * Assumes that the attribute is set.
+	 *
+	 * @return  The attribute's value
+	 */
+	const NumericType& Get() const
+	{
+		// first see if the interpolator has a value
+		if(Interpolator->IsSet())
+		{
+			return Interpolator->Get();
+		}
+		
+		// If we have a getter delegate, then we'll call that to generate the value
+		if( Getter.IsBound() )
+		{
+			// Call the delegate to get the value.  Note that this will assert if the delegate is not currently
+			// safe to call (e.g. object was deleted and we could detect that)
+
+			// NOTE: We purposely overwrite our value copy here so that we can return the value by address in
+			// the most common case, which is an attribute that doesn't have a delegate bound to it at all.
+			DesiredValue = Getter.Execute();
+		}
+
+		if(DesiredValue.IsSet())
+		{
+			// Return the stored value
+			return DesiredValue.GetValue();
+		}
+
+		static const NumericType EmptyResult = NumericType();
+		return EmptyResult;
+	}
+
+	/**
+	 * Gets the attribute's current value. The attribute may not be set, in which case use the default value provided.
+	 * Shorthand for the boilerplate code: MyAttribute.IsSet() ? MyAttribute.Get() : DefaultValue
+	 */
+	const NumericType& Get( const NumericType& DefaultValue ) const
+	{
+		return IsSet() ? Get() : DefaultValue;
+	}
+
+	/**
+	 * Returns the overall deltatime of the interpolator
+	 */
+	double GetOverAllDeltaTime() const
+	{
+		return Interpolator->GetOverAllDeltaTime();
+	}
+
+	/**
+	 * Sets the attribute's delay as a one shot
+	 */
+	void SetDelayOneShot( double InDelay )
+	{
+		Interpolator->SetDelayOneShot(InDelay);
+	}
+
+	/**
+	 * Sets the tolerance of this interpolator
+	 */
+	void SetTolerance( double Tolerance )
+	{
+		Interpolator->SetTolerance(Tolerance);
+	}
+
+	/**
+	 * Checks to see if this attribute has a 'getter' function bound
+	 *
+	 * @return  True if attribute is bound to a getter function
+	 */
+	bool IsBound() const
+	{
+		return Getter.IsBound();
+	}
+
+	/**
+	 * Is this attribute identical to another TAnimationAttribute.
+	 *
+	 * @param InOther The other attribute to compare with.
+	 * @return true if the attributes are identical, false otherwise.
+	 */
+	bool IdenticalTo(const TAnimatedAttribute& InOther) const
+	{
+		if(!Interpolator->IdenticalTo(InOther.Interpolator.Get()))
+		{
+			return false;
+		}
+
+		const bool bIsBound = IsBound();
+		if ( bIsBound == InOther.IsBound() )
+		{
+			if ( bIsBound )
+			{
+				return Getter.GetHandle() == InOther.Getter.GetHandle();
+			}
+			return IsSet() == InOther.IsSet() && DesiredValue == InOther.DesiredValue;
+		}
+
+		return false;
+	}
+
+	/*
+	 * Returns the delegate to react to the interpolator starting to interpolate
+	 */
+	typename TAttributeInterpolator< NumericType >::FInterpolatorEvent& OnInterpolationStarted()
+	{
+		return Interpolator->OnInterpolationStarted();
+	}
+
+	/*
+	 * Returns the delegate to react to the interpolator ending to interpolate
+	 */
+	typename TAttributeInterpolator< NumericType >::FInterpolatorEvent& OnInterpolationStopped()
+	{
+		return Interpolator->OnInterpolationStopped();
+	}
+
+private:
+
+	/**
+	 * Updates the interpolator / ticks the animation
+	 */
+	virtual void Tick(float InDeltaTime) override
+	{
+		if( Getter.IsBound() )
+		{
+			Interpolator->SetValue(Getter.Execute());
+		}
+		if(DesiredValue.IsSet())
+		{
+			Interpolator->SetValue(DesiredValue.GetValue());
+		}
+		Interpolator->Tick(InDeltaTime);
+	}
+
+	// We declare ourselves as a friend (templated using OtherType) so we can access members as needed
+	template< class OtherType > friend class TAnimatedAttribute;
+
+	/** The attribute's interpolator */
+	TUniquePtr< TAttributeInterpolator< NumericType > > Interpolator;
+
+	/** Bound member function for this attribute (may be NULL if no function is bound.)  When set, all attempts
+		to read the attribute's value will instead call this delegate to generate the value. */
+	/** Our attribute's 'getter' delegate */
+	FGetter Getter;
+	
+	/** Current value.  Mutable so that we can cache the value locally when using a bound Getter (allows const ref return value.) */
+	mutable TOptional< NumericType > DesiredValue;
+};
diff --git a/Engine/Source/Runtime/SlateCore/Public/Animation/AnimatedAttributeManager.h b/Engine/Source/Runtime/SlateCore/Public/Animation/AnimatedAttributeManager.h
new file mode 100644
--- /dev/null
+++ b/Engine/Source/Runtime/SlateCore/Public/Animation/AnimatedAttributeManager.h
@@ -0,0 +1,59 @@
+// Copyright Epic Games, Inc. All Rights Reserved.
+
+#pragma once
+
+#include "CoreMinimal.h"
+#include "Containers/Ticker.h"
+
+/**
+ * Animated Attribute Base Layer - used for registering the attribute centrally
+ */
+class TAnimatedAttributeBase : public TSharedFromThis<TAnimatedAttributeBase>
+{
+protected:
+	
+	SLATECORE_API TAnimatedAttributeBase();
+	SLATECORE_API virtual ~TAnimatedAttributeBase();
+
+	SLATECORE_API void Register();
+	SLATECORE_API void Unregister();
+
+	virtual void Tick(float InDeltaTime) = 0;
+
+private:
+
+	TAnimatedAttributeBase(const TAnimatedAttributeBase& InOther) = delete;
+	TAnimatedAttributeBase& operator =(const TAnimatedAttributeBase& Other) = delete;
+
+	bool bIsRegistered;
+
+	friend class FAnimatedAttributeManager;
+};
+
+/**
+ * A central manager for animated attributes
+ */
+class FAnimatedAttributeManager
+{
+
+public:
+
+	FAnimatedAttributeManager();
+	~FAnimatedAttributeManager();
+
+	SLATECORE_API static FAnimatedAttributeManager& Get();
+	SLATECORE_API bool Tick(float InDeltaTime);
+
+	SLATECORE_API void SetupTick();
+	SLATECORE_API void TeardownTick();
+
+private:
+
+	void RegisterAttribute(const TSharedRef<TAnimatedAttributeBase>& InAttribute);
+	void UnregisterAttribute(const TSharedRef<TAnimatedAttributeBase>& InAttribute);
+
+	TArray<TWeakPtr<TAnimatedAttributeBase>> Attributes;
+	FTSTicker::FDelegateHandle TickHandle;
+
+	friend class TAnimatedAttributeBase;
+};
diff --git a/Engine/Source/Runtime/SlateCore/Public/Animation/AttributeInterpolator.h b/Engine/Source/Runtime/SlateCore/Public/Animation/AttributeInterpolator.h
new file mode 100644
--- /dev/null
+++ b/Engine/Source/Runtime/SlateCore/Public/Animation/AttributeInterpolator.h
@@ -0,0 +1,636 @@
+// Copyright Epic Games, Inc. All Rights Reserved.
+
+#pragma once
+
+#include "CoreMinimal.h"
+#include "Math/UnrealMathUtility.h"
+
+enum class EAttributeInterpolatorType : uint8
+{
+	Easing,
+	Arrive,
+	Verlet
+};
+
+/**
+ * Attribute Interpolator Base Class
+ */
+template< typename NumericType >
+class TAttributeInterpolator
+{
+
+public:
+
+	DECLARE_MULTICAST_DELEGATE_OneParam(FInterpolatorEvent, TOptional<NumericType>);
+
+	/*
+	 * Default constructor
+	 */
+	TAttributeInterpolator()
+	: Tolerance(0.001f)
+	{
+	}
+
+	/*
+	 * Default destructor
+	 */
+	virtual ~TAttributeInterpolator()
+	{
+	}
+
+	/**
+	 * Helper function to compare to values - to allow to specialize per type as needed
+	 *
+	 * @param InA The first value to compare
+	 * @param InB The second value to compare
+	 * @return True if the two values are considered "equal"
+	 */
+	static bool Equals(const NumericType& InA, const NumericType& InB, double InTolerance = 0.001)
+	{
+		return InA.Equals(InB, InTolerance);
+	}
+
+	/**
+	 * Returns true if the interpolator is of a given type
+	 *
+	 * @param InType The type to compare agains
+	 * @return True if the interpolator matches the type
+	 */
+	virtual bool IsTypeOf(EAttributeInterpolatorType InType) const = 0;
+
+	/**
+	 * Compares this interpolator to another one
+	 *
+	 * @param InOther The interpolator to compare this one to
+	 * @return True if the this interpolator is considered identical to the passed one
+	 */
+	virtual bool IdenticalTo(const TAttributeInterpolator* InOther) const = 0;
+
+	/**
+	 * Returns true if this interpolator's value has even been set
+	 * @return True if this interpolator's value has even been set
+	 */
+	bool IsSet() const
+	{
+		return DesiredValue.IsSet();
+	}
+
+	/**
+	 * Resets the interpolator
+	 */
+	void Reset() const
+	{
+		const bool bFireStopEvent = Stop();
+		InterpolatedValue.Reset();
+		LastValue.Reset();
+		DesiredValue.Reset();
+		if(bFireStopEvent)
+		{
+			InterpolationStopped.Broadcast(TOptional<NumericType>());
+		}
+	}
+
+	/**
+ 	 * Set the internal state of the interpolator
+ 	 * 
+ 	 * @param InValue The value to be set on the interpolator
+ 	 */
+	void SetValue(const NumericType& InValue) const
+	{
+		if(DesiredValue.IsSet())
+		{
+			if(Equals(DesiredValue.GetValue(), InValue, Tolerance))
+			{
+				return;
+			}
+
+			LastValue = InterpolatedValue.IsSet() ? InterpolatedValue : DesiredValue;
+			DesiredValue = InValue;
+		}
+		else
+		{
+			SetValueAndStop(InValue);
+		}
+	}
+
+	/**
+	 * Set the internal state of the interpolator and disables animation
+	 * 
+	 * @param InValue The value to be set on the interpolator
+	 */
+	void SetValueAndStop(const NumericType& InValue) const
+	{
+		const bool bFireStopEvent = Stop();
+		LastValue = DesiredValue = InValue;
+		if(bFireStopEvent)
+		{
+			InterpolationStopped.Broadcast(InValue);
+		}
+	}
+
+	/**
+	 * Returns the interpolated value
+	 *
+	 * @param InValue The value to interpolate against
+	 * @return The interpolated value
+	 */
+	const NumericType& Get() const
+	{
+		check(IsSet());
+		if(InterpolatedValue.IsSet())
+		{
+			return InterpolatedValue.GetValue();
+		}
+		if(LastValue.IsSet())
+		{
+			return LastValue.GetValue();
+		}
+		return DesiredValue.GetValue();
+	}
+
+	/**
+	 * Updates the interpolator
+	 */
+	void Tick(float InDeltaTime)
+	{
+		if(!IsPlaying() && LastValue.IsSet())
+		{
+			const NumericType& CurrentValue = InterpolatedValue.IsSet() ? InterpolatedValue.GetValue() : DesiredValue.GetValue();
+			if(Equals(CurrentValue, LastValue.GetValue(), Tolerance))
+			{
+				SetValueAndStop(LastValue.GetValue());
+				return;
+			}
+		}
+
+		(void)PlayIfStopped();
+		Interpolate();
+		
+		if(OverallDeltaTime.IsSet())
+		{
+			OverallDeltaTime = OverallDeltaTime.GetValue() + InDeltaTime;
+		}
+	}
+
+	/**
+	 * Returns the delta time between the start of the animation and now
+	 */
+	double GetOverAllDeltaTime() const
+	{
+		if(IsPlaying())
+		{
+			if(OverallDeltaTime.IsSet())
+			{
+				return OverallDeltaTime.GetValue();
+			}
+		}
+		return 0.0f; 
+	}
+
+	/**
+	 * Sets the delay for the next shot - delays are consumed and only used once
+	 */
+	void SetDelayOneShot(float InDelay) const
+	{
+		Delay = InDelay;
+	}
+
+	/**
+	 * Sets the tolerance on this interpolator
+	 */
+	void SetTolerance(float InTolerance)
+	{
+		Tolerance = InTolerance;
+	}
+
+	/*
+	 * Returns the delegate to react to the interpolator starting to interpolate
+	 */
+	FInterpolatorEvent& OnInterpolationStarted() { return InterpolationStarted; }
+
+	/*
+	 * Returns the delegate to react to the interpolator ending to interpolate
+	 */
+	FInterpolatorEvent& OnInterpolationStopped() { return InterpolationStopped; }
+
+protected:
+
+	// Interpolates the internal value against the new value
+	// This method is expected to set the Value member
+	virtual void Interpolate() const = 0;
+
+	// Returns true if this interpolator is currently playing
+	bool IsPlaying() const
+	{
+		return InterpolatedValue.IsSet();
+	}
+
+	// Starts the animation and records time
+	virtual bool Start() const
+	{
+		InterpolatedValue = LastValue;
+		OverallDeltaTime = 0.f;
+		return true;
+	}
+
+	virtual bool PlayIfStopped() const
+	{
+		if(!IsPlaying())
+		{
+			return Start();
+		}
+		return false;
+	}
+
+	// Stops the animation and stops recording time
+	virtual bool Stop() const
+	{
+		Delay.Reset();
+		OverallDeltaTime.Reset();
+		if(IsPlaying())
+		{
+			InterpolatedValue.Reset();
+			return true;
+		}
+		return false;
+	}
+
+	// restart the animation
+	virtual bool Restart() const
+	{
+		return Start();
+	}
+
+	float Tolerance;
+	mutable TOptional<float> Delay;
+	mutable TOptional<float> OverallDeltaTime;
+	mutable TOptional<NumericType> InterpolatedValue;
+	mutable TOptional<NumericType> LastValue;
+	mutable TOptional<NumericType> DesiredValue;
+	
+	FInterpolatorEvent InterpolationStarted;
+	FInterpolatorEvent InterpolationStopped;
+
+	template<typename> friend class TAnimatedAttribute;
+};
+
+template<>
+inline bool TAttributeInterpolator<float>::Equals(const float& InA, const float& InB, double InTolerance)
+{
+	return FMath::IsNearlyEqual(InA, InB, (float)InTolerance);
+}
+
+template<>
+inline bool TAttributeInterpolator<double>::Equals(const double& InA, const double& InB, double InTolerance)
+{
+	return FMath::IsNearlyEqual(InA, InB, InTolerance);
+}
+
+enum class EEasingInterpolatorType : uint8
+{
+	Linear,
+	QuadraticEaseIn,
+	QuadraticEaseOut,
+	QuadraticEaseInOut,
+	CubicEaseIn,
+	CubicEaseOut,
+	CubicEaseInOut,
+	QuarticEaseIn,
+	QuarticEaseOut,
+	QuarticEaseInOut,
+	QuinticEaseIn,
+	QuinticEaseOut,
+	QuinticEaseInOut,
+	SineEaseIn,
+	SineEaseOut,
+	SineEaseInOut,
+	CircularEaseIn,
+	CircularEaseOut,
+	CircularEaseInOut,
+	ExponentialEaseIn,
+	ExponentialEaseOut,
+	ExponentialEaseInOut,
+	ElasticEaseIn,
+	ElasticEaseOut,
+	ElasticEaseInOut,
+	BackEaseIn,
+	BackEaseOut,
+	BackEaseInOut,
+	BounceEaseIn,
+	BounceEaseOut,
+	BounceEaseInOut
+};
+
+SLATECORE_API float EaseInterpolatorRatio(EEasingInterpolatorType InEasingType, float InRatio);
+
+template< typename NumericType >
+class TEasingAttributeInterpolator;
+
+/**
+ * An easing attribute interpolator.
+ * This interpolator is to implement default easing functions such as Cubic etc
+ * probably to rely on FCurveSequence.
+ */
+template< typename NumericType >
+class TEasingAttributeInterpolator : public TAttributeInterpolator<NumericType>
+{
+private:
+
+	typedef TAttributeInterpolator<NumericType> Super;
+
+public:
+
+	struct FSettings
+	{
+		typedef TEasingAttributeInterpolator<NumericType> InterpolatorType;
+	
+		FSettings(EEasingInterpolatorType InEasingType, float InDuration)
+			: EasingType(InEasingType)
+			, Duration(InDuration)
+		{
+		}
+		
+		EEasingInterpolatorType EasingType;
+		float Duration;
+	};
+
+	explicit TEasingAttributeInterpolator(const FSettings& InSettings)
+		: Super()
+		, Settings(InSettings)
+	{
+	}
+
+	//////////////////////////////////////////////////////////
+	// TAttributeInterpolator interface
+	
+	virtual bool IsTypeOf(EAttributeInterpolatorType InType) const override
+	{
+		return InType == EAttributeInterpolatorType::Easing;
+	}
+
+	virtual bool IdenticalTo(const TAttributeInterpolator<NumericType>* InOther) const override
+	{
+		if(!InOther->IsTypeOf(EAttributeInterpolatorType::Easing))
+		{
+			return false; 
+		}
+
+		// C++ cast due to polymorphism dynamic_cast is not available
+		const TEasingAttributeInterpolator<NumericType>* CastInterpolator =
+			(const TEasingAttributeInterpolator<NumericType>*)InOther;
+
+		return FMath::IsNearlyEqual(Settings.Duration, CastInterpolator->Settings.Duration);
+	}
+
+protected:
+
+	virtual void Interpolate() const override
+	{
+		check(Super::LastValue.IsSet());
+		check(Super::DesiredValue.IsSet());
+
+		const float DeltaTime = (float)Super::GetOverAllDeltaTime();
+		if(Settings.Duration <= 0.f || DeltaTime > Settings.Duration)
+		{
+			Super::SetValueAndStop(Super::DesiredValue.GetValue());
+			return;
+		}
+
+		const float Ratio = FMath::Clamp<float>(DeltaTime / Settings.Duration, 0.f, 1.f);
+		const float EasedRatio = EaseInterpolatorRatio(Settings.EasingType, Ratio);
+		Super::InterpolatedValue = FMath::Lerp<NumericType>(Super::LastValue.GetValue(), Super::DesiredValue.GetValue(), EasedRatio);
+	}
+
+	FSettings Settings;
+};
+
+/**
+ * An arrive attribute interpolator. Arrives smoothly at the desired value.
+ */
+template< typename NumericType >
+class TArriveAttributeInterpolator : public TAttributeInterpolator<NumericType>
+{
+private:
+
+	typedef TAttributeInterpolator<NumericType> Super;
+
+public:
+
+	struct FSettings
+	{
+		typedef TArriveAttributeInterpolator<NumericType> InterpolatorType;
+
+		FSettings(int32 InIterations, float InStrength)
+		: Iterations(InIterations)
+		, Strength(InStrength)
+		{
+		}
+		
+		int32 Iterations;
+		float Strength;
+	};
+
+	explicit TArriveAttributeInterpolator(const FSettings& InSettings)
+		: Super()
+		, Settings(InSettings)
+		, LastDeltaTime(0)
+	{
+	}
+
+	//////////////////////////////////////////////////////////
+	// TAttributeInterpolator interface
+	
+	virtual bool IsTypeOf(EAttributeInterpolatorType InType) const override
+	{
+		return InType == EAttributeInterpolatorType::Arrive;
+	}
+
+	virtual bool IdenticalTo(const TAttributeInterpolator<NumericType>* InOther) const override
+	{
+		if(!InOther->IsTypeOf(EAttributeInterpolatorType::Arrive))
+		{
+			return false; 
+		}
+
+		// C++ cast due to polymorphism dynamic_cast is not available
+		const TArriveAttributeInterpolator<NumericType>* CastInterpolator =
+			(const TArriveAttributeInterpolator<NumericType>*)InOther;
+
+		return (Settings.Iterations == CastInterpolator->Settings.Iterations) && 
+			FMath::IsNearlyEqual(Settings.Strength, CastInterpolator->Settings.Strength);
+	}
+
+protected:
+
+	virtual bool Start() const override
+	{
+		LastDeltaTime = 0.0;
+		return Super::Start();
+	}
+
+	virtual bool Restart() const override
+	{
+		const bool bWasPlaying = Super::IsPlaying();
+		TOptional<NumericType> PreviousValue = Super::InterpolatedValue;
+		
+		const bool bResult = Super::Restart();
+		if(bWasPlaying)
+		{
+			if(PreviousValue.IsSet())
+			{
+				Super::LastValue = Super::InterpolatedValue = PreviousValue;
+			}
+		}
+		return bResult;
+	}
+	
+	virtual void Interpolate() const override
+	{
+		check(Super::LastValue.IsSet());
+		check(Super::DesiredValue.IsSet());
+		
+		const float OverallDeltaTime = (float)Super::GetOverAllDeltaTime();
+		const float LocalDeltaTime = OverallDeltaTime - LastDeltaTime;
+
+		const float Blend =  FMath::Max<float>(Settings.Strength, 0) * 8.f * LocalDeltaTime;
+
+		for(int32 Iteration = 0; Iteration < Settings.Iterations; Iteration++)
+		{
+			Super::InterpolatedValue = FMath::Lerp<NumericType>(
+				Super::InterpolatedValue.IsSet() ? Super::InterpolatedValue.GetValue() : Super::LastValue.GetValue(),
+				Super::DesiredValue.GetValue(),
+				Blend);
+		}
+
+		LastDeltaTime = OverallDeltaTime;
+
+		if(Super::Equals(Super::DesiredValue.GetValue(), Super::InterpolatedValue.GetValue(), Super::Tolerance))
+		{
+			Super::SetValueAndStop(Super::DesiredValue.GetValue());
+		}
+	}
+	
+	FSettings Settings;
+	mutable float LastDeltaTime;
+};
+
+/**
+ * A verlet attribute interpolator. Desired values are simulated relying a simple (semi-fake) verlet integration.
+ */
+template< typename NumericType >
+class TVerletAttributeInterpolator : public TAttributeInterpolator<NumericType>
+{
+private:
+
+	typedef TAttributeInterpolator<NumericType> Super;
+
+public:
+
+	struct FSettings
+	{
+		typedef TVerletAttributeInterpolator<NumericType> InterpolatorType;
+
+		FSettings(float InBlend, float InStrength, float InDamping)
+		: Blend(InBlend)
+		, Strength(InStrength)
+		, Damping(InDamping)
+		{
+		}
+		
+		float Blend;
+		float Strength;
+		float Damping;
+	};
+
+	explicit TVerletAttributeInterpolator(const FSettings& InSettings)
+		: Super()
+		, Settings(InSettings)
+		, LastDeltaTime(0)
+	{
+	}
+
+	//////////////////////////////////////////////////////////
+	// TAttributeInterpolator interface
+	
+	virtual bool IsTypeOf(EAttributeInterpolatorType InType) const override
+	{
+		return InType == EAttributeInterpolatorType::Verlet;
+	}
+
+	virtual bool IdenticalTo(const TAttributeInterpolator<NumericType>* InOther) const override
+	{
+		if(!InOther->IsTypeOf(EAttributeInterpolatorType::Verlet))
+		{
+			return false; 
+		}
+
+		// C++ cast due to polymorphism dynamic_cast is not available
+		const TVerletAttributeInterpolator<NumericType>* CastInterpolator =
+			(const TVerletAttributeInterpolator<NumericType>*)InOther;
+
+		return FMath::IsNearlyEqual(Settings.Blend, CastInterpolator->Settings.Blend) &&
+				FMath::IsNearlyEqual(Settings.Strength, CastInterpolator->Settings.Strength) && 
+			FMath::IsNearlyEqual(Settings.Damping, CastInterpolator->Settings.Damping);
+	}
+
+protected:
+
+	virtual bool Start() const override
+	{
+		LastDeltaTime = 0.0;
+		Velocity = NumericType(0);
+		return Super::Start();
+	}
+
+	virtual bool Restart() const override
+	{
+		const bool bWasPlaying = Super::IsPlaying();
+		TOptional<NumericType> PreviousValue = Super::InterpolatedValue;
+		TOptional<NumericType> PreviousVelocity = Velocity;
+		
+		const bool bResult = Super::Restart();
+		if(bWasPlaying)
+		{
+			if(PreviousValue.IsSet())
+			{
+				Super::LastValue = Super::InterpolatedValue = PreviousValue;
+			}
+			if(PreviousVelocity.IsSet())
+			{
+				Velocity = PreviousVelocity;
+			}
+		}
+		return bResult;
+	}
+
+	virtual void Interpolate() const override
+	{
+		check(Super::LastValue.IsSet());
+		check(Super::DesiredValue.IsSet());
+		
+		const float OverallDeltaTime = (float)Super::GetOverAllDeltaTime();
+		const float LocalDeltaTime = OverallDeltaTime - LastDeltaTime;
+
+		const float U = FMath::Clamp<float>(FMath::Max<float>(Settings.Blend, 0) * 8.f * LocalDeltaTime, 0, 1);
+		const float ScaleDown = FMath::Clamp<float>(1.f - Settings.Damping, 0.f, 1.f);
+
+		const NumericType PreviousValue = Super::InterpolatedValue.IsSet() ? Super::InterpolatedValue.GetValue() : Super::LastValue.GetValue();
+		const NumericType Force = (Super::DesiredValue.GetValue() - PreviousValue) * Settings.Strength;
+		const NumericType PreviousVelocity = Velocity.GetValue();
+
+		Velocity = FMath::Lerp<NumericType>(PreviousVelocity, Force, U) * ScaleDown;
+		Super::InterpolatedValue = PreviousValue + Velocity.GetValue() * LocalDeltaTime;
+
+		LastDeltaTime = OverallDeltaTime;
+
+		if(Super::Equals(Super::DesiredValue.GetValue(), Super::InterpolatedValue.GetValue(), Super::Tolerance) &&
+			Super::Equals(Velocity.GetValue(), NumericType(0), Super::Tolerance))
+		{
+			Super::SetValueAndStop(Super::DesiredValue.GetValue());
+		}
+	}
+
+	FSettings Settings;
+	mutable float LastDeltaTime;
+	mutable TOptional<NumericType> Velocity;
+};
diff --git a/Engine/Source/Runtime/SlateCore/SlateCore.Build.cs b/Engine/Source/Runtime/SlateCore/SlateCore.Build.cs
--- a/Engine/Source/Runtime/SlateCore/SlateCore.Build.cs
+++ b/Engine/Source/Runtime/SlateCore/SlateCore.Build.cs
@@ -52,5 +52,7 @@ public class SlateCore : ModuleRules
 		{
 			AddEngineThirdPartyPrivateStaticDependencies(Target, "XInput");
 		}
+		
+		PrivateIncludePaths.Add(System.IO.Path.Combine(EngineDirectory,"Source/ThirdParty/AHEasing/AHEasing-1.3.2"));
 	}
 }
