Fix for MSDF texture channel mix-up on mobile
#rnx <Release Notes exclude>

diff --git a/Engine/Shaders/Private/SlateElementPixelShader.usf b/Engine/Shaders/Private/SlateElementPixelShader.usf
--- a/Engine/Shaders/Private/SlateElementPixelShader.usf
+++ b/Engine/Shaders/Private/SlateElementPixelShader.usf
@@ -278,11 +278,11 @@ half4 GetSdfTextElementColor( VertexToPixelInterpolants InVertex )
 	half4 OutColor = InVertex.Color;
 #endif
 
-#if SHADER_TYPE == ST_MsdfFont
 	const float4 MultiDistance = float4(Texture2DSample( ElementTexture, ElementTextureSampler, AtlasTextureCoordinates ));
+#if SHADER_TYPE == ST_MsdfFont
 	const float SampledDistance = max(min(MultiDistance.r, MultiDistance.g), min(max(MultiDistance.r, MultiDistance.g), MultiDistance.b));
 #else
-	const float SampledDistance = float(Texture2DSample_A8( ElementTexture, ElementTextureSampler, AtlasTextureCoordinates ));
+	const float SampledDistance = lerp(MultiDistance A8_SAMPLE_MASK, MultiDistance.a, ShaderParams.w);
 #endif
 	const float2 ScreenSpaceTextureSize = float2(1.0, 1.0) / fwidth(AtlasTextureCoordinates);
 	const float ScreenSpaceSpread = max(dot(ShaderParams.xy, ScreenSpaceTextureSize), 1.0);
diff --git a/Engine/Source/Runtime/SlateCore/Private/Rendering/ElementBatcher.cpp b/Engine/Source/Runtime/SlateCore/Private/Rendering/ElementBatcher.cpp
--- a/Engine/Source/Runtime/SlateCore/Private/Rendering/ElementBatcher.cpp
+++ b/Engine/Source/Runtime/SlateCore/Private/Rendering/ElementBatcher.cpp
@@ -3163,8 +3163,8 @@ void FSlateElementBatcher::BuildShapedTextSequence(const FShapedTextBuildContext
 	bool bNeedSpaceForEllipsis = false;
 	bool bIsSdfFont = GlyphSequenceToRender->IsSdfFont();
 	bool bRequiresManualSkewing = bIsSdfFont && !FMath::IsNearlyEqual(GlyphSequenceToRender->GetFontSkew(), 0.f);
-	// Note - it would be much better to pass shader params as per-vertex attribute to avoid having to switch batches too often
-	FVector4f SdfShaderParams(0.f, 0.f, 0.f, 0.f);
+	float SdfPixelSpread = 0;
+	float SdfBias = 0;
 	// For left to right overflow direction - Sum of total whitespace we're currently advancing through. Once a non-whitespace glyph is detected this will return to 0
 	// For right to left this value is unused. We just skip all leading whitespace
 	float PreviousWhitespaceAdvance = 0;
@@ -3200,12 +3200,13 @@ void FSlateElementBatcher::BuildShapedTextSequence(const FShapedTextBuildContext
 		FVector2f SpriteOffset(0.f, 0.f);
 		FVector2f QuadMeshSize(0.f, 0.f);
 		FVector2f QuadMeshOffsets(0.f, 0.f);
-		FVector2f GlyphShaderParams(0.f, 0.f);
 
 		if (bCanRenderGlyph)
 		{
 			// Get Sizing and atlas info
 			int8 NextAtlasDataTextureIndex = -1;
+			float NextSdfPixelSpread = 0;
+			float NextSdfBias = 0;
 			const bool bIsSdfGlyph = bIsSdfFont && GlyphToRender.FontFaceData && GlyphToRender.FontFaceData->bSupportsSdf;
 
 			if (bIsSdfGlyph)
@@ -3230,9 +3231,9 @@ void FSlateElementBatcher::BuildShapedTextSequence(const FShapedTextBuildContext
 					// SdfGlyphAtlasData.Metrics values assume this operation, so QuadMeshSize already accounts for this
 					SpriteSize = FVector2f(SdfGlyphAtlasData.USize-1, SdfGlyphAtlasData.VSize-1);
 					SpriteOffset = FVector2f(SdfGlyphAtlasData.StartU+.5f, SdfGlyphAtlasData.StartV+.5f);
-					GlyphShaderParams.X = 0.5f*(SdfGlyphAtlasData.EmInnerSpread+SdfGlyphAtlasData.EmOuterSpread)*static_cast<float>(FontSdfSettings.GetClampedPpem());
+					NextSdfPixelSpread = (SdfGlyphAtlasData.EmInnerSpread+SdfGlyphAtlasData.EmOuterSpread)*static_cast<float>(FontSdfSettings.GetClampedPpem());
 					// Value representing zero distance
-					GlyphShaderParams.Y = (SdfGlyphAtlasData.EmOuterSpread-EmOutlineSize)/(SdfGlyphAtlasData.EmInnerSpread+SdfGlyphAtlasData.EmOuterSpread);
+					NextSdfBias = (SdfGlyphAtlasData.EmOuterSpread-EmOutlineSize)/(SdfGlyphAtlasData.EmInnerSpread+SdfGlyphAtlasData.EmOuterSpread);
 				}
 			}
 			else
@@ -3269,7 +3270,7 @@ void FSlateElementBatcher::BuildShapedTextSequence(const FShapedTextBuildContext
 				}
 
 				check(NextAtlasDataTextureIndex >= 0);
-				if (FontAtlasTexture == nullptr || NextAtlasDataTextureIndex != FontTextureIndex || (bIsSdfGlyph && (GlyphShaderParams.X != SdfShaderParams.W || GlyphShaderParams.Y != SdfShaderParams.Z)))
+				if (FontAtlasTexture == nullptr || NextAtlasDataTextureIndex != FontTextureIndex || (bIsSdfGlyph && (NextSdfPixelSpread != SdfPixelSpread || NextSdfBias != SdfBias)))
 				{
 					// Font has a new texture for this glyph or shader parameters changed. Refresh the batch we use and the index we are currently using
 					FontTextureIndex = NextAtlasDataTextureIndex;
@@ -3286,14 +3287,6 @@ void FSlateElementBatcher::BuildShapedTextSequence(const FShapedTextBuildContext
 					InvTextureSizeX = 1.0f / FontAtlasTexture->GetWidth();
 					InvTextureSizeY = 1.0f / FontAtlasTexture->GetHeight();
 
-					if (bIsSdfGlyph)
-					{
-						SdfShaderParams.X = InvTextureSizeX*GlyphShaderParams.X;
-						SdfShaderParams.Y = InvTextureSizeX*GlyphShaderParams.X;
-						SdfShaderParams.Z = GlyphShaderParams.Y;
-						SdfShaderParams.W = GlyphShaderParams.X;
-					}
-
 					const ESlateFontAtlasContentType ContentType = SlateFontTexture->GetContentType();
 					Tint = ContentType == ESlateFontAtlasContentType::Color ? FColor::White : Context.FontTint;
 					check(bIsSdfGlyph == (ContentType == ESlateFontAtlasContentType::Msdf));
@@ -3318,10 +3311,25 @@ void FSlateElementBatcher::BuildShapedTextSequence(const FShapedTextBuildContext
 					}
 					check(ShaderType != ESlateShader::Default);
 
+					FShaderParams ShaderParams;
+					if (bIsSdfGlyph)
+					{
+						SdfPixelSpread = NextSdfPixelSpread;
+						SdfBias = NextSdfBias;
+						// Note - it would be much better to pass the SDF shader params as per-vertex attributes instead to avoid having to switch batches too often
+						ShaderParams = FShaderParams::MakePixelShaderParams(FVector4f(
+							// Half of horizontal, vertical spread in texture coordinate units
+							.5f*InvTextureSizeX*SdfPixelSpread,
+							.5f*InvTextureSizeY*SdfPixelSpread,
+							// Signed distance sample bias, the color value (between 0 to 1) representing zero distance
+							SdfBias,
+							// The last parameter needs to be 0 for alpha texture single-channel SDF, 1 for BGRA/RGBA MTSDF
+							1.f
+						));
+					}
+
 					RenderBatch = &CreateRenderBatch(Context.LayerId,
-						bIsSdfGlyph
-						? FShaderParams::MakePixelShaderParams(SdfShaderParams)
-						: FShaderParams(),
+						ShaderParams,
 						FontShaderResource,
 						ESlateDrawPrimitive::TriangleList,
 						ShaderType,
diff --git a/Engine/Source/Runtime/SlateRHIRenderer/Private/SlateRHIFontTexture.cpp b/Engine/Source/Runtime/SlateRHIRenderer/Private/SlateRHIFontTexture.cpp
--- a/Engine/Source/Runtime/SlateRHIRenderer/Private/SlateRHIFontTexture.cpp
+++ b/Engine/Source/Runtime/SlateRHIRenderer/Private/SlateRHIFontTexture.cpp
@@ -103,9 +103,8 @@ EPixelFormat FSlateFontTextureRHIResource::GetRHIPixelFormat() const
 		case ESlateFontAtlasContentType::Alpha:
 			return PF_A8;
 		case ESlateFontAtlasContentType::Color:
-			return PF_B8G8R8A8;
 		case ESlateFontAtlasContentType::Msdf:
-			return PF_R8G8B8A8;
+			return PF_B8G8R8A8;
 		default:
 			checkNoEntry();
 			// Default to Color
