Slate: Add unittest for TSlateSlotAttribute. Fix visibility sorting. The sort was not activated if there was no prerequiste.

diff --git a/Engine/Source/Runtime/SlateCore/Private/Test/SlateAttributeTest.cpp b/Engine/Source/Runtime/SlateCore/Private/Test/SlateAttributeTest.cpp
--- a/Engine/Source/Runtime/SlateCore/Private/Test/SlateAttributeTest.cpp
+++ b/Engine/Source/Runtime/SlateCore/Private/Test/SlateAttributeTest.cpp
@@ -17,11 +17,7 @@
 //IMPLEMENT_SIMPLE_AUTOMATION_TEST(FSlateAttributeTest, "Slate.Attribute", EAutomationTestFlags::ApplicationContextMask | EAutomationTestFlags::SmokeFilter)
 IMPLEMENT_SIMPLE_AUTOMATION_TEST(FSlateAttributeTest, "Slate.Attribute", EAutomationTestFlags::ApplicationContextMask | EAutomationTestFlags::EngineFilter)
 
-namespace UE
-{
-namespace Slate
-{
-namespace Private
+namespace UE::Slate::Private
 {
 
 struct FConstructionCounter
@@ -66,6 +62,9 @@ FVector2D CallbackForFVectorAttribute()
 	return FVector2D(1,1);
 }
 
+/**
+ * 
+ */
 class SAttributeLeftWidget_Parent : public SLeafWidget
 {
 	SLATE_DECLARE_WIDGET(SAttributeLeftWidget_Parent, SLeafWidget)
@@ -118,6 +117,9 @@ void SAttributeLeftWidget_Parent::PrivateRegisterAttributes(FSlateAttributeIniti
 }
 
 
+/**
+ * 
+ */
 class SAttributeLeftWidget_Child : public SAttributeLeftWidget_Parent
 {
 	SLATE_DECLARE_WIDGET(SAttributeLeftWidget_Child, SAttributeLeftWidget_Parent)
@@ -163,6 +165,9 @@ void SAttributeLeftWidget_Child::PrivateRegisterAttributes(FSlateAttributeInitia
 		.AffectVisibility();
 }
 
+/**
+ * 
+ */
 class SAttributeLeftWidget_OnInvalidationParent : public SLeafWidget
 {
 	SLATE_DECLARE_WIDGET(SAttributeLeftWidget_OnInvalidationParent, SLeafWidget)
@@ -234,6 +239,9 @@ void SAttributeLeftWidget_OnInvalidationParent::PrivateRegisterAttributes(FSlate
 			}));
 }
 
+/*
+ * 
+ */
 class SAttributeLeftWidget_OnInvalidationChild : public SAttributeLeftWidget_OnInvalidationParent
 {
 	SLATE_DECLARE_WIDGET(SAttributeLeftWidget_OnInvalidationChild, SAttributeLeftWidget_OnInvalidationParent)
@@ -287,10 +295,122 @@ void SAttributeLeftWidget_OnInvalidationChild::PrivateRegisterAttributes(FSlateA
 			}));
 }
 
-}
-}
+/*
+ *
+ */
+class SContainerAttribute : public SLeafWidget
+{
+	SLATE_DECLARE_WIDGET(SContainerAttribute, SLeafWidget)
+public:
+	class FSlot : public TWidgetSlotWithAttributeSupport<FSlot>, public TPaddingWidgetSlotMixin<FSlot>
+	{
+	public:
+		FSlot()
+			: TWidgetSlotWithAttributeSupport<FSlot>()
+			, TPaddingWidgetSlotMixin<FSlot>(FMargin(99.f))
+			, SlotAttribute1(*this, 99)
+			, SlotAttribute2(*this, 99)
+			, SlotAttribute3(*this, 99)
+			, SlotId(INDEX_NONE)
+		{ }
+
+		SLATE_SLOT_BEGIN_ARGS_OneMixin(FSlot, TSlotBase<FSlot>, TPaddingWidgetSlotMixin<FSlot>)
+			SLATE_ATTRIBUTE(int32, SlotAttribute1)
+			SLATE_ATTRIBUTE(int32, SlotAttribute2)
+			SLATE_ATTRIBUTE(int32, SlotAttribute3)
+			SLATE_ARGUMENT(int32, SlotId)
+		SLATE_SLOT_END_ARGS()
+
+		void Construct(const FChildren& SlotOwner, FSlotArguments&& InArgs)
+		{
+			TSlotBase<FSlot>::Construct(SlotOwner, MoveTemp(InArgs));
+			TPaddingWidgetSlotMixin<FSlot>::ConstructMixin(SlotOwner, MoveTemp(InArgs));
+			if (InArgs._SlotAttribute1.IsSet())
+			{
+				SlotAttribute1.Assign(*this, MoveTemp(InArgs._SlotAttribute1));
+			}
+			if (InArgs._SlotAttribute2.IsSet())
+			{
+				SlotAttribute2.Assign(*this, MoveTemp(InArgs._SlotAttribute2));
+			}
+			if (InArgs._SlotAttribute3.IsSet())
+			{
+				SlotAttribute3.Assign(*this, MoveTemp(InArgs._SlotAttribute3));
+			}
+			SlotId = InArgs._SlotId;
+		}
+
+		static void RegisterAttributes(FSlateWidgetSlotAttributeInitializer& AttributeInitializer)
+		{
+			TPaddingWidgetSlotMixin<FSlot>::RegisterAttributesMixin(AttributeInitializer);
+			SLATE_ADD_SLOT_ATTRIBUTE_DEFINITION_WITH_NAME(FSlot, AttributeInitializer, "Slot.Attribute1", SlotAttribute1, EInvalidateWidgetReason::Paint);
+			SLATE_ADD_SLOT_ATTRIBUTE_DEFINITION_WITH_NAME(FSlot, AttributeInitializer, "Slot.Attribute2", SlotAttribute2, EInvalidateWidgetReason::Layout);
+			SLATE_ADD_SLOT_ATTRIBUTE_DEFINITION_WITH_NAME(FSlot, AttributeInitializer, "Slot.Attribute3", SlotAttribute3, EInvalidateWidgetReason::Prepass);
+		}
+
+		TSlateSlotAttribute<int32> SlotAttribute1;
+		TSlateSlotAttribute<int32> SlotAttribute2;
+		TSlateSlotAttribute<int32> SlotAttribute3;
+		int32 SlotId;
+	};
+
+	SLATE_BEGIN_ARGS(SContainerAttribute) {}
+	SLATE_END_ARGS()
+
+	SContainerAttribute()
+		: Slots(this, GET_MEMBER_NAME_CHECKED(SContainerAttribute, Slots))
+		, AttributeA(*this, 99)
+		, AttributeB(*this, 99)
+		, AttributeC(*this, 99)
+	{}
+
+
+	FSlot::FSlotArguments Slot()
+	{
+		return FSlot::FSlotArguments(MakeUnique<FSlot>());
+	}
+
+	TPanelChildren<FSlot>::FScopedWidgetSlotArguments AddSlot()
+	{
+		return TPanelChildren<FSlot>::FScopedWidgetSlotArguments{ MakeUnique<FSlot>(), Slots, INDEX_NONE };
+	}
+
+	bool RemoveSlot(const TSharedRef<SWidget>& SlotWidget)
+	{
+		return Slots.Remove(SlotWidget) != INDEX_NONE;
+	}
+
+	void ClearChildren()
+	{
+		Slots.Empty();
+	}
+
+	virtual int32 OnPaint(const FPaintArgs& Args, const FGeometry& AllottedGeometry, const FSlateRect& MyCullingRect, FSlateWindowElementList& OutDrawElements, int32 LayerId, const FWidgetStyle& InWidgetStyle, bool bParentEnabled) const override { return LayerId; }
+	virtual FVector2D ComputeDesiredSize(float) const override { return FVector2D(0.f, 0.f); }
+	void Construct(const FArguments&) {}
+
+public:
+	/** The slots that are placed into various grid locations */
+	TPanelChildren<FSlot> Slots;
+	TSlateAttribute<int32> AttributeA;
+	TSlateAttribute<int32> AttributeB;
+	TSlateAttribute<int32> AttributeC;
+};
+
+SLATE_IMPLEMENT_WIDGET(SContainerAttribute)
+void SContainerAttribute::PrivateRegisterAttributes(FSlateAttributeInitializer& AttributeInitializer)
+{
+	FSlateWidgetSlotAttributeInitializer SlotInitializer = SLATE_ADD_PANELCHILDREN_DEFINITION(AttributeInitializer, Slots);
+	FSlot::RegisterAttributes(SlotInitializer);
+
+	SLATE_ADD_MEMBER_ATTRIBUTE_DEFINITION(AttributeInitializer, AttributeA, EInvalidateWidgetReason::RenderTransform);
+	SLATE_ADD_MEMBER_ATTRIBUTE_DEFINITION(AttributeInitializer, AttributeB, EInvalidateWidgetReason::Visibility)
+		.AffectVisibility();
+	SLATE_ADD_MEMBER_ATTRIBUTE_DEFINITION(AttributeInitializer, AttributeC, EInvalidateWidgetReason::Volatility);
 }
 
+} //namespace
+
 bool FSlateAttributeTest::RunTest(const FString& Parameters)
 {
 	const int NumberOfAttributeInSWidget = SWidget::StaticWidgetClass().GetAttributeDescriptor().GetAttributeNum();
@@ -469,15 +589,15 @@ bool FSlateAttributeTest::RunTest(const FString& Parameters)
 	{
 		{
 			// This should just compile to TSlateAttribute
-			struct SAttributeAttribute : public SLeafWidget
+			struct SSlateAttribute : public SLeafWidget
 			{
-				SAttributeAttribute()
+				SSlateAttribute()
 					: Toto(6)
 					, AttributeA(*this)
 					, AttributeB(*this, 5)
 					, AttributeC(*this, MoveTemp(Toto))
 				{ }
-				SLATE_BEGIN_ARGS(SAttributeAttribute) {}  SLATE_END_ARGS()
+				SLATE_BEGIN_ARGS(SSlateAttribute) {}  SLATE_END_ARGS()
 				virtual int32 OnPaint(const FPaintArgs& Args, const FGeometry& AllottedGeometry, const FSlateRect& MyCullingRect, FSlateWindowElementList& OutDrawElements, int32 LayerId, const FWidgetStyle& InWidgetStyle, bool bParentEnabled) const override { return LayerId; }
 				virtual FVector2D ComputeDesiredSize(float) const override { return FVector2D(0.f, 0.f); }
 				void Construct(const FArguments&) {	}
@@ -489,7 +609,7 @@ bool FSlateAttributeTest::RunTest(const FString& Parameters)
 				TSlateAttribute<int32, EInvalidateWidgetReason::Paint> AttributeB;
 				TSlateAttribute<int32, EInvalidateWidgetReason::Paint> AttributeC;
 			};
-			TSharedPtr<SAttributeAttribute> Widget = SNew(SAttributeAttribute);
+			TSharedPtr<SSlateAttribute> Widget = SNew(SSlateAttribute);
 
 			{
 				int32 Hello = 7;
@@ -502,7 +622,7 @@ bool FSlateAttributeTest::RunTest(const FString& Parameters)
 				auto Getter1 = TAttribute<int32>::FGetter::CreateStatic(UE::Slate::Private::CallbackForIntAttribute, 1);
 				Widget->AttributeA.Bind(*Widget, Getter1);
 				Widget->AttributeA.Bind(*Widget, MoveTemp(Getter1));
-				Widget->AttributeA.Bind(*Widget, &SAttributeAttribute::Callback);
+				Widget->AttributeA.Bind(*Widget, &SSlateAttribute::Callback);
 			}
 			{
 				int32 TmpInt1 = 7;
@@ -530,9 +650,9 @@ bool FSlateAttributeTest::RunTest(const FString& Parameters)
 			typedef UE::Slate::Private::FConstructionCounter FLocalConstructionCounter;
 
 			// This should just compile to TSlateManagedAttribute
-			struct SAttributeAttribute : public SLeafWidget
+			struct SManagedAttribute : public SLeafWidget
 			{
-				SLATE_BEGIN_ARGS(SAttributeAttribute) {} SLATE_END_ARGS()
+				SLATE_BEGIN_ARGS(SManagedAttribute) {} SLATE_END_ARGS()
 				virtual int32 OnPaint(const FPaintArgs& Args, const FGeometry& AllottedGeometry, const FSlateRect& MyCullingRect, FSlateWindowElementList& OutDrawElements, int32 LayerId, const FWidgetStyle& InWidgetStyle, bool bParentEnabled) const override { return LayerId; }
 				virtual FVector2D ComputeDesiredSize(float) const override { return FVector2D(0.f, 0.f); }
 				void Construct(const FArguments&) {}
@@ -541,7 +661,7 @@ bool FSlateAttributeTest::RunTest(const FString& Parameters)
 
 				using ManagedSlateAttributeType = TSlateManagedAttribute<FLocalConstructionCounter, EInvalidateWidgetReason::Layout>;
 			};
-			TSharedPtr<SAttributeAttribute>	Widget = SNew(SAttributeAttribute);
+			TSharedPtr<SManagedAttribute>	Widget = SNew(SManagedAttribute);
 
 
 			auto AddErrorIfCounterDoNotMatches = [this](int32 Construct, int32 Copy, int32 Move, int32 CopyAssign, int32 MoveAssign, const TCHAR* Message)
@@ -556,93 +676,93 @@ bool FSlateAttributeTest::RunTest(const FString& Parameters)
 
 			{
 				FLocalConstructionCounter::ResetCounter();
-				SAttributeAttribute::ManagedSlateAttributeType Attribute = SAttributeAttribute::ManagedSlateAttributeType(Widget.ToSharedRef() );
+				SManagedAttribute::ManagedSlateAttributeType Attribute = SManagedAttribute::ManagedSlateAttributeType(Widget.ToSharedRef() );
 				AddErrorIfCounterDoNotMatches(1, 0, 0, 0, 0, TEXT("Default & Copy constructor was not used."));
 			}
 			{
 				FLocalConstructionCounter Counter = 1;
 				FLocalConstructionCounter::ResetCounter();
-				SAttributeAttribute::ManagedSlateAttributeType Attribute = SAttributeAttribute::ManagedSlateAttributeType(Widget.ToSharedRef(), Counter);
+				SManagedAttribute::ManagedSlateAttributeType Attribute = SManagedAttribute::ManagedSlateAttributeType(Widget.ToSharedRef(), Counter);
 				AddErrorIfCounterDoNotMatches(0, 1, 0, 0, 0, TEXT("Default & Copy constructor was not used."));
 			}
 			{
 				FLocalConstructionCounter::ResetCounter();
 				FLocalConstructionCounter Counter = 1;
-				SAttributeAttribute::ManagedSlateAttributeType Attribute = SAttributeAttribute::ManagedSlateAttributeType(Widget.ToSharedRef(), MoveTemp(Counter));
+				SManagedAttribute::ManagedSlateAttributeType Attribute = SManagedAttribute::ManagedSlateAttributeType(Widget.ToSharedRef(), MoveTemp(Counter));
 				AddErrorIfCounterDoNotMatches(1, 0, 1, 0, 0, TEXT("Default & Move constructor was not used."));
 			}
 			{
 				TAttribute<FLocalConstructionCounter>::FGetter Getter1 = TAttribute<FLocalConstructionCounter>::FGetter::CreateLambda([](){return FLocalConstructionCounter(1);});
 				FLocalConstructionCounter Counter = 1;
 				FLocalConstructionCounter::ResetCounter();
-				SAttributeAttribute::ManagedSlateAttributeType Attribute = SAttributeAttribute::ManagedSlateAttributeType{Widget.ToSharedRef(), Getter1, Counter};
+				SManagedAttribute::ManagedSlateAttributeType Attribute = SManagedAttribute::ManagedSlateAttributeType{Widget.ToSharedRef(), Getter1, Counter};
 				AddErrorIfCounterDoNotMatches(0, 1, 0, 0, 0, TEXT("Getter & Copy constructor was not used."));
 			}
 			{
 				TAttribute<FLocalConstructionCounter>::FGetter Getter1 = TAttribute<FLocalConstructionCounter>::FGetter::CreateLambda([]() {return FLocalConstructionCounter(1);});
 				FLocalConstructionCounter Counter = 1;
 				FLocalConstructionCounter::ResetCounter();
-				SAttributeAttribute::ManagedSlateAttributeType Attribute = SAttributeAttribute::ManagedSlateAttributeType{Widget.ToSharedRef(), Getter1, MoveTemp(Counter)};
+				SManagedAttribute::ManagedSlateAttributeType Attribute = SManagedAttribute::ManagedSlateAttributeType{Widget.ToSharedRef(), Getter1, MoveTemp(Counter)};
 				AddErrorIfCounterDoNotMatches(0, 0, 1, 0, 0, TEXT("Getter & Move constructor was not used."));
 			}
 			{
 				TAttribute<FLocalConstructionCounter>::FGetter Getter1 = TAttribute<FLocalConstructionCounter>::FGetter::CreateLambda([]() {return FLocalConstructionCounter(1);});
 				FLocalConstructionCounter Counter = 1;
 				FLocalConstructionCounter::ResetCounter();
-				SAttributeAttribute::ManagedSlateAttributeType Attribute = SAttributeAttribute::ManagedSlateAttributeType(Widget.ToSharedRef(), MoveTemp(Getter1), Counter);
+				SManagedAttribute::ManagedSlateAttributeType Attribute = SManagedAttribute::ManagedSlateAttributeType(Widget.ToSharedRef(), MoveTemp(Getter1), Counter);
 				AddErrorIfCounterDoNotMatches(0, 1, 0, 0, 0, TEXT("Move Getter & Copy constructor was not used."));
 			}
 			{
 				TAttribute<FLocalConstructionCounter>::FGetter Getter1 = TAttribute<FLocalConstructionCounter>::FGetter::CreateLambda([]() {return FLocalConstructionCounter(1);});
 				FLocalConstructionCounter Counter = 1;
 				FLocalConstructionCounter::ResetCounter();
-				SAttributeAttribute::ManagedSlateAttributeType Attribute = SAttributeAttribute::ManagedSlateAttributeType(Widget.ToSharedRef(), MoveTemp(Getter1), MoveTemp(Counter));
+				SManagedAttribute::ManagedSlateAttributeType Attribute = SManagedAttribute::ManagedSlateAttributeType(Widget.ToSharedRef(), MoveTemp(Getter1), MoveTemp(Counter));
 				AddErrorIfCounterDoNotMatches(0, 0, 1, 0, 0, TEXT("Move Getter & Move constructor was not used."));
 			}
 			{
 				TAttribute<FLocalConstructionCounter> Attribute1 = MakeAttributeLambda([]() {return FLocalConstructionCounter(1);});
 				FLocalConstructionCounter Counter;
 				FLocalConstructionCounter::ResetCounter();
-				SAttributeAttribute::ManagedSlateAttributeType Attribute = SAttributeAttribute::ManagedSlateAttributeType(Widget.ToSharedRef(), Attribute1, Counter);
+				SManagedAttribute::ManagedSlateAttributeType Attribute = SManagedAttribute::ManagedSlateAttributeType(Widget.ToSharedRef(), Attribute1, Counter);
 				AddErrorIfCounterDoNotMatches(0, 1, 0, 0, 0, TEXT("Attribute & Copy constructor was not used."));
 			}
 			{
 				TAttribute<FLocalConstructionCounter> Attribute1 = MakeAttributeLambda([]() {return FLocalConstructionCounter(1);});
 				FLocalConstructionCounter Counter = 1;
 				FLocalConstructionCounter::ResetCounter();
-				SAttributeAttribute::ManagedSlateAttributeType Attribute = SAttributeAttribute::ManagedSlateAttributeType(Widget.ToSharedRef(), MoveTemp(Attribute1), MoveTemp(Counter));
+				SManagedAttribute::ManagedSlateAttributeType Attribute = SManagedAttribute::ManagedSlateAttributeType(Widget.ToSharedRef(), MoveTemp(Attribute1), MoveTemp(Counter));
 				AddErrorIfCounterDoNotMatches(0, 0, 1, 0, 0, TEXT("Move Attribute & Move constructor was not used."));
 			}
 			{
-				SAttributeAttribute::ManagedSlateAttributeType Attribute = SAttributeAttribute::ManagedSlateAttributeType(Widget.ToSharedRef());
+				SManagedAttribute::ManagedSlateAttributeType Attribute = SManagedAttribute::ManagedSlateAttributeType(Widget.ToSharedRef());
 				FLocalConstructionCounter::ResetCounter();
 				FLocalConstructionCounter Result = Attribute.Get();
 				Attribute.UpdateNow();
 				AddErrorIfCounterDoNotMatches(0, 1, 0, 0, 0, TEXT("Get and UpdateNow failed."));
 			}
 			{
-				SAttributeAttribute::ManagedSlateAttributeType Attribute = SAttributeAttribute::ManagedSlateAttributeType(Widget.ToSharedRef());
+				SManagedAttribute::ManagedSlateAttributeType Attribute = SManagedAttribute::ManagedSlateAttributeType(Widget.ToSharedRef());
 				FLocalConstructionCounter Counter = 1;
 				FLocalConstructionCounter::ResetCounter();
 				Attribute.Set(Counter);
 				AddErrorIfCounterDoNotMatches(0, 0, 0, 1, 0, TEXT("Set Copy failed."));
 			}
 			{
-				SAttributeAttribute::ManagedSlateAttributeType Attribute = SAttributeAttribute::ManagedSlateAttributeType(Widget.ToSharedRef());
+				SManagedAttribute::ManagedSlateAttributeType Attribute = SManagedAttribute::ManagedSlateAttributeType(Widget.ToSharedRef());
 				FLocalConstructionCounter Counter = 1;
 				FLocalConstructionCounter::ResetCounter();
 				Attribute.Set(MoveTemp(Counter));
 				AddErrorIfCounterDoNotMatches(0, 0, 0, 0, 1, TEXT("Set Move failed."));
 			}
 			{
-				SAttributeAttribute::ManagedSlateAttributeType Attribute = SAttributeAttribute::ManagedSlateAttributeType(Widget.ToSharedRef());
+				SManagedAttribute::ManagedSlateAttributeType Attribute = SManagedAttribute::ManagedSlateAttributeType(Widget.ToSharedRef());
 				auto Getter1 = TAttribute<FLocalConstructionCounter>::FGetter::CreateLambda([]() {return FLocalConstructionCounter(1);});
 				FLocalConstructionCounter::ResetCounter();
 				Attribute.Bind(Getter1);
 				AddErrorIfCounterDoNotMatches(0, 0, 0, 0, 0, TEXT("Bind Copy failed."));
 			}
 			{
-				SAttributeAttribute::ManagedSlateAttributeType Attribute = SAttributeAttribute::ManagedSlateAttributeType(Widget.ToSharedRef());
+				SManagedAttribute::ManagedSlateAttributeType Attribute = SManagedAttribute::ManagedSlateAttributeType(Widget.ToSharedRef());
 				auto Getter1 = TAttribute<FLocalConstructionCounter>::FGetter::CreateLambda([]() {return FLocalConstructionCounter(1);});
 				FLocalConstructionCounter::ResetCounter();
 				Attribute.Bind(MoveTemp(Getter1));
@@ -650,7 +770,7 @@ bool FSlateAttributeTest::RunTest(const FString& Parameters)
 			}
 			// Test Assign
 			{
-				SAttributeAttribute::ManagedSlateAttributeType Attribute = SAttributeAttribute::ManagedSlateAttributeType(Widget.ToSharedRef());
+				SManagedAttribute::ManagedSlateAttributeType Attribute = SManagedAttribute::ManagedSlateAttributeType(Widget.ToSharedRef());
 				TAttribute<FLocalConstructionCounter> Attribute1;
 				FLocalConstructionCounter::ResetCounter();
 				Attribute.Assign(Attribute1);
@@ -670,7 +790,7 @@ bool FSlateAttributeTest::RunTest(const FString& Parameters)
 			}
 			// Test unbinded Attribute
 			{
-				SAttributeAttribute::ManagedSlateAttributeType Attribute = SAttributeAttribute::ManagedSlateAttributeType(Widget.ToSharedRef());
+				SManagedAttribute::ManagedSlateAttributeType Attribute = SManagedAttribute::ManagedSlateAttributeType(Widget.ToSharedRef());
 				TAttribute<FLocalConstructionCounter> Attribute1;
 				FLocalConstructionCounter Counter = 1;
 				FLocalConstructionCounter::ResetCounter();
@@ -678,7 +798,7 @@ bool FSlateAttributeTest::RunTest(const FString& Parameters)
 				AddErrorIfCounterDoNotMatches(0, 0, 0, 1, 0, TEXT("Assign Copy/Copy failed."));
 			}
 			{
-				SAttributeAttribute::ManagedSlateAttributeType Attribute = SAttributeAttribute::ManagedSlateAttributeType(Widget.ToSharedRef());
+				SManagedAttribute::ManagedSlateAttributeType Attribute = SManagedAttribute::ManagedSlateAttributeType(Widget.ToSharedRef());
 				TAttribute<FLocalConstructionCounter> Attribute1;
 				FLocalConstructionCounter Counter = 1;
 				FLocalConstructionCounter::ResetCounter();
@@ -686,7 +806,7 @@ bool FSlateAttributeTest::RunTest(const FString& Parameters)
 				AddErrorIfCounterDoNotMatches(0, 0, 0, 0, 1, TEXT("Assign Copy/Move failed."));
 			}
 			{
-				SAttributeAttribute::ManagedSlateAttributeType Attribute = SAttributeAttribute::ManagedSlateAttributeType(Widget.ToSharedRef());
+				SManagedAttribute::ManagedSlateAttributeType Attribute = SManagedAttribute::ManagedSlateAttributeType(Widget.ToSharedRef());
 				TAttribute<FLocalConstructionCounter> Attribute1;
 				FLocalConstructionCounter Counter = 1;
 				FLocalConstructionCounter::ResetCounter();
@@ -694,7 +814,7 @@ bool FSlateAttributeTest::RunTest(const FString& Parameters)
 				AddErrorIfCounterDoNotMatches(0, 0, 0, 1, 0, TEXT("Assign Move/CCopy failed."));
 			}
 			{
-				SAttributeAttribute::ManagedSlateAttributeType Attribute = SAttributeAttribute::ManagedSlateAttributeType(Widget.ToSharedRef());
+				SManagedAttribute::ManagedSlateAttributeType Attribute = SManagedAttribute::ManagedSlateAttributeType(Widget.ToSharedRef());
 				TAttribute<FLocalConstructionCounter> Attribute1;
 				FLocalConstructionCounter Counter = 1;
 				FLocalConstructionCounter::ResetCounter();
@@ -703,7 +823,7 @@ bool FSlateAttributeTest::RunTest(const FString& Parameters)
 			}
 			// Test binded Attribute
 			{
-				SAttributeAttribute::ManagedSlateAttributeType Attribute = SAttributeAttribute::ManagedSlateAttributeType(Widget.ToSharedRef());
+				SManagedAttribute::ManagedSlateAttributeType Attribute = SManagedAttribute::ManagedSlateAttributeType(Widget.ToSharedRef());
 				TAttribute<FLocalConstructionCounter> Attribute1 = MakeAttributeLambda([]() {return FLocalConstructionCounter(1);});
 				FLocalConstructionCounter Counter = 1;
 				FLocalConstructionCounter::ResetCounter();
@@ -711,7 +831,7 @@ bool FSlateAttributeTest::RunTest(const FString& Parameters)
 				AddErrorIfCounterDoNotMatches(0, 0, 0, 0, 0, TEXT("Bind Copy with binded attribute failed."));
 			}
 			{
-				SAttributeAttribute::ManagedSlateAttributeType Attribute = SAttributeAttribute::ManagedSlateAttributeType(Widget.ToSharedRef());
+				SManagedAttribute::ManagedSlateAttributeType Attribute = SManagedAttribute::ManagedSlateAttributeType(Widget.ToSharedRef());
 				TAttribute<FLocalConstructionCounter> Attribute1 = MakeAttributeLambda([]() {return FLocalConstructionCounter(1);});
 				FLocalConstructionCounter Counter = 1;
 				FLocalConstructionCounter::ResetCounter();
@@ -719,7 +839,7 @@ bool FSlateAttributeTest::RunTest(const FString& Parameters)
 				AddErrorIfCounterDoNotMatches(0, 0, 0, 0, 0, TEXT("Bind Move with binded attribute failed."));
 			}
 			{
-				SAttributeAttribute::ManagedSlateAttributeType Attribute = SAttributeAttribute::ManagedSlateAttributeType(Widget.ToSharedRef());
+				SManagedAttribute::ManagedSlateAttributeType Attribute = SManagedAttribute::ManagedSlateAttributeType(Widget.ToSharedRef());
 				TAttribute<FLocalConstructionCounter> Attribute1 = MakeAttributeLambda([]() {return FLocalConstructionCounter(1);});
 				FLocalConstructionCounter Counter = 1;
 				FLocalConstructionCounter::ResetCounter();
@@ -727,7 +847,7 @@ bool FSlateAttributeTest::RunTest(const FString& Parameters)
 				AddErrorIfCounterDoNotMatches(0, 0, 0, 0, 0, TEXT("Assign Copy with binded attribute failed."));
 			}
 			{
-				SAttributeAttribute::ManagedSlateAttributeType Attribute = SAttributeAttribute::ManagedSlateAttributeType(Widget.ToSharedRef());
+				SManagedAttribute::ManagedSlateAttributeType Attribute = SManagedAttribute::ManagedSlateAttributeType(Widget.ToSharedRef());
 				TAttribute<FLocalConstructionCounter> Attribute1 = MakeAttributeLambda([]() {return FLocalConstructionCounter(1);});
 				FLocalConstructionCounter Counter = 1;
 				FLocalConstructionCounter::ResetCounter();
@@ -736,7 +856,7 @@ bool FSlateAttributeTest::RunTest(const FString& Parameters)
 			}
 			// Test set Attribute
 			{
-				SAttributeAttribute::ManagedSlateAttributeType Attribute = SAttributeAttribute::ManagedSlateAttributeType(Widget.ToSharedRef());
+				SManagedAttribute::ManagedSlateAttributeType Attribute = SManagedAttribute::ManagedSlateAttributeType(Widget.ToSharedRef());
 				FLocalConstructionCounter Counter = 1;
 				TAttribute<FLocalConstructionCounter> Attribute1 = Counter;
 				FLocalConstructionCounter::ResetCounter();
@@ -744,7 +864,7 @@ bool FSlateAttributeTest::RunTest(const FString& Parameters)
 				AddErrorIfCounterDoNotMatches(0, 0, 0, 1, 0, TEXT("Assign Set Copy failed."));
 			}
 			{
-				SAttributeAttribute::ManagedSlateAttributeType Attribute = SAttributeAttribute::ManagedSlateAttributeType(Widget.ToSharedRef());
+				SManagedAttribute::ManagedSlateAttributeType Attribute = SManagedAttribute::ManagedSlateAttributeType(Widget.ToSharedRef());
 				FLocalConstructionCounter Counter = 1;
 				TAttribute<FLocalConstructionCounter> Attribute1 = Counter;
 				FLocalConstructionCounter::ResetCounter();
@@ -752,7 +872,7 @@ bool FSlateAttributeTest::RunTest(const FString& Parameters)
 				AddErrorIfCounterDoNotMatches(0, 0, 0, 1, 0, TEXT("Assign Set Copy failed."));
 			}
 			{
-				SAttributeAttribute::ManagedSlateAttributeType Attribute = SAttributeAttribute::ManagedSlateAttributeType(Widget.ToSharedRef());
+				SManagedAttribute::ManagedSlateAttributeType Attribute = SManagedAttribute::ManagedSlateAttributeType(Widget.ToSharedRef());
 				FLocalConstructionCounter Counter = 1;
 				TAttribute<FLocalConstructionCounter> Attribute1 = Counter;
 				FLocalConstructionCounter::ResetCounter();
@@ -760,7 +880,7 @@ bool FSlateAttributeTest::RunTest(const FString& Parameters)
 				AddErrorIfCounterDoNotMatches(0, 0, 2, 0, 1, TEXT("Assign Set Move failed."));
 			}
 			{
-				SAttributeAttribute::ManagedSlateAttributeType Attribute = SAttributeAttribute::ManagedSlateAttributeType(Widget.ToSharedRef());
+				SManagedAttribute::ManagedSlateAttributeType Attribute = SManagedAttribute::ManagedSlateAttributeType(Widget.ToSharedRef());
 				FLocalConstructionCounter Counter = 1;
 				TAttribute<FLocalConstructionCounter> Attribute1 = Counter;
 				FLocalConstructionCounter::ResetCounter();
@@ -768,6 +888,101 @@ bool FSlateAttributeTest::RunTest(const FString& Parameters)
 				AddErrorIfCounterDoNotMatches(0, 0, 2, 0, 1, TEXT("Assign Set Move failed."));
 			}
 		}
+
+		// Slot Attribute test
+		{
+			TSharedRef<UE::Slate::Private::SContainerAttribute> Widget = SNew(UE::Slate::Private::SContainerAttribute);
+
+			{
+				Widget->AttributeA.Assign(Widget.Get(), MakeAttributeLambda(OrderLambda));
+				AddErrorIfFalse(Widget->AttributeA.Get() == 99, TEXT("A It is not the expected value."));
+				Widget->AttributeB.Assign(Widget.Get(), MakeAttributeLambda(OrderLambda));
+				AddErrorIfFalse(Widget->AttributeB.Get() == 99, TEXT("B It is not the expected value."));
+				Widget->AttributeC.Assign(Widget.Get(), MakeAttributeLambda(OrderLambda));
+				AddErrorIfFalse(Widget->AttributeC.Get() == 99, TEXT("C It is not the expected value."));
+
+				Widget->AddSlot()
+					.SlotAttribute1(MakeAttributeLambda(OrderLambda))
+					.SlotAttribute2(MakeAttributeLambda(OrderLambda))
+					.SlotAttribute3(MakeAttributeLambda(OrderLambda))
+					.SlotId(0);
+				Widget->AddSlot()
+					.SlotAttribute1(MakeAttributeLambda(OrderLambda))
+					.SlotAttribute2(MakeAttributeLambda(OrderLambda))
+					.SlotAttribute3(MakeAttributeLambda(OrderLambda))
+					.SlotId(1);
+				Widget->AddSlot()
+					.SlotAttribute1(MakeAttributeLambda(OrderLambda))
+					.SlotAttribute2(MakeAttributeLambda(OrderLambda))
+					.SlotAttribute3(MakeAttributeLambda(OrderLambda))
+					.SlotId(2);
+			}
+			{
+				OrderCounter = 0;
+				Widget->MarkPrepassAsDirty();
+				Widget->SlatePrepass(1.f);
+
+				AddErrorIfFalse(Widget->Slots[0].SlotAttribute1.Get() == 2, TEXT("Slot[0].1 It is not the expected value."));
+				AddErrorIfFalse(Widget->Slots[0].SlotAttribute2.Get() == 3, TEXT("Slot[0].2 It is not the expected value."));
+				AddErrorIfFalse(Widget->Slots[0].SlotAttribute3.Get() == 4, TEXT("Slot[0].3 It is not the expected value."));
+				AddErrorIfFalse(Widget->Slots[0].SlotId == 0, TEXT("Slot[0].SlotId It is not the expected value."));
+
+				AddErrorIfFalse(Widget->Slots[1].SlotAttribute1.Get() == 5, TEXT("Slot[1].1 It is not the expected value."));
+				AddErrorIfFalse(Widget->Slots[1].SlotAttribute2.Get() == 6, TEXT("Slot[1].2 It is not the expected value."));
+				AddErrorIfFalse(Widget->Slots[1].SlotAttribute3.Get() == 7, TEXT("Slot[1].3 It is not the expected value."));
+				AddErrorIfFalse(Widget->Slots[1].SlotId == 1, TEXT("Slot[1].SlotId It is not the expected value."));
+
+				AddErrorIfFalse(Widget->Slots[2].SlotAttribute1.Get() == 8, TEXT("Slot[2].1 It is not the expected value."));
+				AddErrorIfFalse(Widget->Slots[2].SlotAttribute2.Get() == 9, TEXT("Slot[2].2 It is not the expected value."));
+				AddErrorIfFalse(Widget->Slots[2].SlotAttribute3.Get() == 10, TEXT("Slot[2].3 It is not the expected value."));
+				AddErrorIfFalse(Widget->Slots[2].SlotId == 2, TEXT("Slot[2].SlotId It is not the expected value."));
+				
+				AddErrorIfFalse(Widget->AttributeA.Get() == 11, TEXT("A It is not the expected value."));
+				AddErrorIfFalse(Widget->AttributeB.Get() == 1, TEXT("B It is not the expected value."));
+				AddErrorIfFalse(Widget->AttributeC.Get() == 12, TEXT("C It is not the expected value."));
+			}
+			{
+				Widget->Slots.Move(0, 1);
+
+				AddErrorIfFalse(Widget->Slots[1].SlotAttribute1.Get() == 2, TEXT("Slot[0].1 It is not the expected value."));
+				AddErrorIfFalse(Widget->Slots[1].SlotAttribute2.Get() == 3, TEXT("Slot[0].2 It is not the expected value."));
+				AddErrorIfFalse(Widget->Slots[1].SlotAttribute3.Get() == 4, TEXT("Slot[0].3 It is not the expected value."));
+				AddErrorIfFalse(Widget->Slots[1].SlotId == 0, TEXT("Slot[1].SlotId It is not the expected value."));
+
+				AddErrorIfFalse(Widget->Slots[0].SlotAttribute1.Get() == 5, TEXT("Slot[1].1 It is not the expected value."));
+				AddErrorIfFalse(Widget->Slots[0].SlotAttribute2.Get() == 6, TEXT("Slot[1].2 It is not the expected value."));
+				AddErrorIfFalse(Widget->Slots[0].SlotAttribute3.Get() == 7, TEXT("Slot[1].3 It is not the expected value."));
+				AddErrorIfFalse(Widget->Slots[0].SlotId == 1, TEXT("Slot[0].SlotId It is not the expected value."));
+
+				AddErrorIfFalse(Widget->Slots[2].SlotAttribute1.Get() == 8, TEXT("Slot[2].1 It is not the expected value."));
+				AddErrorIfFalse(Widget->Slots[2].SlotAttribute2.Get() == 9, TEXT("Slot[2].2 It is not the expected value."));
+				AddErrorIfFalse(Widget->Slots[2].SlotAttribute3.Get() == 10, TEXT("Slot[2].3 It is not the expected value."));
+				AddErrorIfFalse(Widget->Slots[2].SlotId == 2, TEXT("Slot[2].SlotId It is not the expected value."));
+
+				OrderCounter = 0;
+				Widget->MarkPrepassAsDirty();
+				Widget->SlatePrepass(1.f);
+
+				AddErrorIfFalse(Widget->Slots[0].SlotAttribute1.Get() == 2, TEXT("Slot[0].1 It is not the expected value."));
+				AddErrorIfFalse(Widget->Slots[0].SlotAttribute2.Get() == 3, TEXT("Slot[0].2 It is not the expected value."));
+				AddErrorIfFalse(Widget->Slots[0].SlotAttribute3.Get() == 4, TEXT("Slot[0].3 It is not the expected value."));
+				AddErrorIfFalse(Widget->Slots[0].SlotId == 1, TEXT("Slot[0].SlotId It is not the expected value."));
+
+				AddErrorIfFalse(Widget->Slots[1].SlotAttribute1.Get() == 5, TEXT("Slot[1].1 It is not the expected value."));
+				AddErrorIfFalse(Widget->Slots[1].SlotAttribute2.Get() == 6, TEXT("Slot[1].2 It is not the expected value."));
+				AddErrorIfFalse(Widget->Slots[1].SlotAttribute3.Get() == 7, TEXT("Slot[1].3 It is not the expected value."));
+				AddErrorIfFalse(Widget->Slots[1].SlotId == 0, TEXT("Slot[1].SlotId It is not the expected value."));
+
+				AddErrorIfFalse(Widget->Slots[2].SlotAttribute1.Get() == 8, TEXT("Slot[2].1 It is not the expected value."));
+				AddErrorIfFalse(Widget->Slots[2].SlotAttribute2.Get() == 9, TEXT("Slot[2].2 It is not the expected value."));
+				AddErrorIfFalse(Widget->Slots[2].SlotAttribute3.Get() == 10, TEXT("Slot[2].3 It is not the expected value."));
+				AddErrorIfFalse(Widget->Slots[2].SlotId == 2, TEXT("Slot[2].SlotId It is not the expected value."));
+
+				AddErrorIfFalse(Widget->AttributeA.Get() == 11, TEXT("A It is not the expected value."));
+				AddErrorIfFalse(Widget->AttributeB.Get() == 1, TEXT("B It is not the expected value."));
+				AddErrorIfFalse(Widget->AttributeC.Get() == 12, TEXT("C It is not the expected value."));
+			}
+		}
 	}
 
 	// Test OnInvalidation
diff --git a/Engine/Source/Runtime/SlateCore/Private/Types/SlateAttributeDescriptor.cpp b/Engine/Source/Runtime/SlateCore/Private/Types/SlateAttributeDescriptor.cpp
--- a/Engine/Source/Runtime/SlateCore/Private/Types/SlateAttributeDescriptor.cpp
+++ b/Engine/Source/Runtime/SlateCore/Private/Types/SlateAttributeDescriptor.cpp
@@ -176,7 +176,9 @@ FSlateAttributeDescriptor::FInitializer::~FInitializer()
 {
 	// Confirm that the Visibility attribute is marked as "bAffectVisibility"
 	{
-		const FAttribute* FoundVisibilityAttribute = Descriptor.FindAttribute("Visibility");
+		static const FName NAME_Visibility = "Visibility";
+		const FAttribute* FoundVisibilityAttribute = Descriptor.FindAttribute(NAME_Visibility);
+		checkf(&Descriptor.GetAttributeAtIndex(0) == FoundVisibilityAttribute, TEXT(""));
 		checkf(FoundVisibilityAttribute, TEXT("The visibility attribute doesn't exist."));
 		checkf(FoundVisibilityAttribute->bAffectVisibility, TEXT("The Visibility attribute must be marked as 'Affect Visibility'"));
 	}
@@ -265,7 +267,7 @@ FSlateAttributeDescriptor::FInitializer::~FInitializer()
 	TArray<FPrerequisiteSort, TInlineAllocator<32>> Prerequisites;
 	Prerequisites.Reserve(Descriptor.Attributes.Num());
 
-	bool bHavePrerequisite = false;
+	bool bSortNeeded = false;
 	for (int32 Index = 0; Index < Descriptor.Attributes.Num(); ++Index)
 	{
 		FAttribute& Attribute = Descriptor.Attributes[Index];
@@ -291,7 +293,7 @@ FSlateAttributeDescriptor::FInitializer::~FInitializer()
 			if (ensureAlwaysMsgf(Descriptor.Attributes.IsValidIndex(PrerequisiteIndex), TEXT("The Prerequisite '%s' doesn't exist"), *Prerequisite.ToString()))
 			{
 				Prerequisites.Emplace(Index, PrerequisiteIndex, -1);
-				bHavePrerequisite = true;
+				bSortNeeded = true;
 			}
 			else
 			{
@@ -300,11 +302,18 @@ FSlateAttributeDescriptor::FInitializer::~FInitializer()
 		}
 		else
 		{
-			Prerequisites.Emplace(Index, INDEX_NONE, 0);
+			// index 0 is the visibility attribute
+			int32 PrerequisiteIndex = INDEX_NONE;
+			if (Index != 0 && Attribute.bAffectVisibility)
+			{
+				bSortNeeded = true;
+				PrerequisiteIndex = 0; // After Visibility
+			}
+			Prerequisites.Emplace(Index, PrerequisiteIndex, 0);
 		}
 	}
 
-	if (bHavePrerequisite)
+	if (bSortNeeded)
 	{
 		// Get the depth order
 		for (FPrerequisiteSort& PrerequisiteSort : Prerequisites)
