Fix TObservableArray calling IndexByPredicate instead of IndexOfByPredicate
Deprecate IndexByPredicate in favor of IndexOfByPredicate to match TArray's method

diff --git a/Engine/Source/Runtime/SlateCore/Public/Containers/ObservableArray.h b/Engine/Source/Runtime/SlateCore/Public/Containers/ObservableArray.h
--- a/Engine/Source/Runtime/SlateCore/Public/Containers/ObservableArray.h
+++ b/Engine/Source/Runtime/SlateCore/Public/Containers/ObservableArray.h
@@ -484,9 +484,16 @@ public:
 	}
 
 	template <typename InPredicate>
+	UE_DEPRECATED(5.4, "IndexByPredicate is deprecated. Use IndexOfByPredicate instead")
 	SizeType IndexByPredicate(InPredicate Pred) const
 	{
-		return Array.IndexByPredicate(Pred);
+		return Array.IndexOfByPredicate(Pred);
+	}
+
+	template <typename InPredicate>
+	SizeType IndexOfByPredicate(InPredicate Pred) const
+	{
+		return Array.IndexOfByPredicate(Pred);
 	}
 
 public:
