Fixed SWindow::SetCachedSize to apply the new size

diff --git a/Engine/Source/Runtime/SlateCore/Private/Widgets/SWindow.cpp b/Engine/Source/Runtime/SlateCore/Private/Widgets/SWindow.cpp
--- a/Engine/Source/Runtime/SlateCore/Private/Widgets/SWindow.cpp
+++ b/Engine/Source/Runtime/SlateCore/Private/Widgets/SWindow.cpp
@@ -1050,6 +1050,7 @@ void SWindow::SetCachedSize( UE::Slate::FDeprecateVector2DParameter NewSize )
 	{
 		FVector2d NewSize2d(NewSize);
 		NativeWindow->AdjustCachedSize( NewSize2d );
+		NewSize = UE::Slate::CastToVector2f(NewSize2d);
 	}
 
 	if(Size != NewSize)
