Fix typo in documentation comment of FWidget Path

diff --git a/Engine/Source/Runtime/SlateCore/Public/Layout/WidgetPath.h b/Engine/Source/Runtime/SlateCore/Public/Layout/WidgetPath.h
--- a/Engine/Source/Runtime/SlateCore/Public/Layout/WidgetPath.h
+++ b/Engine/Source/Runtime/SlateCore/Public/Layout/WidgetPath.h
@@ -70,19 +70,19 @@ public:
 	TOptional<FWidgetAndPointer> FindArrangedWidgetAndCursor( TSharedRef<const SWidget> WidgetToFind ) const;
 
 	/**
-	 * Get the first (top-most) widget in this path, which is always a window; assumes path is valid
+	 * Get the first (top-most) widget in this path, which is always a window; assumes path is valid.
 	 *
-	 * @return Window at the top of this path
+	 * @return Window at the top of this path.
 	 */
 	TSharedRef<SWindow> GetWindow() const;
 
 	/**
 	 * Get the deepest (bottom-most) window in this path; assumes path is valid.
 	 *
-	 * @return Window at the top of this path
+	 * @return Window at the bottom of this path.
 	 */
 	TSharedRef<SWindow> GetDeepestWindow() const;
-	
+
 	/** A valid path has at least one widget in it */
 	bool IsValid() const;
 	
