Slate: Protect the HittestGrid from invalid grid size.

diff --git a/Engine/Source/Runtime/SlateCore/Private/Input/HittestGrid.cpp b/Engine/Source/Runtime/SlateCore/Private/Input/HittestGrid.cpp
--- a/Engine/Source/Runtime/SlateCore/Private/Input/HittestGrid.cpp
+++ b/Engine/Source/Runtime/SlateCore/Private/Input/HittestGrid.cpp
@@ -253,7 +253,17 @@ bool FHittestGrid::SetHittestArea(const UE::Slate::FDeprecateVector2DParameter&
 		NumCells = FIntPoint(FMath::CeilToInt(GridSize.X / CellSize.X), FMath::CeilToInt(GridSize.Y / CellSize.Y));
 		
 		const int32 NewTotalCells = NumCells.X * NumCells.Y;
-		ClearInternal(NewTotalCells);
+		if (NewTotalCells < 0)
+		{
+			ensureMsgf(false, TEXT("The gridsize is invalid. (%f, %f) (%d, %d)"), GridSize.X, GridSize.Y, NumCells.X, NumCells.Y);
+			GridSize = FVector2f(0.f, 0.f);
+			NumCells = 0;
+			ClearInternal(0);
+		}
+		else
+		{
+			ClearInternal(NewTotalCells);
+		}
 
 		bWasCleared = true;
 	}
