Fixed issue where cursor lock position would be misaligned after moving window with global invalidation enabled

diff --git a/Engine/Source/Runtime/SlateCore/Private/FastUpdate/SlateInvalidationRoot.cpp b/Engine/Source/Runtime/SlateCore/Private/FastUpdate/SlateInvalidationRoot.cpp
--- a/Engine/Source/Runtime/SlateCore/Private/FastUpdate/SlateInvalidationRoot.cpp
+++ b/Engine/Source/Runtime/SlateCore/Private/FastUpdate/SlateInvalidationRoot.cpp
@@ -805,6 +805,7 @@ void FSlateInvalidationRoot::AdjustWidgetsDesktopGeometry(UE::Slate::FDeprecateV
 
 	FastWidgetPathList->ForEachWidget([WindowToDesktopTransform, &WindowToDesktop](SWidget& Widget)
 		{
+			Widget.PersistentState.DesktopGeometry = Widget.PersistentState.AllottedGeometry;
 			Widget.PersistentState.DesktopGeometry.AppendTransform(WindowToDesktop);
 		});
 }
