Slate: fix retainer box not doing prepass on its content when in slow path.
bNeedsSlowPath is set to false when we call FSlateInvalidationRoot::PaintInvalidationRoot, so we can't rely on it later in the prepass call to determine if we are in slow path.
The issue is more hidden in Invalidation Panel as it avoids calling InvalidateRootLayout() and PaintInvalidationRoot() at the same time in most cases (the first call sets NeedsSlowPath to true and the second sets it to false before we get to prepass).

diff --git a/Engine/Source/Runtime/Slate/Private/Widgets/SInvalidationPanel.cpp b/Engine/Source/Runtime/Slate/Private/Widgets/SInvalidationPanel.cpp
--- a/Engine/Source/Runtime/Slate/Private/Widgets/SInvalidationPanel.cpp
+++ b/Engine/Source/Runtime/Slate/Private/Widgets/SInvalidationPanel.cpp
@@ -277,6 +277,10 @@ bool SInvalidationPanel::CustomPrepass(float LayoutScaleMultiplier)
 
 	if (GetCanCache())
 	{
+		if (NeedsPrepass())
+		{
+			SetNeedsSlowPath(true);
+		}
 		ProcessInvalidation();
 		if (NeedsSlowPath())
 		{
diff --git a/Engine/Source/Runtime/SlateCore/Private/FastUpdate/SlateInvalidationRoot.cpp b/Engine/Source/Runtime/SlateCore/Private/FastUpdate/SlateInvalidationRoot.cpp
--- a/Engine/Source/Runtime/SlateCore/Private/FastUpdate/SlateInvalidationRoot.cpp
+++ b/Engine/Source/Runtime/SlateCore/Private/FastUpdate/SlateInvalidationRoot.cpp
@@ -1358,6 +1358,11 @@ void FSlateInvalidationRoot::ClearAllFastPathData(bool bClearResourcesImmediatel
 	FinalUpdateList.Empty();
 }
 
+void FSlateInvalidationRoot::SetNeedsSlowPath(bool InNeedsSlowPath)
+{
+	bNeedsSlowPath = InNeedsSlowPath;
+}
+
 void FSlateInvalidationRoot::HandleInvalidateAllWidgets(bool bClearResourcesImmediately)
 {
 	Advanced_ResetInvalidation(bClearResourcesImmediately);
diff --git a/Engine/Source/Runtime/SlateCore/Public/FastUpdate/SlateInvalidationRoot.h b/Engine/Source/Runtime/SlateCore/Public/FastUpdate/SlateInvalidationRoot.h
--- a/Engine/Source/Runtime/SlateCore/Public/FastUpdate/SlateInvalidationRoot.h
+++ b/Engine/Source/Runtime/SlateCore/Public/FastUpdate/SlateInvalidationRoot.h
@@ -158,6 +158,8 @@ protected:
 
 	virtual void OnRootInvalidated() { }
 
+	SLATECORE_API void SetNeedsSlowPath(bool InNeedsSlowPath);
+
 private:
 	FSlateInvalidationWidgetList& GetFastPathWidgetList() { return *FastWidgetPathList; }
 	void HandleInvalidateAllWidgets(bool bClearResourcesImmediately);
diff --git a/Engine/Source/Runtime/UMG/Private/Slate/SRetainerWidget.cpp b/Engine/Source/Runtime/UMG/Private/Slate/SRetainerWidget.cpp
--- a/Engine/Source/Runtime/UMG/Private/Slate/SRetainerWidget.cpp
+++ b/Engine/Source/Runtime/UMG/Private/Slate/SRetainerWidget.cpp
@@ -705,6 +705,10 @@ bool SRetainerWidget::CustomPrepass(float LayoutScaleMultiplier)
 {
 	if (bEnableRetainedRendering)
 	{
+		if (NeedsPrepass())
+		{
+			SetNeedsSlowPath(true);
+		}
 		ProcessInvalidation();
 		if (NeedsSlowPath())
 		{
