Further optimization for atlas texture slot lookup; keep track of free slots in power-of-two width buckets to reduce unnecessary lookups. Results in a 4x average speedup (3.6ms vs. 14.4ms of total look-up time in a capture).

diff --git a/Engine/Source/Runtime/SlateCore/Private/Textures/TextureAtlas.cpp b/Engine/Source/Runtime/SlateCore/Private/Textures/TextureAtlas.cpp
--- a/Engine/Source/Runtime/SlateCore/Private/Textures/TextureAtlas.cpp
+++ b/Engine/Source/Runtime/SlateCore/Private/Textures/TextureAtlas.cpp
@@ -62,30 +62,24 @@ FSlateTextureAtlas::~FSlateTextureAtlas()
 void FSlateTextureAtlas::EmptyAtlasData()
 {
 	// Remove all nodes
-	TArray<FAtlasedTextureSlot*, FConcurrentLinearArrayAllocator> DeleteSlots;
-
-	for (FAtlasedTextureSlot::TIterator SlotIt(AtlasUsedSlots); SlotIt; SlotIt.Next())
+	for (FAtlasedTextureSlot* AtlasEmptySlots : AtlasEmptySlotsMap)
 	{
-		FAtlasedTextureSlot& CurSlot = *SlotIt;
-		DeleteSlots.Add(&CurSlot);
+		for (FAtlasedTextureSlot::TIterator SlotIt(AtlasEmptySlots); SlotIt;)
+		{
+			FAtlasedTextureSlot& CurSlot = *SlotIt;
+			SlotIt.Next();
+			delete &CurSlot;
+		}
 	}
+	AtlasEmptySlotsMap.Reset();
 
-	for (FAtlasedTextureSlot::TIterator SlotIt(AtlasEmptySlots); SlotIt; SlotIt.Next())
+	for (FAtlasedTextureSlot::TIterator SlotIt(AtlasUsedSlots); SlotIt;)
 	{
 		FAtlasedTextureSlot& CurSlot = *SlotIt;
-		DeleteSlots.Add(&CurSlot);
+		SlotIt.Next();
+		delete &CurSlot;
 	}
-
-	AtlasUsedSlots = NULL;
-	AtlasEmptySlots = NULL;
-
-	for (FAtlasedTextureSlot* CurSlot : DeleteSlots)
-	{
-		delete CurSlot;
-	}
-
-	DeleteSlots.Empty();
-
+	AtlasUsedSlots = nullptr;
 
 	STAT(SIZE_T MemoryBefore = AtlasData.GetAllocatedSize());
 
@@ -131,11 +125,11 @@ void FSlateTextureAtlas::InitAtlasData()
 {
 	LLM_SCOPE_BYTAG(UI_Texture);
 
-	check(AtlasEmptySlots == NULL && AtlasData.Num() == 0);
-
-	FAtlasedTextureSlot* RootSlot = new FAtlasedTextureSlot(0, 0, AtlasWidth, AtlasHeight, GetPaddingAmount());
+	check(AtlasEmptySlotsMap.IsEmpty() && AtlasData.Num() == 0);
 
-	RootSlot->LinkHead(AtlasEmptySlots);
+	const int32 MapSlotCount = GetFreeSlotSearchIndex(AtlasWidth, AtlasHeight) + 1;
+	AtlasEmptySlotsMap.SetNumZeroed(MapSlotCount);
+	AddFreeSlot(0, 0, AtlasWidth, AtlasHeight);
 
 	AtlasData.Reserve(AtlasWidth * AtlasHeight * BytesPerPixel);
 	AtlasData.AddZeroed(AtlasWidth * AtlasHeight * BytesPerPixel);
@@ -287,112 +281,94 @@ const FAtlasedTextureSlot* FSlateTextureAtlas::GetSlotAtPosition(FIntPoint InPos
 
 const FAtlasedTextureSlot* FSlateTextureAtlas::FindSlotForTexture(uint32 InWidth, uint32 InHeight)
 {
-	FAtlasedTextureSlot* ReturnVal = NULL;
-
 	// Account for padding on both sides
 	const uint8 Padding = GetPaddingAmount();
 	const uint32 TotalPadding = Padding * 2;
 	const uint32 PaddedWidth = InWidth + TotalPadding;
 	const uint32 PaddedHeight = InHeight + TotalPadding;
-
-	// Previously, slots were stored as a binary tree - this has been replaced with a linked-list of slots on the edge of the tree
-	// (slots on the edge of the tree represent empty slots); this iterates empty slots in same order as a binary depth-first-search,
-	// except much faster.
-	for (FAtlasedTextureSlot::TIterator SlotIt(AtlasEmptySlots); SlotIt; SlotIt++)
+	const int32 StartSearchIndex = GetFreeSlotSearchIndex(PaddedWidth, PaddedHeight);
+	for (int32 SlotsIndex = StartSearchIndex; SlotsIndex < AtlasEmptySlotsMap.Num(); ++SlotsIndex)
 	{
-		FAtlasedTextureSlot& CurSlot = *SlotIt;
-
-		if (PaddedWidth <= CurSlot.Width && PaddedHeight <= CurSlot.Height)
+		FAtlasedTextureSlot* const AtlasEmptySlots = AtlasEmptySlotsMap[SlotsIndex];
+		for (FAtlasedTextureSlot::TIterator SlotIt(AtlasEmptySlots); SlotIt; SlotIt++)
 		{
-			ReturnVal = &CurSlot;
-			break;
-		}
-	}
-
-
-	if (ReturnVal != NULL)
-	{
-		// The width and height of the new child node
-		const uint32 RemainingWidth =  FMath::Max<int32>(0, ReturnVal->Width - PaddedWidth);
-		const uint32 RemainingHeight = FMath::Max<int32>(0, ReturnVal->Height - PaddedHeight);
+			FAtlasedTextureSlot& CurSlot = *SlotIt;
+			if (PaddedWidth <= CurSlot.Width && PaddedHeight <= CurSlot.Height)
+			{
+				// The width and height of the new child node(s)
+				const uint32 RemainingWidth = CurSlot.Width - PaddedWidth;
+				const uint32 RemainingHeight = CurSlot.Height - PaddedHeight;
 
-		// New slots must have a minimum width/height, to avoid excessive slots i.e. excessive memory usage and iteration.
-		// No glyphs seem to use slots this small, and cutting these slots out improves performance/memory-usage a fair bit
-		const uint32 MinSlotDim = 2;
+				// New slots must have a minimum width/height, to avoid excessive slots i.e. excessive memory usage and iteration.
+				// No glyphs seem to use slots this small, and cutting these slots out improves performance/memory-usage a fair bit
+				const uint32 MinSlotDim = 2;
 
-		// Split the remaining area around this slot into two children.
-		if (RemainingHeight >= MinSlotDim || RemainingWidth >= MinSlotDim)
-		{
-			FAtlasedTextureSlot* LeftSlot = nullptr;
-			FAtlasedTextureSlot* RightSlot = nullptr;
-
-			if (RemainingHeight <= RemainingWidth)
-			{
-				// Split vertically
-				// - - - - - - - - -
-				// |       |       |
-				// |  Slot |       |
-				// |       |       |
-				// | - - - | Right |
-				// |       |       |
-				// |  Left |       |
-				// |       |       |
-				// - - - - - - - - -
-				if (RemainingHeight >= MinSlotDim)
-				{
-					LeftSlot = new FAtlasedTextureSlot(ReturnVal->X, ReturnVal->Y + PaddedHeight, PaddedWidth, RemainingHeight, Padding);
-				}
-				RightSlot = new FAtlasedTextureSlot(ReturnVal->X + PaddedWidth, ReturnVal->Y, RemainingWidth, ReturnVal->Height, Padding);
-			}
-			else
-			{
-				// Split horizontally
-				// - - - - - - - - -
-				// |       |       |
-				// |  Slot | Left  |
-				// |       |       |
-				// | - - - - - - - |
-				// |               |
-				// |     Right     |
-				// |               |
-				// - - - - - - - - -
-				if (RemainingWidth >= MinSlotDim)
+				// Split the remaining area around this slot into two children.
+				if (RemainingHeight >= MinSlotDim || RemainingWidth >= MinSlotDim)
 				{
-					LeftSlot = new FAtlasedTextureSlot(ReturnVal->X + PaddedWidth, ReturnVal->Y, RemainingWidth, PaddedHeight, Padding);
+					if (RemainingHeight <= RemainingWidth)
+					{
+						// Split vertically
+						// - - - - - - - - -
+						// |       |       |
+						// |  Slot |       |
+						// |       |       |
+						// | - - - | Right |
+						// |       |       |
+						// |  Left |       |
+						// |       |       |
+						// - - - - - - - - -
+						AddFreeSlot(CurSlot.X + PaddedWidth, CurSlot.Y, RemainingWidth, CurSlot.Height);
+						if (RemainingHeight >= MinSlotDim)
+						{
+							AddFreeSlot(CurSlot.X, CurSlot.Y + PaddedHeight, PaddedWidth, RemainingHeight);
+						}
+					}
+					else
+					{
+						// Split horizontally
+						// - - - - - - - - -
+						// |       |       |
+						// |  Slot | Left  |
+						// |       |       |
+						// | - - - - - - - |
+						// |               |
+						// |     Right     |
+						// |               |
+						// - - - - - - - - -
+						AddFreeSlot(CurSlot.X, CurSlot.Y + PaddedHeight, CurSlot.Width, RemainingHeight);
+						if (RemainingWidth >= MinSlotDim)
+						{
+							AddFreeSlot(CurSlot.X + PaddedWidth, CurSlot.Y, RemainingWidth, PaddedHeight);
+						}
+					}
 				}
-				RightSlot = new FAtlasedTextureSlot(ReturnVal->X, ReturnVal->Y + PaddedHeight, ReturnVal->Width, RemainingHeight, Padding);
-			}
 
-			// Replace the old slot within AtlasEmptySlots, with the new Left and Right slot, then add the old slot to AtlasUsedSlots
-			if (LeftSlot)
-			{
-				LeftSlot->LinkReplace(ReturnVal);
-				RightSlot->LinkAfter(LeftSlot);
+				// Shrink and moved to used slot list
+				CurSlot.Width = PaddedWidth;
+				CurSlot.Height = PaddedHeight;
+				CurSlot.Unlink();
+				CurSlot.LinkHead(AtlasUsedSlots);
+				return &CurSlot;
 			}
-			else
-			{
-				RightSlot->LinkReplace(ReturnVal);
-			}
-
-			ReturnVal->LinkHead(AtlasUsedSlots);
 		}
-		else
-		{
-			// Remove the old slot from AtlasEmptySlots, into AtlasUsedSlots
-			ReturnVal->Unlink();
-			ReturnVal->LinkHead(AtlasUsedSlots);
-		}
-
-
-		// Shrink the slot to the remaining area.
-		ReturnVal->Width = PaddedWidth;
-		ReturnVal->Height = PaddedHeight;
 	}
 
-	return ReturnVal;
+	return nullptr;
 }
 
+int32 FSlateTextureAtlas::GetFreeSlotSearchIndex(uint32 InWidth, uint32 InHeight)
+{
+	// Currently only bucketing by width, but leaving this open to be extended
+	return FMath::CeilLogTwo(InWidth);
+}
 
+void FSlateTextureAtlas::AddFreeSlot(uint32 InX, uint32 InY, uint32 InWidth, uint32 InHeight)
+{
+	FAtlasedTextureSlot* NewSlot = new FAtlasedTextureSlot(InX, InY, InWidth, InHeight, GetPaddingAmount());
+	const uint32 SlotIndex = GetFreeSlotSearchIndex(InWidth, InHeight);
+	NewSlot->LinkHead(AtlasEmptySlotsMap[SlotIndex]);
+}
 
 FSlateFlushableAtlasCache::FSlateFlushableAtlasCache(const FAtlasFlushParams* InFlushParams)
 	: FlushParams(InFlushParams)
diff --git a/Engine/Source/Runtime/SlateCore/Public/Textures/TextureAtlas.h b/Engine/Source/Runtime/SlateCore/Public/Textures/TextureAtlas.h
--- a/Engine/Source/Runtime/SlateCore/Public/Textures/TextureAtlas.h
+++ b/Engine/Source/Runtime/SlateCore/Public/Textures/TextureAtlas.h
@@ -101,7 +101,7 @@ public:
 	FSlateTextureAtlas( uint32 InWidth, uint32 InHeight, uint32 InBytesPerPixel, ESlateTextureAtlasPaddingStyle InPaddingStyle, bool bInUpdatesAfterInitialization )
 		: AtlasData()
 		, AtlasUsedSlots(NULL)
-		, AtlasEmptySlots(NULL)
+		, AtlasEmptySlotsMap()
 		, AtlasWidth( InWidth )
 		, AtlasHeight( InHeight )
 		, BytesPerPixel( InBytesPerPixel )
@@ -163,6 +163,16 @@ protected:
 	 */
 	SLATECORE_API const FAtlasedTextureSlot* FindSlotForTexture( uint32 InWidth, uint32 InHeight );
 
+	/**
+	 * Get the index to start looking for a free slot.
+	 */
+	static int32 GetFreeSlotSearchIndex(uint32 InWidth, uint32 InHeight);
+
+	/**
+	 * Adds a new slot to the free slot list.
+	 */
+	void AddFreeSlot(uint32 InX, uint32 InY, uint32 InWidth, uint32 InHeight);
+
 	/**
 	 * Creates enough space for a single texture the width and height of the atlas
 	 */
@@ -222,7 +232,7 @@ protected:
 	/** The list of atlas slots pointing to used texture data in the atlas */
 	FAtlasedTextureSlot* AtlasUsedSlots;
 	/** The list of atlas slots pointing to empty texture data in the atlas */
-	FAtlasedTextureSlot* AtlasEmptySlots;
+	TArray<FAtlasedTextureSlot*> AtlasEmptySlotsMap;
 	/** Width of the atlas */
 	uint32 AtlasWidth;
 	/** Height of the atlas */
