Suppressing PVS warning about pointer leak, it's an intrusive linked list node that unconditionally gets set as the head.

#rnx

diff --git a/Engine/Source/Runtime/SlateCore/Private/Textures/TextureAtlas.cpp b/Engine/Source/Runtime/SlateCore/Private/Textures/TextureAtlas.cpp
--- a/Engine/Source/Runtime/SlateCore/Private/Textures/TextureAtlas.cpp
+++ b/Engine/Source/Runtime/SlateCore/Private/Textures/TextureAtlas.cpp
@@ -368,7 +368,7 @@ void FSlateTextureAtlas::AddFreeSlot(uint32 InX, uint32 InY, uint32 InWidth, uin
 	FAtlasedTextureSlot* NewSlot = new FAtlasedTextureSlot(InX, InY, InWidth, InHeight, GetPaddingAmount());
 	const uint32 SlotIndex = GetFreeSlotSearchIndex(InWidth, InHeight);
 	NewSlot->LinkHead(AtlasEmptySlotsMap[SlotIndex]);
-}
+} //-V773
 
 FSlateFlushableAtlasCache::FSlateFlushableAtlasCache(const FAtlasFlushParams* InFlushParams)
 	: FlushParams(InFlushParams)
