[MemoryTrace] Covered more allocations with asset metadata scopes
- Added package/object/class metadata for Iris network spawned actors
- Added asset metadata for ConfigFiles
- Added package metadata for allocations coming from LinkerLoadPackageSummary
- Added metadata for Object::Serialize allocations
- Added metadata for UI/Slate and UI/UMG.
- Removed metadata scope for UObjectHash/UObjectIndex global buffers realloc.

#tests PIE,client

diff --git a/Engine/Source/Runtime/Core/Private/Misc/ConfigCacheIni.cpp b/Engine/Source/Runtime/Core/Private/Misc/ConfigCacheIni.cpp
--- a/Engine/Source/Runtime/Core/Private/Misc/ConfigCacheIni.cpp
+++ b/Engine/Source/Runtime/Core/Private/Misc/ConfigCacheIni.cpp
@@ -12,6 +12,7 @@
 #include "Math/Vector4.h"
 #include "Stats/Stats.h"
 #include "HAL/IConsoleManager.h"
+#include "HAL/LowLevelMemStats.h"
 #include "HAL/LowLevelMemTracker.h"
 #include "Misc/CoreDelegates.h"
 #include "Misc/App.h"
@@ -22,6 +23,7 @@
 #include "Misc/StringBuilder.h"
 #include "Misc/Paths.h"
 #include "ProfilingDebugging/CpuProfilerTrace.h"
+#include "ProfilingDebugging/AssetMetadataTrace.h"
 #include "Serialization/MemoryReader.h"
 #include "Serialization/MemoryWriter.h"
 #include "Serialization/LargeMemoryReader.h"
@@ -755,6 +757,12 @@ static void WarnAboutKeyRemap(const FString& OldValue, const FString& NewValue,
 
 void FConfigFile::CombineFromBuffer(const FString& Buffer, const FString& FileHint)
 {
+	static const FName ConfigFileClassName = TEXT("ConfigFile");
+	const FName FileName = FName(*FileHint);
+	LLM_SCOPE_DYNAMIC_STAT_OBJECTPATH_FNAME(FileName, ELLMTagSet::Assets);
+	LLM_SCOPE_DYNAMIC_STAT_OBJECTPATH_FNAME(ConfigFileClassName, ELLMTagSet::AssetClasses);
+	UE_TRACE_METADATA_SCOPE_ASSET_FNAME(FileName, ConfigFileClassName, FileName);
+
 	const TCHAR* Ptr = *Buffer;
 	FConfigSection* CurrentSection = nullptr;
 	FString CurrentSectionName;
@@ -986,6 +994,12 @@ void FConfigFile::CombineFromBuffer(const FString& Buffer, const FString& FileHi
  */
 void FConfigFile::ProcessInputFileContents(FStringView Contents, const FString& FileHint)
 {
+	static const FName ConfigFileClassName = TEXT("ConfigFile");
+	const FName FileName = FName(*FileHint);
+	LLM_SCOPE_DYNAMIC_STAT_OBJECTPATH_FNAME(FileName, ELLMTagSet::Assets);
+	LLM_SCOPE_DYNAMIC_STAT_OBJECTPATH_FNAME(ConfigFileClassName, ELLMTagSet::AssetClasses);
+	UE_TRACE_METADATA_SCOPE_ASSET_FNAME(FileName, ConfigFileClassName, FileName);
+
 	const TCHAR* Ptr = Contents.Len() > 0 ? Contents.GetData() : nullptr;
 	FConfigSection* CurrentSection = nullptr;
 	FString CurrentSectionName;
diff --git a/Engine/Source/Runtime/Core/Private/Misc/ConfigContext.cpp b/Engine/Source/Runtime/Core/Private/Misc/ConfigContext.cpp
--- a/Engine/Source/Runtime/Core/Private/Misc/ConfigContext.cpp
+++ b/Engine/Source/Runtime/Core/Private/Misc/ConfigContext.cpp
@@ -11,6 +11,8 @@
 #include "Misc/Paths.h"
 #include "Misc/DataDrivenPlatformInfoRegistry.h"
 #include "ProfilingDebugging/CpuProfilerTrace.h"
+#include "ProfilingDebugging/AssetMetadataTrace.h"
+#include "HAL/LowLevelMemStats.h"
 
 namespace
 {
@@ -331,6 +333,7 @@ static void LoadAnIniFile(const FString& FilenameToLoad, FConfigFile& ConfigFile
 bool FConfigContext::PerformLoad()
 {
 	LLM_SCOPE(ELLMTag::ConfigSystem);
+	static const FName ConfigContextClassName = TEXT("ConfigContext");
 
 	// if bIsBaseIniName is false, that means the .ini is a ready-to-go .ini file, and just needs to be loaded into the FConfigFile
 	if (!bIsHierarchicalConfig)
@@ -347,15 +350,24 @@ bool FConfigContext::PerformLoad()
 			DestIniFilename = FString::Printf(TEXT("%s/%s.ini"), *ProjectConfigDir, *BaseIniName);
 		}
 
+		const FName BaseName = FName(*BaseIniName);
+		LLM_SCOPE_DYNAMIC_STAT_OBJECTPATH_FNAME(BaseName, ELLMTagSet::Assets);
+		LLM_SCOPE_DYNAMIC_STAT_OBJECTPATH_FNAME(ConfigContextClassName, ELLMTagSet::AssetClasses);
+		UE_TRACE_METADATA_SCOPE_ASSET_FNAME(BaseName, ConfigContextClassName, BaseName);
+
 		// load the .ini file straight up
 		LoadAnIniFile(*DestIniFilename, *ConfigFile);
 
-		ConfigFile->Name = FName(*BaseIniName);
+		ConfigFile->Name = BaseName;
 		ConfigFile->PlatformName.Reset();
 		ConfigFile->bHasPlatformName = false;
 	}
 	else
 	{
+		const FName BaseName = FName(*BaseIniName);
+		LLM_SCOPE_DYNAMIC_STAT_OBJECTPATH_FNAME(BaseName, ELLMTagSet::Assets);
+		LLM_SCOPE_DYNAMIC_STAT_OBJECTPATH_FNAME(ConfigContextClassName, ELLMTagSet::AssetClasses);
+		UE_TRACE_METADATA_SCOPE_ASSET_FNAME(BaseName, ConfigContextClassName, BaseName);
 #if DISABLE_GENERATED_INI_WHEN_COOKED
 		if (BaseIniName != TEXT("GameUserSettings"))
 		{
@@ -384,7 +396,7 @@ bool FConfigContext::PerformLoad()
 		// just write out the exact same thing it read in!
 		bool bNeedsWrite = GenerateDestIniFile();
 
-		ConfigFile->Name = FName(*BaseIniName);
+		ConfigFile->Name = BaseName;
 		ConfigFile->PlatformName = Platform;
 		ConfigFile->bHasPlatformName = true;
 
diff --git a/Engine/Source/Runtime/CoreUObject/Private/Serialization/AsyncLoading2.cpp b/Engine/Source/Runtime/CoreUObject/Private/Serialization/AsyncLoading2.cpp
--- a/Engine/Source/Runtime/CoreUObject/Private/Serialization/AsyncLoading2.cpp
+++ b/Engine/Source/Runtime/CoreUObject/Private/Serialization/AsyncLoading2.cpp
@@ -5378,6 +5378,8 @@ void FAsyncPackage2::StartLoading(FAsyncLoadingThreadState2& ThreadState, FIoBat
 EEventLoadNodeExecutionResult FAsyncPackage2::ProcessLinkerLoadPackageSummary(FAsyncLoadingThreadState2& ThreadState)
 {
 	TRACE_CPUPROFILER_EVENT_SCOPE(ProcessLinkerLoadPackageSummary);
+	LLM_SCOPE_DYNAMIC_STAT_OBJECTPATH_FNAME(LinkerLoadState->Linker->GetPackagePath().GetPackageFName(), ELLMTagSet::Assets);
+	UE_TRACE_METADATA_SCOPE_ASSET_FNAME(NAME_None, NAME_None, LinkerLoadState->Linker->GetPackagePath().GetPackageFName());
 	SCOPED_LOADTIMER_ASSET_TEXT(*LinkerLoadState->Linker->GetDebugName());
 
 #if ALT2_ENABLE_NEW_ARCHIVE_FOR_LINKERLOAD
diff --git a/Engine/Source/Runtime/CoreUObject/Private/UObject/LinkerLoad.cpp b/Engine/Source/Runtime/CoreUObject/Private/UObject/LinkerLoad.cpp
--- a/Engine/Source/Runtime/CoreUObject/Private/UObject/LinkerLoad.cpp
+++ b/Engine/Source/Runtime/CoreUObject/Private/UObject/LinkerLoad.cpp
@@ -43,6 +43,8 @@
 #include "Serialization/AsyncLoading.h"
 #include "Serialization/ArchiveSerializedPropertyChain.h"
 #include "ProfilingDebugging/LoadTimeTracker.h"
+#include "ProfilingDebugging/AssetMetadataTrace.h"
+#include "HAL/LowLevelMemStats.h"
 #include "HAL/ThreadHeartBeat.h"
 #include "Internationalization/TextPackageNamespaceUtil.h"
 #include "Serialization/BulkData.h"
@@ -4552,6 +4554,9 @@ void FLinkerLoad::Preload( UObject* Object )
 				{
 					SCOPE_CYCLE_COUNTER(STAT_LinkerSerialize);
 					TRACE_LOADTIME_SERIALIZE_EXPORT_SCOPE(Object, Export.SerialSize);
+					LLM_SCOPE_DYNAMIC_STAT_OBJECTPATH(Object->GetPackage(), ELLMTagSet::Assets);
+					LLM_SCOPE_DYNAMIC_STAT_OBJECTPATH(Object->GetClass(), ELLMTagSet::AssetClasses);
+					UE_TRACE_METADATA_SCOPE_ASSET(Object, Object->GetClass());
 #if USE_CIRCULAR_DEPENDENCY_LOAD_DEFERRING
 					// communicate with FLinkerPlaceholderBase, what object is currently seriAlizing in
 					FScopedPlaceholderContainerTracker SerializingObjTracker(Object);
diff --git a/Engine/Source/Runtime/CoreUObject/Private/UObject/UObjectArray.cpp b/Engine/Source/Runtime/CoreUObject/Private/UObject/UObjectArray.cpp
--- a/Engine/Source/Runtime/CoreUObject/Private/UObject/UObjectArray.cpp
+++ b/Engine/Source/Runtime/CoreUObject/Private/UObject/UObjectArray.cpp
@@ -6,7 +6,9 @@
 
 #include "UObject/UObjectArray.h"
 #include "HAL/IConsoleManager.h"
+#include "HAL/LowLevelMemStats.h"
 #include "Misc/ScopeLock.h"
+#include "ProfilingDebugging/AssetMetadataTrace.h"
 #include "UObject/UObjectAllocator.h"
 #include "UObject/Class.h"
 #include "UObject/UObjectIterator.h"
@@ -190,6 +192,10 @@ void FUObjectArray::DisableDisregardForGC()
 void FUObjectArray::AllocateUObjectIndex(UObjectBase* Object, EInternalObjectFlags InitialFlags, int32 AlreadyAllocatedIndex, int32 SerialNumber)
 {
 	LLM_SCOPE(ELLMTag::UObject);
+	// Clear asset scopes
+	LLM_SCOPE_DYNAMIC_STAT_OBJECTPATH_FNAME(FName{NAME_Default}, ELLMTagSet::Assets);
+	LLM_SCOPE_DYNAMIC_STAT_OBJECTPATH_FNAME(FName{NAME_Default}, ELLMTagSet::AssetClasses);
+	UE_TRACE_METADATA_SCOPE_ASSET_FNAME(NAME_None, NAME_None, NAME_None);
 
 	int32 Index = INDEX_NONE;
 	check(Object->InternalIndex == INDEX_NONE);
diff --git a/Engine/Source/Runtime/CoreUObject/Private/UObject/UObjectHash.cpp b/Engine/Source/Runtime/CoreUObject/Private/UObject/UObjectHash.cpp
--- a/Engine/Source/Runtime/CoreUObject/Private/UObject/UObjectHash.cpp
+++ b/Engine/Source/Runtime/CoreUObject/Private/UObject/UObjectHash.cpp
@@ -14,6 +14,7 @@
 #include "HAL/IConsoleManager.h"
 #include "HAL/LowLevelMemStats.h"
 #include "UObject/AnyPackagePrivate.h"
+#include "ProfilingDebugging/AssetMetadataTrace.h"
 
 #include <atomic>
 
@@ -1122,6 +1123,9 @@ FORCEINLINE static UPackage* UnassignExternalPackageFromObject(FUObjectHashTable
 void ShrinkUObjectHashTables()
 {
 	LLM_SCOPE_BYTAG(UObjectHash);
+	LLM_SCOPE_DYNAMIC_STAT_OBJECTPATH_FNAME(FName{NAME_Default}, ELLMTagSet::Assets);
+	LLM_SCOPE_DYNAMIC_STAT_OBJECTPATH_FNAME(FName{NAME_Default}, ELLMTagSet::AssetClasses);
+	UE_TRACE_METADATA_SCOPE_ASSET_FNAME(NAME_None, NAME_None, NAME_None);
 	TRACE_CPUPROFILER_EVENT_SCOPE(ShrinkUObjectHashTables);
 	FUObjectHashTables& ThreadHash = FUObjectHashTables::Get();
 	FHashTableLock HashLock(ThreadHash);
@@ -1568,6 +1572,9 @@ void HashObject(UObjectBase* Object)
 	if (Name != NAME_None)
 	{
 		LLM_SCOPE_BYTAG(UObjectHash);
+		LLM_SCOPE_DYNAMIC_STAT_OBJECTPATH_FNAME(FName{NAME_Default}, ELLMTagSet::Assets);
+		LLM_SCOPE_DYNAMIC_STAT_OBJECTPATH_FNAME(FName{NAME_Default}, ELLMTagSet::AssetClasses);
+		UE_TRACE_METADATA_SCOPE_ASSET_FNAME(NAME_None, NAME_None, NAME_None);
 #if !UE_BUILD_TEST && !UE_BUILD_SHIPPING
 		SCOPE_CYCLE_COUNTER(STAT_Hash_HashObject);
 #endif
@@ -1615,6 +1622,9 @@ void UnhashObject(UObjectBase* Object)
 	if (Name != NAME_None)
 	{
 		LLM_SCOPE_BYTAG(UObjectHash);
+		LLM_SCOPE_DYNAMIC_STAT_OBJECTPATH_FNAME(FName{NAME_Default}, ELLMTagSet::Assets);
+		LLM_SCOPE_DYNAMIC_STAT_OBJECTPATH_FNAME(FName{NAME_Default}, ELLMTagSet::AssetClasses);
+		UE_TRACE_METADATA_SCOPE_ASSET_FNAME(NAME_None, NAME_None, NAME_None);
 #if !UE_BUILD_TEST && !UE_BUILD_SHIPPING
 		SCOPE_CYCLE_COUNTER(STAT_Hash_UnhashObject);
 #endif
@@ -1652,6 +1662,9 @@ void UnhashObject(UObjectBase* Object)
 void HashObjectExternalPackage(UObjectBase* Object, UPackage* Package)
 {
 	LLM_SCOPE_BYTAG(UObjectHash);
+	LLM_SCOPE_DYNAMIC_STAT_OBJECTPATH_FNAME(FName{NAME_Default}, ELLMTagSet::Assets);
+	LLM_SCOPE_DYNAMIC_STAT_OBJECTPATH_FNAME(FName{NAME_Default}, ELLMTagSet::AssetClasses);
+	UE_TRACE_METADATA_SCOPE_ASSET_FNAME(NAME_None, NAME_None, NAME_None);
 	if (Package)
 	{
 		FUObjectHashTables& ThreadHash = FUObjectHashTables::Get();
@@ -1676,6 +1689,9 @@ void HashObjectExternalPackage(UObjectBase* Object, UPackage* Package)
 void UnhashObjectExternalPackage(class UObjectBase* Object)
 {
 	LLM_SCOPE_BYTAG(UObjectHash);
+	LLM_SCOPE_DYNAMIC_STAT_OBJECTPATH_FNAME(FName{NAME_Default}, ELLMTagSet::Assets);
+	LLM_SCOPE_DYNAMIC_STAT_OBJECTPATH_FNAME(FName{NAME_Default}, ELLMTagSet::AssetClasses);
+	UE_TRACE_METADATA_SCOPE_ASSET_FNAME(NAME_None, NAME_None, NAME_None);
 	FUObjectHashTables& ThreadHash = FUObjectHashTables::Get();
 	FHashTableLock LockHash(ThreadHash);
 	Object->AtomicallyClearFlags(RF_HasExternalPackage);
diff --git a/Engine/Source/Runtime/Engine/Private/Net/Iris/ReplicationSystem/ActorReplicationBridge.cpp b/Engine/Source/Runtime/Engine/Private/Net/Iris/ReplicationSystem/ActorReplicationBridge.cpp
--- a/Engine/Source/Runtime/Engine/Private/Net/Iris/ReplicationSystem/ActorReplicationBridge.cpp
+++ b/Engine/Source/Runtime/Engine/Private/Net/Iris/ReplicationSystem/ActorReplicationBridge.cpp
@@ -30,11 +30,13 @@
 #include "Engine/Level.h"
 #include "Engine/World.h"
 #include "GameFramework/Actor.h"
+#include "HAL/LowLevelMemStats.h"
 #include "Net/DataBunch.h"
 #include "Net/DataChannel.h"
 #include "Net/Core/Connection/NetCloseResult.h"
 #include "Net/Core/Misc/NetSubObjectRegistry.h"
 #include "Net/NetSubObjectRegistryGetter.h"
+#include "ProfilingDebugging/AssetMetadataTrace.h"
 #include "Templates/Casts.h"
 #include "UObject/Package.h"
 #include <limits>
@@ -566,6 +568,9 @@ FObjectReplicationBridgeInstantiateResult UActorReplicationBridge::BeginInstanti
 
 			if (Archetype)
 			{
+				LLM_SCOPE_DYNAMIC_STAT_OBJECTPATH(Archetype->GetPackage(), ELLMTagSet::Assets);
+				LLM_SCOPE_DYNAMIC_STAT_OBJECTPATH(Archetype->GetClass(), ELLMTagSet::AssetClasses);
+				UE_TRACE_METADATA_SCOPE_ASSET(Archetype, Archetype->GetClass());
 				// For streaming levels, it's possible that the owning level has been made not-visible but is still loaded.
 				// In that case, the level will still be found but the owning world will be invalid.
 				// If that happens, wait to spawn the Actor until the next time the level is streamed in.
diff --git a/Engine/Source/Runtime/SlateCore/Private/Widgets/SWidget.cpp b/Engine/Source/Runtime/SlateCore/Private/Widgets/SWidget.cpp
--- a/Engine/Source/Runtime/SlateCore/Private/Widgets/SWidget.cpp
+++ b/Engine/Source/Runtime/SlateCore/Private/Widgets/SWidget.cpp
@@ -29,6 +29,9 @@
 #include "Misc/ScopeRWLock.h"
 #include "HAL/CriticalSection.h"
 #include "Widgets/SWidgetUtils.h"
+#include "ProfilingDebugging/AssetMetadataTrace.h"
+#include "HAL/LowLevelMemStats.h"
+#include "UObject/Package.h"
 
 #include <limits>
 
@@ -687,6 +690,28 @@ void SWidget::SlatePrepass(float InLayoutScaleMultiplier)
 
 	if (!GSlateIsOnFastUpdatePath || bNeedsPrepass)
 	{
+		LLM_SCOPE_BYNAME("UI/Slate/Prepass");
+#if UE_TRACE_ASSET_METADATA_ENABLED
+		FName AssetName = NAME_None;
+		FName ClassName = NAME_None;
+		FName PackageName = NAME_None;
+		if (UE_TRACE_CHANNELEXPR_IS_ENABLED(AssetMetadataChannel))
+		{
+			TSharedPtr<FReflectionMetaData> AssetMetaData = FReflectionMetaData::GetWidgetOrParentMetaData(this);
+			if (AssetMetaData.IsValid())
+			{
+				if (const UObject* AssetPtr = AssetMetaData->Asset.Get())
+				{
+					AssetName = AssetMetaData->Name;
+					ClassName = AssetMetaData->Class.Get()->GetFName();
+					PackageName = AssetPtr->GetPackage()->GetFName();
+				}
+			}
+		}
+		LLM_SCOPE_DYNAMIC_STAT_OBJECTPATH_FNAME(PackageName, ELLMTagSet::Assets);
+		LLM_SCOPE_DYNAMIC_STAT_OBJECTPATH_FNAME(ClassName, ELLMTagSet::AssetClasses);
+		UE_TRACE_METADATA_SCOPE_ASSET_FNAME(AssetName, ClassName, PackageName);
+#endif
 		if (HasRegisteredSlateAttribute() && IsAttributesUpdatesEnabled() && !GSlateIsOnFastProcessInvalidation)
 		{
 			FSlateAttributeMetaData::UpdateAllAttributes(*this, FSlateAttributeMetaData::EInvalidationPermission::AllowInvalidationIfConstructed);
@@ -1569,8 +1594,33 @@ int32 SWidget::Paint(const FPaintArgs& Args, const FGeometry& AllottedGeometry,
 	}
 #endif
 	
+	int32 NewLayerId = 0;
+	{
+		LLM_SCOPE_BYNAME("UI/Slate/OnPaint");
+#if UE_TRACE_ASSET_METADATA_ENABLED
+		FName AssetName = NAME_None;
+		FName ClassName = NAME_None;
+		FName PackageName = NAME_None;
+		if (UE_TRACE_CHANNELEXPR_IS_ENABLED(AssetMetadataChannel))
+		{
+			TSharedPtr<FReflectionMetaData> AssetMetaData = FReflectionMetaData::GetWidgetOrParentMetaData(this);
+			if (AssetMetaData.IsValid())
+			{
+				if (const UObject* AssetPtr = AssetMetaData->Asset.Get())
+				{
+					AssetName = AssetMetaData->Name;
+					ClassName = AssetMetaData->Class.Get()->GetFName();
+					PackageName = AssetPtr->GetPackage()->GetFName();
+				}
+			}
+		}
+		LLM_SCOPE_DYNAMIC_STAT_OBJECTPATH_FNAME(PackageName, ELLMTagSet::Assets);
+		LLM_SCOPE_DYNAMIC_STAT_OBJECTPATH_FNAME(ClassName, ELLMTagSet::AssetClasses);
+		UE_TRACE_METADATA_SCOPE_ASSET_FNAME(AssetName, ClassName, PackageName);
+#endif
 	// Paint the geometry of this widget.
-	int32 NewLayerId = OnPaint(UpdatedArgs, AllottedGeometry, CullingBounds, OutDrawElements, LayerId, ContentWidgetStyle, bParentEnabled);
+		NewLayerId = OnPaint(UpdatedArgs, AllottedGeometry, CullingBounds, OutDrawElements, LayerId, ContentWidgetStyle, bParentEnabled);
+	}
 
 	// Just repainted
 	MutableThis->RemoveUpdateFlags(EWidgetUpdateFlags::NeedsRepaint);
diff --git a/Engine/Source/Runtime/UMG/Private/Components/Widget.cpp b/Engine/Source/Runtime/UMG/Private/Components/Widget.cpp
--- a/Engine/Source/Runtime/UMG/Private/Components/Widget.cpp
+++ b/Engine/Source/Runtime/UMG/Private/Components/Widget.cpp
@@ -13,6 +13,7 @@
 #include "UObject/ObjectSaveContext.h"
 #include "UObject/UObjectHash.h"
 #include "UObject/UObjectIterator.h"
+#include "UObject/Package.h"
 #include "Engine/LocalPlayer.h"
 #include "Engine/UserInterfaceSettings.h"
 #include "Framework/Application/SlateApplication.h"
@@ -36,6 +37,8 @@
 #include "Trace/SlateMemoryTags.h"
 #include "Serialization/PropertyLocalizationDataGathering.h"
 #include "Components/NamedSlotInterface.h"
+#include "ProfilingDebugging/AssetMetadataTrace.h"
+#include "HAL/LowLevelMemStats.h"
 
 #include UE_INLINE_GENERATED_CPP_BY_NAME(Widget)
 
@@ -957,6 +960,13 @@ TSharedRef<SWidget> UWidget::TakeWidget()
 
 TSharedRef<SWidget> UWidget::TakeWidget_Private(ConstructMethodType ConstructMethod)
 {
+#if WIDGET_INCLUDE_RELFECTION_METADATA
+	UObject* SourceAsset = GetSourceAssetOrClass();
+	UClass* WidgetClass = GetClass();
+	LLM_SCOPE_DYNAMIC_STAT_OBJECTPATH(SourceAsset->GetPackage(), ELLMTagSet::Assets);
+	LLM_SCOPE_DYNAMIC_STAT_OBJECTPATH(WidgetClass, ELLMTagSet::AssetClasses);
+	UE_TRACE_METADATA_SCOPE_ASSET(SourceAsset, WidgetClass);
+#endif
 	bool bNewlyCreated = false;
 	TSharedPtr<SWidget> PublicWidget;
 
@@ -1025,7 +1035,7 @@ TSharedRef<SWidget> UWidget::TakeWidget_Private(ConstructMethodType ConstructMet
 
 #if WIDGET_INCLUDE_RELFECTION_METADATA
 		// We only need to do this once, when the slate widget is created.
-		PublicWidget->AddMetadata<FReflectionMetaData>(MakeShared<FReflectionMetaData>(GetFName(), GetClass(), this, GetSourceAssetOrClass()));
+		PublicWidget->AddMetadata<FReflectionMetaData>(MakeShared<FReflectionMetaData>(GetFName(), WidgetClass, this, SourceAsset));
 #endif
 
 		SynchronizeProperties();
diff --git a/Engine/Source/Runtime/UMG/Private/UserWidget.cpp b/Engine/Source/Runtime/UMG/Private/UserWidget.cpp
--- a/Engine/Source/Runtime/UMG/Private/UserWidget.cpp
+++ b/Engine/Source/Runtime/UMG/Private/UserWidget.cpp
@@ -36,6 +36,7 @@
 #include "Editor/WidgetCompilerLog.h"
 #include "GameFramework/InputSettings.h"
 #include "Engine/InputDelegateBinding.h"
+#include "ProfilingDebugging/AssetMetadataTrace.h"
 
 #include UE_INLINE_GENERATED_CPP_BY_NAME(UserWidget)
 
@@ -2259,6 +2260,9 @@ UUserWidget* UUserWidget::CreateInstanceInternal(UObject* Outer, TSubclassOf<UUs
 		FMessageLog("PIE").Error(FText::Format(LOCTEXT("OuterNull", "Unable to create the widget {0}, no outer provided."), FText::FromName(UserWidgetClass->GetFName())));
 		return nullptr;
 	}
+	LLM_SCOPE_DYNAMIC_STAT_OBJECTPATH(Outer->GetPackage(), ELLMTagSet::Assets);
+	LLM_SCOPE_DYNAMIC_STAT_OBJECTPATH(UserWidgetClass, ELLMTagSet::AssetClasses);
+	UE_TRACE_METADATA_SCOPE_ASSET_FNAME(InstanceName, UserWidgetClass->GetFName(), Outer->GetPackage()->GetFName());
 
 	UUserWidget* NewWidget = NewObject<UUserWidget>(Outer, UserWidgetClass, InstanceName, RF_Transactional);
 	
diff --git a/Engine/Source/Runtime/UMG/UMG.Build.cs b/Engine/Source/Runtime/UMG/UMG.Build.cs
--- a/Engine/Source/Runtime/UMG/UMG.Build.cs
+++ b/Engine/Source/Runtime/UMG/UMG.Build.cs
@@ -9,6 +9,7 @@ public class UMG : ModuleRules
 		PrivateDependencyModuleNames.AddRange(
 			new string[] {
 				"Core",
+				"TraceLog",
 				"CoreUObject",
 				"DeveloperSettings",
 				"Engine",
