[MemoryTrace] Missed asset metadata scope for UI/WidgetProxy

#tests PIE
#rnx

diff --git a/Engine/Source/Runtime/SlateCore/Private/FastUpdate/WidgetProxy.cpp b/Engine/Source/Runtime/SlateCore/Private/FastUpdate/WidgetProxy.cpp
--- a/Engine/Source/Runtime/SlateCore/Private/FastUpdate/WidgetProxy.cpp
+++ b/Engine/Source/Runtime/SlateCore/Private/FastUpdate/WidgetProxy.cpp
@@ -6,6 +6,9 @@
 #include "Widgets/SWindow.h"
 #include "FastUpdate/SlateInvalidationWidgetHeap.h"
 #include "FastUpdate/SlateInvalidationWidgetList.h"
+#include "HAL/LowLevelMemStats.h"
+#include "UObject/Package.h"
+#include "ProfilingDebugging/AssetMetadataTrace.h"
 #include "ProfilingDebugging/CsvProfiler.h"
 #include "Types/ReflectionMetadata.h"
 #include "Input/HittestGrid.h"
@@ -110,6 +113,27 @@ void FWidgetProxy::ProcessLayoutInvalidation(FSlateInvalidationWidgetPostHeap& U
 {
 	SWidget* WidgetPtr = GetWidget();
 	SCOPE_CYCLE_SWIDGET(WidgetPtr);
+#if UE_TRACE_ASSET_METADATA_ENABLED
+	FName AssetName = NAME_None;
+	FName ClassName = NAME_None;
+	FName PackageName = NAME_None;
+	if (UE_TRACE_CHANNELEXPR_IS_ENABLED(AssetMetadataChannel))
+	{
+		TSharedPtr<FReflectionMetaData> AssetMetaData = FReflectionMetaData::GetWidgetOrParentMetaData(WidgetPtr);
+		if (AssetMetaData.IsValid())
+		{
+			if (const UObject* AssetPtr = AssetMetaData->Asset.Get())
+			{
+				AssetName = AssetMetaData->Name;
+				ClassName = AssetMetaData->Class.Get()->GetFName();
+				PackageName = AssetPtr->GetPackage()->GetFName();
+			}
+		}
+	}
+	LLM_SCOPE_DYNAMIC_STAT_OBJECTPATH_FNAME(PackageName, ELLMTagSet::Assets);
+	LLM_SCOPE_DYNAMIC_STAT_OBJECTPATH_FNAME(ClassName, ELLMTagSet::AssetClasses);
+	UE_TRACE_METADATA_SCOPE_ASSET_FNAME(AssetName, ClassName, PackageName);
+#endif
 	// When layout changes compute a new desired size for this widget
 	const FVector2f CurrentDesiredSize = WidgetPtr->GetDesiredSize();
 	FVector2f NewDesiredSize = FVector2f::ZeroVector;
