TAnimatedAttribute: Retrieving value based on getter fix

#rnx

diff --git a/Engine/Source/Runtime/SlateCore/Public/Animation/AnimatedAttribute.h b/Engine/Source/Runtime/SlateCore/Public/Animation/AnimatedAttribute.h
--- a/Engine/Source/Runtime/SlateCore/Public/Animation/AnimatedAttribute.h
+++ b/Engine/Source/Runtime/SlateCore/Public/Animation/AnimatedAttribute.h
@@ -73,8 +73,8 @@ public:
 	[[nodiscard]] static TSharedRef< TAnimatedAttribute > CreateWithGetter( const InterpolatorSettings& InSettings, const FGetter& InGetter )
 	{
 		TSharedRef< TAnimatedAttribute< NumericType > > Attribute = MakeShared< TAnimatedAttribute< NumericType > >(FPrivateToken{}, InSettings);
+		Attribute->Set(InGetter.Execute());
 		Attribute->Getter = InGetter;
-		Attribute->Set(Attribute->Getter.Execute());
 		Attribute->Register();
 		return Attribute;
 	}
@@ -89,8 +89,8 @@ public:
 	[[nodiscard]] static TSharedRef< TAnimatedAttribute > CreateWithGetter( const InterpolatorSettings& InSettings, FGetter&& InGetter )
 	{
 		TSharedRef< TAnimatedAttribute< NumericType > > Attribute = MakeShared< TAnimatedAttribute< NumericType > >(FPrivateToken{}, InSettings);
+		Attribute->Set(InGetter.Execute());
 		Attribute->Getter = MoveTemp(InGetter);
-		Attribute->Set(Attribute->Getter.Execute());
 		Attribute->Register();
 		return Attribute;
 	}
@@ -152,12 +152,6 @@ public:
 	 */
 	const NumericType& Get() const
 	{
-		// first see if the interpolator has a value
-		if(Interpolator->IsSet())
-		{
-			return Interpolator->Get();
-		}
-		
 		// If we have a getter delegate, then we'll call that to generate the value
 		if( Getter.IsBound() )
 		{
@@ -167,8 +161,20 @@ public:
 			// NOTE: We purposely overwrite our value copy here so that we can return the value by address in
 			// the most common case, which is an attribute that doesn't have a delegate bound to it at all.
 			DesiredValue = Getter.Execute();
+
+			// Update the interpolator
+			if(DesiredValue.IsSet())
+			{
+				Interpolator->SetValue(DesiredValue.GetValue());
+			}
 		}
 
+		// first see if the interpolator has a value
+		if(Interpolator->IsSet())
+		{
+			return Interpolator->Get();
+		}
+		
 		if(DesiredValue.IsSet())
 		{
 			// Return the stored value
@@ -212,6 +218,22 @@ public:
 		Interpolator->SetTolerance(Tolerance);
 	}
 
+	/**
+	 * Enables (or disables) the interpolator and returns values in interpolated or immediate mode
+	 */
+	void EnableInterpolation(bool bEnabled = true)
+	{
+		Interpolator->SetEnabled(bEnabled);
+	}
+
+	/**
+	 * Disables the interpolator and returns values in immediate mode
+	 */
+	void DisableInterpolation()
+	{
+		EnableInterpolation(false);
+	}
+
 	/**
 	 * Checks to see if this attribute has a 'getter' function bound
 	 *
diff --git a/Engine/Source/Runtime/SlateCore/Public/Animation/AttributeInterpolator.h b/Engine/Source/Runtime/SlateCore/Public/Animation/AttributeInterpolator.h
--- a/Engine/Source/Runtime/SlateCore/Public/Animation/AttributeInterpolator.h
+++ b/Engine/Source/Runtime/SlateCore/Public/Animation/AttributeInterpolator.h
@@ -27,7 +27,8 @@ public:
 	 * Default constructor
 	 */
 	TAttributeInterpolator()
-	: Tolerance(0.001f)
+	: bEnabled(true)
+	, Tolerance(0.001f)
 	{
 	}
 
@@ -90,6 +91,32 @@ public:
 		}
 	}
 
+	/*
+	 * Returns true if the interpolator is enabled
+	 */
+	bool IsEnabled() const
+	{
+		return bEnabled;
+	}
+
+	/*
+	 * Enables (or disables) the interpolator
+	 */
+	void SetEnabled(bool InEnabled = true)
+	{
+		if(bEnabled != InEnabled)
+		{
+			bEnabled = InEnabled;
+
+			// if we disabled or re-enabled the interpolator set
+			// the interpolated values based on the desired value
+			if(DesiredValue.IsSet())
+			{
+				SetValueAndStop(DesiredValue.GetValue());
+			}
+		}
+	}
+
 	/**
  	 * Set the internal state of the interpolator
  	 * 
@@ -131,12 +158,15 @@ public:
 	/**
 	 * Returns the interpolated value
 	 *
-	 * @param InValue The value to interpolate against
 	 * @return The interpolated value
 	 */
 	const NumericType& Get() const
 	{
 		check(IsSet());
+		if(!IsEnabled())
+		{
+			return DesiredValue.GetValue();
+		}
 		if(InterpolatedValue.IsSet())
 		{
 			return InterpolatedValue.GetValue();
@@ -261,6 +291,7 @@ protected:
 		return Start();
 	}
 
+	bool bEnabled;
 	float Tolerance;
 	mutable TOptional<float> Delay;
 	mutable TOptional<float> OverallDeltaTime;
