PR #10934: Fix RoundedBox SlateBrushes defaulting with an outline radius of (0, 0, 0, 1)

diff --git a/Engine/Source/Runtime/SlateCore/Public/Styling/SlateBrush.h b/Engine/Source/Runtime/SlateCore/Public/Styling/SlateBrush.h
--- a/Engine/Source/Runtime/SlateCore/Public/Styling/SlateBrush.h
+++ b/Engine/Source/Runtime/SlateCore/Public/Styling/SlateBrush.h
@@ -135,7 +135,7 @@ struct FSlateBrushOutlineSettings
 	GENERATED_USTRUCT_BODY()
 
 	FSlateBrushOutlineSettings()
-		: CornerRadii(0.0)
+		: CornerRadii(FVector4(0.0, 0.0, 0.0, 0.0))
 		, Color(FLinearColor::Transparent)
 		, Width(0.0)
 		, RoundingType(ESlateBrushRoundingType::HalfHeightRadius)
@@ -159,7 +159,7 @@ struct FSlateBrushOutlineSettings
 	{}
 
 	FSlateBrushOutlineSettings(const FSlateColor& InColor, float InWidth)
-		: CornerRadii(0.0)
+		: CornerRadii(FVector4(0.0, 0.0, 0.0, 0.0))
 		, Color(InColor)
 		, Width(InWidth)
 		, RoundingType(ESlateBrushRoundingType::HalfHeightRadius)
