Fix Slate Insights not working due to untracked widget names on Trace Start / Stop.

diff --git a/Engine/Source/Runtime/Slate/Private/Framework/Application/SlateApplication.cpp b/Engine/Source/Runtime/Slate/Private/Framework/Application/SlateApplication.cpp
--- a/Engine/Source/Runtime/Slate/Private/Framework/Application/SlateApplication.cpp
+++ b/Engine/Source/Runtime/Slate/Private/Framework/Application/SlateApplication.cpp
@@ -732,6 +732,8 @@ TSharedRef<FSlateApplication> FSlateApplication::Create(const TSharedRef<class G
 	CurrentApplication = MakeShareable( new FSlateApplication() );
 	CurrentBaseApplication = CurrentApplication;
 
+	UE_TRACE_SLATE_APPLICATION_REGISTER_TRACE_EVENTS(*CurrentApplication);
+
 	PlatformApplication->SetMessageHandler( CurrentApplication.ToSharedRef() );
 #if WITH_ACCESSIBILITY
 	PlatformApplication->SetAccessibleMessageHandler(CurrentApplication->GetAccessibleMessageHandler());
diff --git a/Engine/Source/Runtime/SlateCore/Private/Trace/SlateTrace.cpp b/Engine/Source/Runtime/SlateCore/Private/Trace/SlateTrace.cpp
--- a/Engine/Source/Runtime/SlateCore/Private/Trace/SlateTrace.cpp
+++ b/Engine/Source/Runtime/SlateCore/Private/Trace/SlateTrace.cpp
@@ -10,6 +10,7 @@
 #include "HAL/ThreadHeartBeat.h"
 #include "HAL/PlatformStackWalk.h"
 #include "HAL/PlatformProcess.h"
+#include "ProfilingDebugging/TraceAuxiliary.h"
 #include "Trace/Trace.inl"
 #include "Types/ReflectionMetadata.h"
 #include "Widgets/SWidget.h"
@@ -227,6 +228,7 @@ FSlateTrace::FScopedWidgetPaintTrace::FScopedWidgetPaintTrace(const SWidget* InW
 	{
 		StartCycle = FPlatformTime::Cycles64();
 		SlateTraceDetail::AddStartWidgetPaintSteps(Widget, StartCycle);
+		FSlateTrace::ConditionallyUpdateWidgetInfo(Widget);
 	}
 
 	++SlateTraceDetail::GScopedPaintCount;
@@ -263,6 +265,7 @@ FSlateTrace::FScopedWidgetUpdateTrace::FScopedWidgetUpdateTrace(const SWidget* I
 		{
 			StartCycle = FPlatformTime::Cycles64();
 			UpdateFlags = InWidget->UpdateFlags & (EWidgetUpdateFlags::NeedsTick | EWidgetUpdateFlags::NeedsActiveTimerUpdate);
+			FSlateTrace::ConditionallyUpdateWidgetInfo(Widget);
 		}
 	}
 	++SlateTraceDetail::GScopedUpdateCount;
@@ -295,6 +298,11 @@ void FSlateTrace::ApplicationTickAndDrawWidgets(float DeltaTime)
 				const FString Path = FReflectionMetaData::GetWidgetPath(SharedWidget.Get());
 				const FString DebugInfo = FReflectionMetaData::GetWidgetDebugInfo(SharedWidget.Get());
 
+				if (FTraceAuxiliary::IsConnected())
+				{
+					SharedWidget->Debug_SetWidgetInfoTraced(FSlateTrace::TraceCounter);
+				}
+
 				UE_TRACE_LOG(SlateTrace, WidgetInfo, SlateChannel)
 					<< WidgetInfo.WidgetId(WidgetId)
 					<< WidgetInfo.Path(*Path)
@@ -330,6 +338,13 @@ void FSlateTrace::ApplicationTickAndDrawWidgets(float DeltaTime)
 	}
 }
 
+void FSlateTrace::ApplicationRegisterTraceEvents(FSlateApplicationBase& /*SlateApplication*/)
+{
+	// Registering directly on the slate application would mean slate app headers need to know trace types.
+	// Instead register a static here that gets the application on a trace start.
+	FTraceAuxiliary::OnTraceStarted.AddStatic(&FSlateTrace::HandleOnTraceStarted);
+}
+
 void FSlateTrace::OutputWidgetUpdate(const SWidget* Widget, uint64 StartCycle, uint64 EndCycle, EWidgetUpdateFlags UpdateFlags, uint32 AffectedCount)
 {
 	if (EnumHasAnyFlags(UpdateFlags, EWidgetUpdateFlags::NeedsActiveTimerUpdate))
@@ -362,6 +377,8 @@ void FSlateTrace::WidgetInvalidated(const SWidget* Widget, const SWidget* Invest
 {
 	if (UE_TRACE_CHANNELEXPR_IS_ENABLED(SlateChannel) && Reason != EInvalidateWidgetReason::None)
 	{
+		FSlateTrace::ConditionallyUpdateWidgetInfo(Widget);
+
 		++SlateTraceDetail::GFrameInvalidateCount;
 
 		static_assert(sizeof(EInvalidateWidgetReason) == sizeof(uint8), "EInvalidateWidgetReason is not a uint8");
@@ -454,6 +471,8 @@ void FSlateTrace::RootInvalidated(const SWidget* Widget, const SWidget* Investig
 {
 	if (UE_TRACE_CHANNELEXPR_IS_ENABLED(SlateChannel))
 	{
+		FSlateTrace::ConditionallyUpdateWidgetInfo(Widget);
+
 		++SlateTraceDetail::GFrameInvalidateCount;
 
 		const uint64 WidgetId = SlateTraceDetail::GetWidgetId(Widget);
@@ -470,6 +489,8 @@ void FSlateTrace::RootChildOrderInvalidated(const SWidget* Widget, const SWidget
 {
 	if (UE_TRACE_CHANNELEXPR_IS_ENABLED(SlateChannel))
 	{
+		FSlateTrace::ConditionallyUpdateWidgetInfo(Widget);
+
 		++SlateTraceDetail::GFrameInvalidateCount;
 
 		const uint64 WidgetId = SlateTraceDetail::GetWidgetId(Widget);
@@ -521,8 +542,31 @@ void FSlateTrace::RemoveWidget(const SWidget* Widget)
 	--SlateTraceDetail::GWidgetCount;
 }
 
-void FSlateTrace::OutputWidgetPaint(const SWidget* Widget, uint64 StartCycle, uint64 EndCycle, uint32 PaintCount)
+void FSlateTrace::ConditionallyUpdateWidgetInfo(const SWidget* Widget)
 {
+	if (UE_TRACE_CHANNELEXPR_IS_ENABLED(SlateChannel) && FTraceAuxiliary::IsConnected())
+	{
+		if (Widget && Widget->Debug_GetWidgetInfoTraced() < FSlateTrace::TraceCounter)
+		{
+			const uint64 WidgetId = SlateTraceDetail::GetWidgetId(Widget);
+			const FString Path = FReflectionMetaData::GetWidgetPath(Widget);
+			const FString DebugInfo = FReflectionMetaData::GetWidgetDebugInfo(Widget);
+
+			UE_TRACE_LOG(SlateTrace, WidgetInfo, SlateChannel)
+				<< WidgetInfo.WidgetId(WidgetId)
+				<< WidgetInfo.Path(*Path)
+				<< WidgetInfo.DebugInfo(*DebugInfo);
+
+			Widget->Debug_SetWidgetInfoTraced(FSlateTrace::TraceCounter);
+		}
+	}
 }
 
+void FSlateTrace::HandleOnTraceStarted(FTraceAuxiliary::EConnectionType TraceType, const FString& TraceDestination)
+{
+	FSlateTrace::TraceCounter++;
+}
+
+uint8 FSlateTrace::TraceCounter = 0;
+
  #endif // UE_SLATE_TRACE_ENABLED
diff --git a/Engine/Source/Runtime/SlateCore/Public/Trace/SlateTrace.h b/Engine/Source/Runtime/SlateCore/Public/Trace/SlateTrace.h
--- a/Engine/Source/Runtime/SlateCore/Public/Trace/SlateTrace.h
+++ b/Engine/Source/Runtime/SlateCore/Public/Trace/SlateTrace.h
@@ -3,6 +3,7 @@
 #pragma once
 
 #include "CoreMinimal.h"
+#include "ProfilingDebugging/TraceAuxiliary.h"
 #include "Trace/Config.h"
 #include "Trace/Trace.h"
 
@@ -58,17 +59,34 @@ public:
 	};
 
 	SLATECORE_API static void ApplicationTickAndDrawWidgets(float DeltaTime);
+	SLATECORE_API static void ApplicationRegisterTraceEvents(FSlateApplicationBase& SlateApplication);
 	SLATECORE_API static void WidgetInvalidated(const SWidget* Widget, const SWidget* Investigator, EInvalidateWidgetReason Reason);
 	SLATECORE_API static void RootInvalidated(const SWidget* Widget, const SWidget* Investigator);
 	SLATECORE_API static void RootChildOrderInvalidated(const SWidget* Widget, const SWidget* Investigator);
 
+	/** Called in widget ctor */
 	static void AddWidget(const SWidget* Widget);
+
+	/** Called as part of TSlateDecl widget creation (SNew, SAssignNew, SArgumentNew) */
 	static void UpdateWidgetInfo(const SWidget* Widget);
+
+	/** Called in widget dtor */
 	static void RemoveWidget(const SWidget* Widget);
 
 private:
-	static void OutputWidgetPaint(const SWidget* Widget, uint64 StartCycle, uint64 EndCycle, uint32 PaintCount);
+
+	/** Called on all widget updates except for Add & Remove to ensure that trace WidgetInfo is populated */
+	static void ConditionallyUpdateWidgetInfo(const SWidget* Widget);
+
+	/** Called on all widget updates except for Add & Remove to ensure that trace WidgetInfo is populated */
+	static void HandleOnTraceStarted(FTraceAuxiliary::EConnectionType TraceType, const FString& TraceDestination);
+
+	/** Used to log a particular widget update, tick timer, repaint, volatile prepass, etc. See: 'EWidgetUpdateFlags' for more info */
 	static void OutputWidgetUpdate(const SWidget* Widget, uint64 StartCycle, uint64 EndCycle, EWidgetUpdateFlags UpdateFlags, uint32 AffectedCount);
+
+private:
+	/** Counter for number of traces started, used to conditionally re-send widget infos as needed */
+	static uint8 TraceCounter;
 };
 
 #define UE_TRACE_SLATE_BOOKMARK(Format, ...) \
@@ -78,6 +96,9 @@ private:
 	}
 #define UE_TRACE_SLATE_APPLICATION_TICK_AND_DRAW_WIDGETS(DeltaTime) \
 	FSlateTrace::ApplicationTickAndDrawWidgets(DeltaTime);
+
+#define UE_TRACE_SLATE_APPLICATION_REGISTER_TRACE_EVENTS(SlateApplication) \
+	FSlateTrace::ApplicationRegisterTraceEvents(SlateApplication);
 	
 #define UE_TRACE_SLATE_WIDGET_ADDED(Widget) \
 	FSlateTrace::AddWidget(Widget);
@@ -107,6 +128,7 @@ private:
 
 #define UE_TRACE_SLATE_BOOKMARK(...)
 #define UE_TRACE_SLATE_APPLICATION_TICK_AND_DRAW_WIDGETS(DeltaTime)
+#define UE_TRACE_SLATE_APPLICATION_REGISTER_TRACE_EVENTS(SlateApplication)
 #define UE_TRACE_SLATE_WIDGET_ADDED(Widget)
 #define UE_TRACE_SLATE_WIDGET_DEBUG_INFO(Widget)
 #define UE_TRACE_SLATE_WIDGET_REMOVED(Widget)
diff --git a/Engine/Source/Runtime/SlateCore/Public/Widgets/SWidget.h b/Engine/Source/Runtime/SlateCore/Public/Widgets/SWidget.h
--- a/Engine/Source/Runtime/SlateCore/Public/Widgets/SWidget.h
+++ b/Engine/Source/Runtime/SlateCore/Public/Widgets/SWidget.h
@@ -796,12 +796,29 @@ private:
 	SLATECORE_API void UpdateWidgetProxy(int32 NewLayerId, FSlateCachedElementsHandle& CacheHandle);
 
 public:
+#if UE_SLATE_TRACE_ENABLED
+	uint8 Debug_GetWidgetInfoTraced() const 
+	{ 
+		return Debug_LastTraceInfoSent; 
+	}
+
+	void Debug_SetWidgetInfoTraced(uint8 InDebug_LastTraceInfoSent) const 
+	{ 
+		Debug_LastTraceInfoSent = InDebug_LastTraceInfoSent; 
+	}
+#endif // UE_SLATE_TRACE_ENABLED
 
 #if WITH_SLATE_DEBUGGING
-	uint32 Debug_GetLastPaintFrame() const { return LastPaintFrame; }
+	uint32 Debug_GetLastPaintFrame() const 
+	{ 
+		return LastPaintFrame; 
+	}
 private:
-	void Debug_UpdateLastPaintFrame() { LastPaintFrame = GFrameNumber; }
-#endif
+	void Debug_UpdateLastPaintFrame() 
+	{ 
+		LastPaintFrame = GFrameNumber; 
+	}
+#endif // WITH_SLATE_DEBUGGING
 
 public:
 
@@ -1897,12 +1914,21 @@ private:
 	FName CreatedInLocation;
 #endif
 
+#if UE_SLATE_TRACE_ENABLED
+	/**
+	 * If the widget info is sent when a trace is not active, it will be ignored.
+	 * In this scenario, if the info is not re-sent the trace will not function correctly.
+	 * Track which traces we have sent info for so we can re-send when needed.
+	 */
+	mutable uint8 Debug_LastTraceInfoSent = 0;
+#endif // UE_SLATE_TRACE_ENABLED
+
 #if WITH_SLATE_DEBUGGING
 	/** The last time this widget got painted. */
 	uint32 LastPaintFrame = 0;
 	/** Flag to help detect when we access an invalid Widget. */
 	uint8 Debug_DestroyedTag = 0xDC;
-#endif
+#endif // WITH_SLATE_DEBUGGING
 
 #if !(UE_BUILD_SHIPPING || UE_BUILD_TEST)
 	UE_DEPRECATED(4.27, "Access to SWidget::Cursor is deprecated and will not function. Call SetCursor/GetCursor instead")
