Replaced UE_NODISCARD with [[nodiscard]].

diff --git a/Engine/Source/Runtime/SlateCore/Private/FastUpdate/SlateInvalidationWidgetHeap.h b/Engine/Source/Runtime/SlateCore/Private/FastUpdate/SlateInvalidationWidgetHeap.h
--- a/Engine/Source/Runtime/SlateCore/Private/FastUpdate/SlateInvalidationWidgetHeap.h
+++ b/Engine/Source/Runtime/SlateCore/Private/FastUpdate/SlateInvalidationWidgetHeap.h
@@ -57,7 +57,7 @@ public:
 	}
 
 	/** Returns and removes the biggest WidgetIndex from the list. */
-	UE_NODISCARD FSlateInvalidationWidgetIndex HeapPop()
+	[[nodiscard]] FSlateInvalidationWidgetIndex HeapPop()
 	{
 		check(Heap.Num() > 0);
 		FSlateInvalidationWidgetIndex Result = Heap.HeapTop().GetWidgetIndex();
@@ -76,14 +76,14 @@ public:
 	}
 	
 	/** Returns the biggest WidgetIndex from the list. */
-	UE_NODISCARD inline FSlateInvalidationWidgetIndex HeapPeek() const
+	[[nodiscard]] inline FSlateInvalidationWidgetIndex HeapPeek() const
 	{
 		check(Heap.Num() > 0);
 		return Heap.HeapTop().GetWidgetIndex();
 	}
 	
 	/** Returns the biggest WidgetIndex from the list. */
-	UE_NODISCARD inline const FElement& HeapPeekElement() const
+	[[nodiscard]] inline const FElement& HeapPeekElement() const
 	{
 		check(Heap.Num() > 0);
 		return Heap.HeapTop();
@@ -122,13 +122,13 @@ public:
 	}
 
 	/** Returns the number of elements in the list. */
-	UE_NODISCARD inline int32 Num() const
+	[[nodiscard]] inline int32 Num() const
 	{
 		return Heap.Num();
 	}
 
 	/** Does it contains the widget index. */
-	UE_NODISCARD bool Contains_Debug(const FSlateInvalidationWidgetIndex WidgetIndex) const
+	[[nodiscard]] bool Contains_Debug(const FSlateInvalidationWidgetIndex WidgetIndex) const
 	{
 		return Heap.ContainsByPredicate([WidgetIndex](const FElement& Element)
 			{
@@ -137,13 +137,13 @@ public:
 	}
 
 	/** @retuns true if the heap is heapified. */
-	UE_NODISCARD inline bool IsValidHeap_Debug()
+	[[nodiscard]] inline bool IsValidHeap_Debug()
 	{
 		return Algo::IsHeap(Heap, SortPredicate());
 	}
 
 	/** Returns the raw list. */
-	UE_NODISCARD inline const FElementContainer& GetRaw() const
+	[[nodiscard]] inline const FElementContainer& GetRaw() const
 	{
 		return Heap;
 	}
@@ -240,7 +240,7 @@ public:
 	}
 
 	/** Returns and removes the biggest WidgetIndex from the list. */
-	UE_NODISCARD FElement HeapPop()
+	[[nodiscard]] FElement HeapPop()
 	{
 		check(bIsHeap == true);
 		FElement Result = Heap.HeapTop();
@@ -292,19 +292,19 @@ public:
 	}
 
 	/** Returns the number of elements in the list. */
-	UE_NODISCARD inline int32 Num() const
+	[[nodiscard]] inline int32 Num() const
 	{
 		return Heap.Num();
 	}
 
 	/** Returns the number of elements in the list. */
-	UE_NODISCARD inline bool IsHeap() const
+	[[nodiscard]] inline bool IsHeap() const
 	{
 		return bIsHeap;
 	}
 
 	/** Does it contains the widget index. */
-	UE_NODISCARD bool Contains_Debug(const FSlateInvalidationWidgetIndex WidgetIndex) const
+	[[nodiscard]] bool Contains_Debug(const FSlateInvalidationWidgetIndex WidgetIndex) const
 	{
 		return Heap.ContainsByPredicate([WidgetIndex](const FElement& Element)
 			{
@@ -313,13 +313,13 @@ public:
 	}
 
 	/** @retuns true if the heap is heapified. */
-	UE_NODISCARD inline bool IsValidHeap_Debug()
+	[[nodiscard]] inline bool IsValidHeap_Debug()
 	{
 		return Algo::IsHeap(Heap, SortPredicate());
 	}
 
 	/** Returns the raw list. */
-	UE_NODISCARD inline const FElementContainer& GetRaw() const
+	[[nodiscard]] inline const FElementContainer& GetRaw() const
 	{
 		return Heap;
 	}
@@ -431,7 +431,7 @@ public:
 	}
 
 	/** Returns and removes the biggest WidgetIndex from the list. */
-	UE_NODISCARD FElement HeapPop()
+	[[nodiscard]] FElement HeapPop()
 	{
 		check(bIsHeap == true);
 		FElement Result = Heap.HeapTop();
@@ -483,19 +483,19 @@ public:
 	}
 
 	/** Returns the number of elements in the list. */
-	UE_NODISCARD inline int32 Num() const
+	[[nodiscard]] inline int32 Num() const
 	{
 		return Heap.Num();
 	}
 
 	/** Returns the number of elements in the list. */
-	UE_NODISCARD inline bool IsHeap() const
+	[[nodiscard]] inline bool IsHeap() const
 	{
 		return bIsHeap;
 	}
 
 	/** Does it contains the widget index. */
-	UE_NODISCARD bool Contains_Debug(const FSlateInvalidationWidgetIndex WidgetIndex) const
+	[[nodiscard]] bool Contains_Debug(const FSlateInvalidationWidgetIndex WidgetIndex) const
 	{
 		return Heap.ContainsByPredicate([WidgetIndex](const FElement& Element)
 			{
@@ -504,13 +504,13 @@ public:
 	}
 
 	/** @retuns true if the heap is heapified. */
-	UE_NODISCARD inline bool IsValidHeap_Debug()
+	[[nodiscard]] inline bool IsValidHeap_Debug()
 	{
 		return Algo::IsHeap(Heap, SortPredicate());
 	}
 
 	/** Returns the raw list. */
-	UE_NODISCARD inline const FElementContainer& GetRaw() const
+	[[nodiscard]] inline const FElementContainer& GetRaw() const
 	{
 		return Heap;
 	}
diff --git a/Engine/Source/Runtime/SlateCore/Private/FastUpdate/SlateInvalidationWidgetList.h b/Engine/Source/Runtime/SlateCore/Private/FastUpdate/SlateInvalidationWidgetList.h
--- a/Engine/Source/Runtime/SlateCore/Private/FastUpdate/SlateInvalidationWidgetList.h
+++ b/Engine/Source/Runtime/SlateCore/Private/FastUpdate/SlateInvalidationWidgetList.h
@@ -47,18 +47,18 @@ public:
 			check(OrderMin <= OrderMax);
 #endif
 		}
-		UE_NODISCARD bool Include(FSlateInvalidationWidgetSortOrder Other) const
+		[[nodiscard]] bool Include(FSlateInvalidationWidgetSortOrder Other) const
 		{
 			return OrderMin <= Other && Other <= OrderMax;
 		}
-		UE_NODISCARD bool IsValid() const { return InclusiveMin != FSlateInvalidationWidgetIndex::Invalid; }
+		[[nodiscard]] bool IsValid() const { return InclusiveMin != FSlateInvalidationWidgetIndex::Invalid; }
 
 		FSlateInvalidationWidgetIndex GetInclusiveMinWidgetIndex() const { return InclusiveMin; }
 		FSlateInvalidationWidgetIndex GetInclusiveMaxWidgetIndex() const { return InclusiveMax; }
 		FSlateInvalidationWidgetSortOrder GetInclusiveMinWidgetSortOrder() const { return OrderMin; }
 		FSlateInvalidationWidgetSortOrder GetInclusiveMaxWidgetSortOrder() const { return OrderMax; }
 
-		UE_NODISCARD bool operator==(const FIndexRange& Other) const { return Other.InclusiveMin == InclusiveMin && Other.InclusiveMax == InclusiveMax; }
+		[[nodiscard]] bool operator==(const FIndexRange& Other) const { return Other.InclusiveMin == InclusiveMin && Other.InclusiveMax == InclusiveMax; }
 	};
 
 public:
@@ -105,8 +105,8 @@ public:
 		struct FReIndexOperation
 		{
 			FReIndexOperation(const FIndexRange& InRange, FSlateInvalidationWidgetIndex InReIndexTarget) : Range(InRange), ReIndexTarget(InReIndexTarget) {}
-			UE_NODISCARD const FIndexRange& GetRange() const { return Range; }
-			UE_NODISCARD FSlateInvalidationWidgetIndex ReIndex(FSlateInvalidationWidgetIndex Index) const;
+			[[nodiscard]] const FIndexRange& GetRange() const { return Range; }
+			[[nodiscard]] FSlateInvalidationWidgetIndex ReIndex(FSlateInvalidationWidgetIndex Index) const;
 		private:
 			const FIndexRange& Range;
 			FSlateInvalidationWidgetIndex ReIndexTarget = FSlateInvalidationWidgetIndex::Invalid;
@@ -114,7 +114,7 @@ public:
 		struct FReSortOperation
 		{
 			FReSortOperation(const FIndexRange& InRange) : Range(InRange) {}
-			UE_NODISCARD const FIndexRange& GetRange() const { return Range; }
+			[[nodiscard]] const FIndexRange& GetRange() const { return Range; }
 		private:
 			const FIndexRange& Range;
 		};
@@ -316,7 +316,7 @@ public:
 		void AdvanceToNextParent();
 
 		/** Is the iterator pointing to a valid widget index. */
-		UE_NODISCARD bool IsValid() const { return CurrentWidgetIndex != FSlateInvalidationWidgetIndex::Invalid; }
+		[[nodiscard]] bool IsValid() const { return CurrentWidgetIndex != FSlateInvalidationWidgetIndex::Invalid; }
 
 	private:
 		void AdvanceArrayIndex(int32 ArrayIndex);
@@ -348,7 +348,7 @@ public:
 		void Advance();
 
 		/** Is the iterator pointing to a valid widget index. */
-		UE_NODISCARD bool IsValid() const { return CurrentWidgetIndex != FSlateInvalidationWidgetIndex::Invalid; }
+		[[nodiscard]] bool IsValid() const { return CurrentWidgetIndex != FSlateInvalidationWidgetIndex::Invalid; }
 
 	private:
 		void Internal_Advance();
@@ -393,13 +393,13 @@ public:
 	}
 
 	/** Returns true if there is not element in the WidgetList. */
-	UE_NODISCARD bool IsEmpty() const
+	[[nodiscard]] bool IsEmpty() const
 	{
 		return FirstArrayIndex == INDEX_NONE || Data[FirstArrayIndex].ElementList.Num() == 0;
 	}
 
 	/** Returns the first index from the WidgetList. */
-	UE_NODISCARD FSlateInvalidationWidgetIndex FirstIndex() const
+	[[nodiscard]] FSlateInvalidationWidgetIndex FirstIndex() const
 	{
 		return FirstArrayIndex == INDEX_NONE
 			? FSlateInvalidationWidgetIndex::Invalid
@@ -407,19 +407,19 @@ public:
 	}
 
 	/** Returns the last index from the WidgetList. */
-	UE_NODISCARD FSlateInvalidationWidgetIndex LastIndex() const
+	[[nodiscard]] FSlateInvalidationWidgetIndex LastIndex() const
 	{
 		return LastArrayIndex == INDEX_NONE ? FSlateInvalidationWidgetIndex::Invalid : FSlateInvalidationWidgetIndex{ (IndexType)LastArrayIndex, (IndexType)(Data[LastArrayIndex].ElementList.Num() - 1) };
 	}
 
 	/** Increment a widget index to the next entry in the WidgetList. */
-	UE_NODISCARD FSlateInvalidationWidgetIndex IncrementIndex(FSlateInvalidationWidgetIndex Index) const;
+	[[nodiscard]] FSlateInvalidationWidgetIndex IncrementIndex(FSlateInvalidationWidgetIndex Index) const;
 
 	/** Decrement a widget index to the next entry in the WidgetList. */
-	UE_NODISCARD FSlateInvalidationWidgetIndex DecrementIndex(FSlateInvalidationWidgetIndex Index) const;
+	[[nodiscard]] FSlateInvalidationWidgetIndex DecrementIndex(FSlateInvalidationWidgetIndex Index) const;
 
 	/** Find the next widget index that share the same parent. Return Invalid if there is no sibling. */
-	UE_NODISCARD FSlateInvalidationWidgetIndex FindNextSibling(FSlateInvalidationWidgetIndex WidgetIndex) const;
+	[[nodiscard]] FSlateInvalidationWidgetIndex FindNextSibling(FSlateInvalidationWidgetIndex WidgetIndex) const;
 
 	/** Empties the WidgetList. */
 	void Empty();
@@ -438,7 +438,7 @@ public:
 	void RemoveWidget(const SWidget& WidgetToRemove);
 
 	/** For testing purposes. Use to test ProcessChildOrderInvalidation */
-	UE_NODISCARD TArray<TSharedPtr<SWidget>> FindChildren(const SWidget& Widget) const;
+	[[nodiscard]] TArray<TSharedPtr<SWidget>> FindChildren(const SWidget& Widget) const;
 
 	/**
 	 * For testing purposes.
@@ -527,9 +527,9 @@ private:
 		int32 OldElementIndexStart = INDEX_NONE;
 	};
 private:
-	UE_NODISCARD IndexType AddArrayNodeIfNeeded(bool bReserveElementList);
-	UE_NODISCARD IndexType InsertArrayNodeIfNeeded(IndexType AfterArrayIndex, bool bReserveElementList);
-	UE_NODISCARD IndexType InsertDataNodeAfter(IndexType AfterIndex, bool bReserveElementList);
+	[[nodiscard]] IndexType AddArrayNodeIfNeeded(bool bReserveElementList);
+	[[nodiscard]] IndexType InsertArrayNodeIfNeeded(IndexType AfterArrayIndex, bool bReserveElementList);
+	[[nodiscard]] IndexType InsertDataNodeAfter(IndexType AfterIndex, bool bReserveElementList);
 	void RemoveDataNode(IndexType Index);
 	void RebuildOrderIndex(IndexType StartFrom);
 	void UpdateParentLeafIndex(const InvalidationWidgetType& InvalidationWidget, FSlateInvalidationWidgetIndex OldIndex, FSlateInvalidationWidgetIndex NewIndex);
diff --git a/Engine/Source/Runtime/SlateCore/Public/Types/Attributes/SlateAttributeBase.inl b/Engine/Source/Runtime/SlateCore/Public/Types/Attributes/SlateAttributeBase.inl
--- a/Engine/Source/Runtime/SlateCore/Public/Types/Attributes/SlateAttributeBase.inl
+++ b/Engine/Source/Runtime/SlateCore/Public/Types/Attributes/SlateAttributeBase.inl
@@ -244,7 +244,7 @@ namespace SlateAttributePrivate
 
 	public:
 		/** @return the SlateAttribute cached value. If the SlateAttribute is bound, the value will be cached at the end of the every frame. */
-		UE_NODISCARD const ObjectType& Get() const
+		[[nodiscard]] const ObjectType& Get() const
 		{
 			return Value;
 		}
@@ -470,7 +470,7 @@ namespace SlateAttributePrivate
 
 	public:
 		/** Build a Attribute from this SlateAttribute. */
-		UE_NODISCARD TAttribute<ObjectType> ToAttribute(const ContainerType& Widget) const
+		[[nodiscard]] TAttribute<ObjectType> ToAttribute(const ContainerType& Widget) const
 		{
 			if (ISlateAttributeGetter* Delegate = ProtectedFindGetter(Widget, InAttributeType))
 			{
@@ -480,14 +480,14 @@ namespace SlateAttributePrivate
 		}
 
 		/** @return True if the SlateAttribute is bound to a getter function. */
-		UE_NODISCARD bool IsBound(const ContainerType& Widget) const
+		[[nodiscard]] bool IsBound(const ContainerType& Widget) const
 		{
 			VerifyOwningWidget(Widget);
 			return ProtectedIsBound(Widget, InAttributeType);
 		}
 
 		/** @return True if they have the same Getter or the same value. */
-		UE_NODISCARD bool IdenticalTo(const ContainerType& Widget, const TSlateAttributeBase& Other) const
+		[[nodiscard]] bool IdenticalTo(const ContainerType& Widget, const TSlateAttributeBase& Other) const
 		{
 			VerifyOwningWidget(Widget);
 			FDelegateHandle ThisDelegateHandle = ProtectedFindGetterHandle(Widget, InAttributeType);
@@ -504,7 +504,7 @@ namespace SlateAttributePrivate
 		}
 
 		/** @return True if they have the same Getter or, if the Attribute is set, the same value. */
-		UE_NODISCARD bool IdenticalTo(const ContainerType& Widget, const TAttribute<ObjectType>& Other) const
+		[[nodiscard]] bool IdenticalTo(const ContainerType& Widget, const TAttribute<ObjectType>& Other) const
 		{
 			VerifyOwningWidget(Widget);
 			FDelegateHandle ThisDelegateHandle = ProtectedFindGetterHandle(Widget, InAttributeType);
diff --git a/Engine/Source/Runtime/SlateCore/Public/Types/Attributes/SlateAttributeMember.inl b/Engine/Source/Runtime/SlateCore/Public/Types/Attributes/SlateAttributeMember.inl
--- a/Engine/Source/Runtime/SlateCore/Public/Types/Attributes/SlateAttributeMember.inl
+++ b/Engine/Source/Runtime/SlateCore/Public/Types/Attributes/SlateAttributeMember.inl
@@ -119,13 +119,13 @@ namespace SlateAttributePrivate
 
 	public:
 		/** @return if the reference is valid. A reference can be invalid if the SWidget is destroyed. */
-		UE_NODISCARD bool IsValid() const
+		[[nodiscard]] bool IsValid() const
 		{
 			return Owner.IsValid();
 		}
 
 		/** @return the SlateAttribute cached value; undefined when IsValid() returns false. */
-		UE_NODISCARD const ObjectType& Get() const
+		[[nodiscard]] const ObjectType& Get() const
 		{
 			if (TSharedPtr<const SWidget> Pin = Owner.Pin())
 			{
@@ -137,7 +137,7 @@ namespace SlateAttributePrivate
 		}
 
 		/** @return the SlateAttribute cached value or the DefaultValue if the reference is invalid. */
-		UE_NODISCARD const ObjectType& Get(const ObjectType& DefaultValue) const
+		[[nodiscard]] const ObjectType& Get(const ObjectType& DefaultValue) const
 		{
 			if (TSharedPtr<const SWidget> Pin = Owner.Pin())
 			{
@@ -159,7 +159,7 @@ namespace SlateAttributePrivate
 		 * Assumes the reference is valid.
 		 * Shorthand for the boilerplace code: MyAttribute.UpdateValueNow(); MyAttribute.Get();
 		 */
-		UE_NODISCARD const ObjectType& UpdateAndGet()
+		[[nodiscard]] const ObjectType& UpdateAndGet()
 		{
 			if (TSharedPtr<const SWidget> Pin = Owner.Pin())
 			{
@@ -173,7 +173,7 @@ namespace SlateAttributePrivate
 		 * Shorthand for the boilerplace code: MyAttribute.UpdateValueNow(); MyAttribute.Get(DefaultValue);
 		 * @return the SlateAttribute cached value or the DefaultValue if the reference is invalid.
 		 */
-		UE_NODISCARD const ObjectType& UpdateAndGet(const ObjectType& DefaultValue)
+		[[nodiscard]] const ObjectType& UpdateAndGet(const ObjectType& DefaultValue)
 		{
 			if (TSharedPtr<const SWidget> Pin = Owner.Pin())
 			{
@@ -184,7 +184,7 @@ namespace SlateAttributePrivate
 		}
 
 		/** Build a Attribute from this SlateAttribute. */
-		UE_NODISCARD TAttribute<ObjectType> ToAttribute() const
+		[[nodiscard]] TAttribute<ObjectType> ToAttribute() const
 		{
 			if (TSharedPtr<const SWidget> Pin = Owner.Pin())
 			{
@@ -194,7 +194,7 @@ namespace SlateAttributePrivate
 		}
 
 		/** @return True if the SlateAttribute is bound to a getter function. */
-		UE_NODISCARD bool IsBound() const
+		[[nodiscard]] bool IsBound() const
 		{
 			if (TSharedPtr<const SWidget> Pin = Owner.Pin())
 			{
@@ -204,7 +204,7 @@ namespace SlateAttributePrivate
 		}
 
 		/** @return True if they have the same Getter or the same value. */
-		UE_NODISCARD bool IdenticalTo(const TSlateMemberAttributeRef& Other) const
+		[[nodiscard]] bool IdenticalTo(const TSlateMemberAttributeRef& Other) const
 		{
 			TSharedPtr<const SWidget> SelfPin = Owner.Pin();
 			TSharedPtr<const SWidget> OtherPin = Other.Owner.Pin();
@@ -216,7 +216,7 @@ namespace SlateAttributePrivate
 		}
 
 		/** @return True if they have the same Getter or, if the Attribute is set, the same value. */
-		UE_NODISCARD bool IdenticalTo(const TAttribute<ObjectType>& Other) const
+		[[nodiscard]] bool IdenticalTo(const TAttribute<ObjectType>& Other) const
 		{
 			if (TSharedPtr<const SWidget> Pin = Owner.Pin())
 			{
diff --git a/Engine/Source/Runtime/SlateCore/Public/Types/PaintArgs.h b/Engine/Source/Runtime/SlateCore/Public/Types/PaintArgs.h
--- a/Engine/Source/Runtime/SlateCore/Public/Types/PaintArgs.h
+++ b/Engine/Source/Runtime/SlateCore/Public/Types/PaintArgs.h
@@ -27,7 +27,7 @@ public:
 	SLATECORE_API FPaintArgs(const SWidget* PaintParent, FHittestGrid& InRootHittestGrid, FHittestGrid& InCurrentHitTestGrid, UE::Slate::FDeprecateVector2DParameter InWindowOffset, double InCurrentTime, float InDeltaTime);
 	SLATECORE_API FPaintArgs(const SWidget* PaintParent, FHittestGrid& InRootHittestGrid, UE::Slate::FDeprecateVector2DParameter InWindowOffset, double InCurrentTime, float InDeltaTime);
 
-	UE_NODISCARD FORCEINLINE_DEBUGGABLE FPaintArgs WithNewParent(const SWidget* PaintParent) const
+	[[nodiscard]] FORCEINLINE_DEBUGGABLE FPaintArgs WithNewParent(const SWidget* PaintParent) const
 	{
 		FPaintArgs Args(*this);
 		Args.PaintParentPtr = PaintParent;
@@ -35,7 +35,7 @@ public:
 		return Args;
 	}
 
-	UE_NODISCARD FORCEINLINE_DEBUGGABLE FPaintArgs WithNewHitTestGrid(FHittestGrid& NewHitTestGrid) const
+	[[nodiscard]] FORCEINLINE_DEBUGGABLE FPaintArgs WithNewHitTestGrid(FHittestGrid& NewHitTestGrid) const
 	{
 		FPaintArgs NewArgs(PaintParentPtr, RootGrid, NewHitTestGrid, WindowOffset, CurrentTime, DeltaTime);
 		return NewArgs;
diff --git a/Engine/Source/Runtime/VirtualFileCache/Private/VirtualFileCacheInternal.h b/Engine/Source/Runtime/VirtualFileCache/Private/VirtualFileCacheInternal.h
--- a/Engine/Source/Runtime/VirtualFileCache/Private/VirtualFileCacheInternal.h
+++ b/Engine/Source/Runtime/VirtualFileCache/Private/VirtualFileCacheInternal.h
@@ -356,9 +356,9 @@ struct FVirtualFileCacheThread : public FRunnable
 	void RequestErase(VFCKey Target);
 	void Shutdown();
 
-	UE_NODISCARD FFileTableReader ReadFileTable() const;
-	UE_NODISCARD FFileTableMutator MutateFileTable();
-	UE_NODISCARD FFileTableWriter ModifyFileTable();
+	[[nodiscard]] FFileTableReader ReadFileTable() const;
+	[[nodiscard]] FFileTableMutator MutateFileTable();
+	[[nodiscard]] FFileTableWriter ModifyFileTable();
 
 public:
 	void DeleteUnexpectedCacheFiles(TSet<FString>& ExpectedFiles);
