PR #10894: Fix slate text layout tab character width calculation.
Used the code from the PR as base, and changed it a little bit to optimize it and for readability.
Fix issue with leading tabs, but still need work to have it working properly for tabulation at any position in the text.

#tests In editor, used test data with various amount of tabs at different position of the text. Multiple tabs at beggining of lines now works correctly, however, tabs at random position inside a text still don't aling the text properly (it won't be aligned to the next spacex4 multplier from the beggining of the line, but relatively to the last characters.

diff --git a/Engine/Source/Runtime/SlateCore/Private/Fonts/SlateTextShaper.cpp b/Engine/Source/Runtime/SlateCore/Private/Fonts/SlateTextShaper.cpp
--- a/Engine/Source/Runtime/SlateCore/Private/Fonts/SlateTextShaper.cpp
+++ b/Engine/Source/Runtime/SlateCore/Private/Fonts/SlateTextShaper.cpp
@@ -955,7 +955,16 @@ bool FSlateTextShaper::InsertSubstituteGlyphs(const TCHAR* InText, const int32 I
 		GetSpecifiedGlyphIndexAndAdvance(TEXT(' '), SpaceGlyphIndex, SpaceXAdvance);
 
 		// We insert a spacer glyph with (up-to) the width of 4 space glyphs in-place of a tab character
-		const int16 NumSpacesToInsert = 4 - (OutGlyphsToRender.Num() % 4);
+		// TODO: Tabulation handling should be refactored to work properly: the tabbing is currently relative to the last characters,
+		// while it should be relative to the beginning of a line). This existing implementation work properly only with leading tabulation.
+		static const int16 TabWidthInSpaces = 4;
+		int16 NumSpacesToIgnore = 0;
+		for (int32 Idx = FMath::Max(0, InCharIndex - TabWidthInSpaces); Idx < InCharIndex; ++Idx)
+		{
+			NumSpacesToIgnore = InText[Idx] == TEXT('\t') ? 0 : (NumSpacesToIgnore + 1) % TabWidthInSpaces;
+		}
+
+		const int16 NumSpacesToInsert = TabWidthInSpaces - NumSpacesToIgnore;
 		if (NumSpacesToInsert > 0)
 		{
 			FShapedGlyphEntry& ShapedGlyphEntry = OutGlyphsToRender.AddDefaulted_GetRef();
