PR #10524: Slate: Fix aliased glyph quad edges

To avoid aliasing on quad edges (for non SDF font only), the glyph quads are now drawn slighty bigger, with a larger UV area (by taking advantage of the padding, if there is any).
The rendered result will look exactly the same most of the time, but as soon as a scale or a non axis aligned rotation will be applyed to the render transform of the text, it will look a lot smoother on the edges.
By default, texts that are scaled up/down may now looks a bit different due to pixel rounding + different size of the quad.
Scrolling texts (without pixel snapping) move more smoothly.

Tweaked functionnal tests used for font to detect smaller local changes that got ignored otherwise, while they were important in the context of font. Maximum local error is now 0.5%.

#tests Ran engine test font functionnal test and updated them. Compared the editor with and without the change. Compared scrolling text (pixel snapping disabled) with and without the change. Tolerance for local change is now

#virtualized

diff --git a/Engine/Source/Runtime/SlateCore/Private/Fonts/FontCache.cpp b/Engine/Source/Runtime/SlateCore/Private/Fonts/FontCache.cpp
--- a/Engine/Source/Runtime/SlateCore/Private/Fonts/FontCache.cpp
+++ b/Engine/Source/Runtime/SlateCore/Private/Fonts/FontCache.cpp
@@ -960,18 +960,23 @@ bool FSlateFontCache::AddNewEntry(const FShapedGlyphEntry& InShapedGlyph, const
 	FCharacterRenderData RenderData;
 	const bool bDidRender = FontRenderer->GetRenderData(InShapedGlyph, InOutlineSettings, RenderData);
 
-	OutAtlasData.Valid = bDidRender && AddNewEntry(RenderData, OutAtlasData.TextureIndex, OutAtlasData.StartU, OutAtlasData.StartV, OutAtlasData.USize, OutAtlasData.VSize);
+	uint8 PaddingOffset = 0;
+	OutAtlasData.Valid = bDidRender && AddNewEntry(RenderData, OutAtlasData.TextureIndex, OutAtlasData.StartU, OutAtlasData.StartV, OutAtlasData.USize, OutAtlasData.VSize, PaddingOffset);
 	if (OutAtlasData.Valid)
 	{
-		OutAtlasData.VerticalOffset = RenderData.VerticalOffset;
-		OutAtlasData.HorizontalOffset = RenderData.HorizontalOffset;
+		int32 VerticalOffset = RenderData.VerticalOffset + (int32)PaddingOffset;
+		int32 HorizontalOffset = RenderData.HorizontalOffset - (int32)PaddingOffset;
+		ensureMsgf(VerticalOffset >= std::numeric_limits<int16>::min() && VerticalOffset <= std::numeric_limits<int16>::max(), TEXT("The Glyph offset is too big"));
+		ensureMsgf(HorizontalOffset >= std::numeric_limits<int16>::min() && HorizontalOffset <= std::numeric_limits<int16>::max(), TEXT("The Glyph offset is too big"));
+		OutAtlasData.VerticalOffset = (int16)VerticalOffset;
+		OutAtlasData.HorizontalOffset = (int16)HorizontalOffset;
 		OutAtlasData.SupportsOutline = RenderData.bSupportsOutline;
 	}
 
 	return OutAtlasData.Valid;
 }
 
-bool FSlateFontCache::AddNewEntry( const FCharacterRenderData InRenderData, uint8& OutTextureIndex, uint16& OutGlyphX, uint16& OutGlyphY, uint16& OutGlyphWidth, uint16& OutGlyphHeight )
+bool FSlateFontCache::AddNewEntry( const FCharacterRenderData InRenderData, uint8& OutTextureIndex, uint16& OutGlyphX, uint16& OutGlyphY, uint16& OutGlyphWidth, uint16& OutGlyphHeight, uint8& OutPaddingOffset )
 {
 	// Will this entry fit within any atlas texture?
 	const FIntPoint FontAtlasSize = FontAtlasFactory->GetAtlasSize(InRenderData.ContentType);
@@ -999,7 +1004,8 @@ bool FSlateFontCache::AddNewEntry( const FCharacterRenderData InRenderData, uint
 			OutGlyphY = 0;
 			OutGlyphWidth = InRenderData.SizeX;
 			OutGlyphHeight = InRenderData.SizeY;
-			
+			OutPaddingOffset = 0;
+
 			if (!bFlushRequested)
 			{
 				UpdateFlushCounters(GrayscaleFontAtlasIndices.Num(), ColorFontAtlasIndices.Num(), MsdfFontAtlasIndices.Num(), NonAtlasedTextureIndices.Num());
@@ -1017,18 +1023,19 @@ bool FSlateFontCache::AddNewEntry( const FCharacterRenderData InRenderData, uint
 
 	auto FillOutputParamsFromAtlasedTextureSlot = [&](const FAtlasedTextureSlot& AtlasedTextureSlot)
 	{
-		int32 GlyphX = AtlasedTextureSlot.X + (int32)AtlasedTextureSlot.Padding;
-		int32 GlyphY = AtlasedTextureSlot.Y + (int32)AtlasedTextureSlot.Padding;
-		int32 GlyphWidth = AtlasedTextureSlot.Width - (int32)(2.f * AtlasedTextureSlot.Padding);
-		int32 GlyphHeight = AtlasedTextureSlot.Height - (int32)(2.f * AtlasedTextureSlot.Padding);
-		ensureMsgf(GlyphX >= 0 && GlyphX <= std::numeric_limits<uint16>::max(), TEXT("The Glyph size is too big"));
-		ensureMsgf(GlyphY >= 0 && GlyphY <= std::numeric_limits<uint16>::max(), TEXT("The Glyph size is too big"));
-		ensureMsgf(GlyphWidth >= 0 && GlyphWidth <= std::numeric_limits<uint16>::max(), TEXT("The Glyph size is too big"));
-		ensureMsgf(GlyphHeight >= 0 && GlyphHeight <= std::numeric_limits<uint16>::max(), TEXT("The Glyph size is too big"));
+		uint32 GlyphX = AtlasedTextureSlot.X;
+		uint32 GlyphY = AtlasedTextureSlot.Y;
+		uint32 GlyphWidth = AtlasedTextureSlot.Width;
+		uint32 GlyphHeight = AtlasedTextureSlot.Height;
+		ensureMsgf(GlyphX <= std::numeric_limits<uint16>::max(), TEXT("The Glyph size is too big"));
+		ensureMsgf(GlyphY <= std::numeric_limits<uint16>::max(), TEXT("The Glyph size is too big"));
+		ensureMsgf(GlyphWidth <= std::numeric_limits<uint16>::max(), TEXT("The Glyph size is too big"));
+		ensureMsgf(GlyphHeight <= std::numeric_limits<uint16>::max(), TEXT("The Glyph size is too big"));
 		OutGlyphX = (uint16)GlyphX;
 		OutGlyphY = (uint16)GlyphY;
 		OutGlyphWidth = (uint16)GlyphWidth;
 		OutGlyphHeight = (uint16)GlyphHeight;
+		OutPaddingOffset = AtlasedTextureSlot.Padding;
 	};
 
 	TArray<uint8>* FontAtlasIndicesPtr = NULL;
diff --git a/Engine/Source/Runtime/SlateCore/Public/Fonts/FontCache.h b/Engine/Source/Runtime/SlateCore/Public/Fonts/FontCache.h
--- a/Engine/Source/Runtime/SlateCore/Public/Fonts/FontCache.h
+++ b/Engine/Source/Runtime/SlateCore/Public/Fonts/FontCache.h
@@ -1029,7 +1029,7 @@ private:
 	 */
 	SLATECORE_API bool AddNewEntry(const FShapedGlyphEntry& InShapedGlyph, const FFontOutlineSettings& InOutlineSettings, FShapedGlyphFontAtlasData& OutAtlasData);
 
-	SLATECORE_API bool AddNewEntry(const FCharacterRenderData InRenderData, uint8& OutTextureIndex, uint16& OutGlyphX, uint16& OutGlyphY, uint16& OutGlyphWidth, uint16& OutGlyphHeight);
+	SLATECORE_API bool AddNewEntry(const FCharacterRenderData InRenderData, uint8& OutTextureIndex, uint16& OutGlyphX, uint16& OutGlyphY, uint16& OutGlyphWidth, uint16& OutGlyphHeight, uint8& OutPaddingOffset);
 
 #if !UE_BUILD_SHIPPING
 	/** Dump statistics about Font cache usage if needed */
diff --git a/Engine/Source/Runtime/SlateCore/Public/Textures/TextureAtlas.h b/Engine/Source/Runtime/SlateCore/Public/Textures/TextureAtlas.h
--- a/Engine/Source/Runtime/SlateCore/Public/Textures/TextureAtlas.h
+++ b/Engine/Source/Runtime/SlateCore/Public/Textures/TextureAtlas.h
@@ -65,7 +65,7 @@ SLATECORE_API uint32 GetSlateFontAtlasContentBytesPerPixel(ESlateFontAtlasConten
  *
  * When a slot is occupied by texture data, the remaining space in the slot (if big enough) is split off into two new (smaller) slots,
  * building a tree of texture rectangles which, instead of being stored as a tree, are flattened into two linked-lists:
- *	- AtlastEmptySlots:	A linked-list of empty slots ready for texture data - iterates in same order as a depth-first-search on a tree
+ *	- AtlasEmptySlots:	A linked-list of empty slots ready for texture data - iterates in same order as a depth-first-search on a tree
  *	- AtlasUsedSlots:	An unordered linked-list of slots containing texture data
  */
 struct FAtlasedTextureSlot : public TIntrusiveLinkedList<FAtlasedTextureSlot>
