Moved some bool-taking container wrappers over to using EAllowShrinking.
FGraphAStar::PopIndex() isn't passing on its argument to the container - that existing behavior has been retained.

diff --git a/Engine/Source/Runtime/AIModule/Public/GraphAStar.h b/Engine/Source/Runtime/AIModule/Public/GraphAStar.h
--- a/Engine/Source/Runtime/AIModule/Public/GraphAStar.h
+++ b/Engine/Source/Runtime/AIModule/Public/GraphAStar.h
@@ -299,13 +299,18 @@ struct FGraphAStar
 			check(false); // We should never reach here.
 		}
 
-		int32 PopIndex(bool bAllowShrinking = true)
+		int32 PopIndex(EAllowShrinking AllowShrinking = EAllowShrinking::Yes)
 		{
 			int32 SearchNodeIndex = INDEX_NONE;
-			Super::HeapPop(SearchNodeIndex, NodeSorter, /*bAllowShrinking = */false);
+			Super::HeapPop(SearchNodeIndex, NodeSorter, EAllowShrinking::No);
 			NodePool[SearchNodeIndex].MarkNotOpened();
 			return SearchNodeIndex;
 		}
+		UE_ALLOWSHRINKING_BOOL_DEPRECATED("PopIndex")
+		FORCEINLINE int32 PopIndex(bool bAllowShrinking)
+		{
+			return PopIndex(bAllowShrinking ? EAllowShrinking::Yes : EAllowShrinking::No);
+		}
 	};
 
 	const TGraph& Graph;
diff --git a/Engine/Source/Runtime/Engine/Classes/Animation/PreviewAssetAttachComponent.h b/Engine/Source/Runtime/Engine/Classes/Animation/PreviewAssetAttachComponent.h
--- a/Engine/Source/Runtime/Engine/Classes/Animation/PreviewAssetAttachComponent.h
+++ b/Engine/Source/Runtime/Engine/Classes/Animation/PreviewAssetAttachComponent.h
@@ -139,7 +139,12 @@ public:
 	/** 
 	 * RemoveAtSwap passthrough
 	 */
-	ENGINE_API void RemoveAtSwap( int32 Index, int32 Count = 1, bool bAllowShrinking = true );
+	ENGINE_API void RemoveAtSwap( int32 Index, int32 Count = 1, EAllowShrinking AllowShrinking = EAllowShrinking::Yes);
+	UE_ALLOWSHRINKING_BOOL_DEPRECATED("RemoveAtSwap")
+	FORCEINLINE void RemoveAtSwap(int32 Index, int32 Count, bool bAllowShrinking)
+	{
+		RemoveAtSwap(Index, Count, bAllowShrinking ? EAllowShrinking::Yes : EAllowShrinking::No);
+	}
 
 	/**
 	 * Helper function to fix up attached objects after property deprecation
diff --git a/Engine/Source/Runtime/Engine/Private/Animation/PreviewAssetAttachComponent.cpp b/Engine/Source/Runtime/Engine/Private/Animation/PreviewAssetAttachComponent.cpp
--- a/Engine/Source/Runtime/Engine/Private/Animation/PreviewAssetAttachComponent.cpp
+++ b/Engine/Source/Runtime/Engine/Private/Animation/PreviewAssetAttachComponent.cpp
@@ -36,7 +36,7 @@ void FPreviewAssetAttachContainer::RemoveAttachedObject( UObject* ObjectToRemove
 
 		if( Pair.GetAttachedObject() == ObjectToRemove && Pair.AttachedTo == AttachName )
 		{
-			AttachedObjects.RemoveAtSwap( i, 1, false );
+			AttachedObjects.RemoveAtSwap( i, 1, EAllowShrinking::No);
 			break;
 		}
 	}
@@ -84,9 +84,9 @@ TIterator FPreviewAssetAttachContainer::CreateIterator()
 	return AttachedObjects.CreateIterator();
 }
 
-void FPreviewAssetAttachContainer::RemoveAtSwap( int32 Index, int32 Count /* = 1 */, bool bAllowShrinking /*=true */ )
+void FPreviewAssetAttachContainer::RemoveAtSwap( int32 Index, int32 Count /* = 1 */, EAllowShrinking AllowShrinking /*=EAllowShrinking::Yes */ )
 {
-	AttachedObjects.RemoveAtSwap( Index, Count, bAllowShrinking );
+	AttachedObjects.RemoveAtSwap( Index, Count, AllowShrinking );
 }
 
 void FPreviewAssetAttachContainer::SaveAttachedObjectsFromDeprecatedProperties()
diff --git a/Engine/Source/Runtime/Experimental/Chaos/Public/RewindData.h b/Engine/Source/Runtime/Experimental/Chaos/Public/RewindData.h
--- a/Engine/Source/Runtime/Experimental/Chaos/Public/RewindData.h
+++ b/Engine/Source/Runtime/Experimental/Chaos/Public/RewindData.h
@@ -1242,12 +1242,12 @@ public:
 		return nullptr;
 	}
 
-	void Remove(const TKey Key, const bool bAllowShrinking)
+	void Remove(const TKey Key, const EAllowShrinking AllowShrinking)
 	{
 		if (const int32* Idx = KeyToIdx.Find(Key))
 		{
 			constexpr int32 Count = 1;
-			DenseVals.RemoveAtSwap(*Idx, Count, bAllowShrinking);
+			DenseVals.RemoveAtSwap(*Idx, Count, AllowShrinking);
 
 			if(*Idx < DenseVals.Num())
 			{
@@ -1258,6 +1258,11 @@ public:
 			KeyToIdx.Remove(Key);
 		}
 	}
+	UE_ALLOWSHRINKING_BOOL_DEPRECATED("Remove")
+	FORCEINLINE void Remove(const TKey Key, const bool bAllowShrinking)
+	{
+		Remove(Key, bAllowShrinking ? EAllowShrinking::Yes : EAllowShrinking::No);
+	}
 
 	void Shrink()
 	{
@@ -1325,14 +1330,24 @@ public:
 		return Managers[Frame].DeltaTime;
 	}
 
-	void RemoveObject(const FGeometryParticleHandle* Particle, const bool bAllowShrinking=true)
+	void RemoveObject(const FGeometryParticleHandle* Particle, const EAllowShrinking AllowShrinking=EAllowShrinking::Yes)
+	{
+		DirtyParticles.Remove(Particle, AllowShrinking);
+	}
+	UE_ALLOWSHRINKING_BOOL_DEPRECATED("RemoveObject")
+	FORCEINLINE void RemoveObject(const FGeometryParticleHandle* Particle, const bool bAllowShrinking)
 	{
-		DirtyParticles.Remove(Particle, bAllowShrinking);
+		RemoveObject(Particle, bAllowShrinking ? EAllowShrinking::Yes : EAllowShrinking::No);
 	}
 
-	void RemoveObject(const FPBDJointConstraintHandle* Joint, const bool bAllowShrinking = true)
+	void RemoveObject(const FPBDJointConstraintHandle* Joint, const EAllowShrinking AllowShrinking = EAllowShrinking::Yes)
+	{
+		DirtyJoints.Remove(Joint, AllowShrinking);
+	}
+	UE_ALLOWSHRINKING_BOOL_DEPRECATED("RemoveObject")
+	FORCEINLINE void RemoveObject(const FPBDJointConstraintHandle* Joint, const bool bAllowShrinking)
 	{
-		DirtyJoints.Remove(Joint, bAllowShrinking);
+		RemoveObject(Joint, bAllowShrinking ? EAllowShrinking::Yes : EAllowShrinking::No);
 	}
 
 	int32 GetEarliestFrame_Internal() const { return CurFrame - FramesSaved; }
diff --git a/Engine/Source/Runtime/GeometryCore/Private/DynamicMesh/ColliderMesh.cpp b/Engine/Source/Runtime/GeometryCore/Private/DynamicMesh/ColliderMesh.cpp
--- a/Engine/Source/Runtime/GeometryCore/Private/DynamicMesh/ColliderMesh.cpp
+++ b/Engine/Source/Runtime/GeometryCore/Private/DynamicMesh/ColliderMesh.cpp
@@ -14,15 +14,15 @@ FColliderMesh::FColliderMesh(const FDynamicMesh3& SourceMesh, const FBuildOption
 
 FColliderMesh::FColliderMesh()
 {
-	Reset(false);
+	Reset(EAllowShrinking::No);
 }
 
-void FColliderMesh::Reset(bool bReleaseMemory)
+void FColliderMesh::Reset(EAllowShrinking AllowShrinking)
 {
-	Vertices.SetNum(0, bReleaseMemory);
-	SourceVertexIDs.SetNum(0, bReleaseMemory);
-	Triangles.SetNum(0, bReleaseMemory);
-	SourceTriangleIDs.SetNum(0, bReleaseMemory);
+	Vertices.SetNum(0, AllowShrinking);
+	SourceVertexIDs.SetNum(0, AllowShrinking);
+	Triangles.SetNum(0, AllowShrinking);
+	SourceTriangleIDs.SetNum(0, AllowShrinking);
 
 	bSourceWasCompactV = bSourceWasCompactT = true;
 
diff --git a/Engine/Source/Runtime/GeometryCore/Public/DynamicMesh/ColliderMesh.h b/Engine/Source/Runtime/GeometryCore/Public/DynamicMesh/ColliderMesh.h
--- a/Engine/Source/Runtime/GeometryCore/Public/DynamicMesh/ColliderMesh.h
+++ b/Engine/Source/Runtime/GeometryCore/Public/DynamicMesh/ColliderMesh.h
@@ -65,7 +65,12 @@ public:
 
 	GEOMETRYCORE_API void Initialize(const FDynamicMesh3& SourceMesh, const FBuildOptions& BuildOptions = FBuildOptions());
 
-	GEOMETRYCORE_API void Reset(bool bReleaseMemory);
+	GEOMETRYCORE_API void Reset(EAllowShrinking AllowShrinking);
+	UE_ALLOWSHRINKING_BOOL_DEPRECATED("Reset")
+	FORCEINLINE void Reset(bool bReleaseMemory)
+	{
+		Reset(bReleaseMemory ? EAllowShrinking::Yes : EAllowShrinking::No);
+	}
 
 	// mesh API required for TMeshAABBTree3
 	bool IsTriangle(int32 TriangleID) const { return TriangleID >= 0 && TriangleID < Triangles.Num(); }
diff --git a/Engine/Source/Runtime/GeometryCore/Public/Spatial/DenseGrid2.h b/Engine/Source/Runtime/GeometryCore/Public/Spatial/DenseGrid2.h
--- a/Engine/Source/Runtime/GeometryCore/Public/Spatial/DenseGrid2.h
+++ b/Engine/Source/Runtime/GeometryCore/Public/Spatial/DenseGrid2.h
@@ -63,11 +63,16 @@ public:
 		return FVector2i((int32)DimensionX, (int32)DimensionY);
 	}
 
-	void Resize(int32 DimX, int32 DimY, bool bAllowShrinking = true)
+	void Resize(int32 DimX, int32 DimY, EAllowShrinking AllowShrinking = EAllowShrinking::Yes)
 	{
 		DimensionX = (int64)DimX;
 		DimensionY = (int64)DimY;
-		Buffer.SetNumUninitialized( DimensionX*DimensionY, bAllowShrinking);
+		Buffer.SetNumUninitialized( DimensionX*DimensionY, AllowShrinking);
+	}
+	UE_ALLOWSHRINKING_BOOL_DEPRECATED("Resize")
+	FORCEINLINE void Resize(int32 DimX, int32 DimY, bool bAllowShrinking)
+	{
+		Resize(DimX, DimY, bAllowShrinking ? EAllowShrinking::Yes : EAllowShrinking::No);
 	}
 
 	void AssignAll(ElemType Value)
diff --git a/Engine/Source/Runtime/GeometryCore/Public/Spatial/DenseGrid3.h b/Engine/Source/Runtime/GeometryCore/Public/Spatial/DenseGrid3.h
--- a/Engine/Source/Runtime/GeometryCore/Public/Spatial/DenseGrid3.h
+++ b/Engine/Source/Runtime/GeometryCore/Public/Spatial/DenseGrid3.h
@@ -58,12 +58,17 @@ public:
 		return Dimensions;
 	}
 
-	void Resize(int DimX, int DimY, int DimZ, bool bAllowShrinking = true)
+	void Resize(int DimX, int DimY, int DimZ, EAllowShrinking AllowShrinking = EAllowShrinking::Yes)
 	{
 		check((int64)DimX * (int64)DimY * (int64)DimZ < INT_MAX);
-		Buffer.SetNumUninitialized(DimX * DimY * DimZ, bAllowShrinking);
+		Buffer.SetNumUninitialized(DimX * DimY * DimZ, AllowShrinking);
 		Dimensions = FVector3i(DimX, DimY, DimZ);
 	}
+	UE_ALLOWSHRINKING_BOOL_DEPRECATED("Resize")
+	FORCEINLINE void Resize(int DimX, int DimY, int DimZ, bool bAllowShrinking)
+	{
+		Resize(DimX, DimY, DimZ, bAllowShrinking ? EAllowShrinking::Yes : EAllowShrinking::No);
+	}
 
 	void Assign(ElemType Value)
 	{
diff --git a/Engine/Source/Runtime/RenderCore/Public/ShaderSource.h b/Engine/Source/Runtime/RenderCore/Public/ShaderSource.h
--- a/Engine/Source/Runtime/RenderCore/Public/ShaderSource.h
+++ b/Engine/Source/Runtime/RenderCore/Public/ShaderSource.h
@@ -23,9 +23,9 @@
 class FShaderSource
 {
 private:
-	inline void SetLen(int32 Num, bool bShrink = true)
+	inline void SetLen(int32 Num, EAllowShrinking AllowShrinking = EAllowShrinking::Yes)
 	{
-		Source.SetNumUninitialized(Num + ShaderSourceSimdPadding, bShrink);
+		Source.SetNumUninitialized(Num + ShaderSourceSimdPadding, AllowShrinking);
 		FMemory::Memzero(Source.GetData() + Num, sizeof(CharType) * ShaderSourceSimdPadding);
 	}
 
@@ -83,12 +83,17 @@ public:
 
 	/* Reduces the set size of the stored string length, optionally shrinking the allocation.
 	 * @param Num the desired allocation size (padding bytes will be added on top of this)
-	 * @param bRealloc whether to reallocate or keep the existing larger size allocation
+	 * @param AllowShrinking whether to reallocate or keep the existing larger size allocation
 	 */
-	inline void ShrinkToLen(int32 Num, bool bShrink = true)
+	inline void ShrinkToLen(int32 Num, EAllowShrinking AllowShrinking = EAllowShrinking::Yes)
 	{
 		checkf(Num <= Len(), TEXT("Trying to shrink to %d characters but existing allocation is smaller (%d characters)"), Num, Len());
-		SetLen(Num, bShrink);
+		SetLen(Num, AllowShrinking);
+	}
+	UE_ALLOWSHRINKING_BOOL_DEPRECATED("ShrinkToLen")
+	FORCEINLINE void ShrinkToLen(int32 Num, bool bShrink)
+	{
+		ShrinkToLen(Num, bShrink ? EAllowShrinking::Yes : EAllowShrinking::No);
 	}
 
 	/* String view accessor.
diff --git a/Engine/Source/Runtime/SlateCore/Public/Containers/ObservableArray.h b/Engine/Source/Runtime/SlateCore/Public/Containers/ObservableArray.h
--- a/Engine/Source/Runtime/SlateCore/Public/Containers/ObservableArray.h
+++ b/Engine/Source/Runtime/SlateCore/Public/Containers/ObservableArray.h
@@ -299,35 +299,45 @@ public:
 		}
 	}
 
-	SizeType RemoveSingle(const ElementType& Item, bool bAllowShrinking = true)
+	SizeType RemoveSingle(const ElementType& Item, EAllowShrinking AllowShrinking = EAllowShrinking::Yes)
 	{
 		SizeType Index = Array.Find(Item);
 		if (Index == INDEX_NONE)
 		{
 			return 0;
 		}
-		RemoveAt(Index, 1, bAllowShrinking);
+		RemoveAt(Index, 1, AllowShrinking);
 		return 1;
 	}
+	UE_ALLOWSHRINKING_BOOL_DEPRECATED("RemoveSingle")
+	FORCEINLINE SizeType RemoveSingle(const ElementType& Item, bool bAllowShrinking)
+	{
+		return RemoveSingle(Item, bAllowShrinking ? EAllowShrinking::Yes : EAllowShrinking::No);
+	}
 
-	SizeType RemoveSingleSwap(const ElementType& Item, bool bAllowShrinking = true)
+	SizeType RemoveSingleSwap(const ElementType& Item, EAllowShrinking AllowShrinking = EAllowShrinking::Yes)
 	{
 		SizeType Index = Array.Find(Item);
 		if (Index == INDEX_NONE)
 		{
 			return 0;
 		}
-		RemoveAtSwap(Index, 1, bAllowShrinking);
+		RemoveAtSwap(Index, 1, AllowShrinking);
 		return 1;
 	}
+	UE_ALLOWSHRINKING_BOOL_DEPRECATED("RemoveSingleSwap")
+	FORCEINLINE SizeType RemoveSingleSwap(const ElementType& Item, bool bAllowShrinking)
+	{
+		return RemoveSingleSwap(Item, bAllowShrinking ? EAllowShrinking::Yes : EAllowShrinking::No);
+	}
 
 	void RemoveAt(SizeType Index)
 	{
-		RemoveAt(Index, 1, true);
+		RemoveAt(Index, 1, EAllowShrinking::Yes);
 	}
 
 	template <typename CountType>
-	void RemoveAt(SizeType Index, CountType NumToRemove = 1, bool bAllowShrinking = true)
+	void RemoveAt(SizeType Index, CountType NumToRemove = 1, EAllowShrinking AllowShrinking = EAllowShrinking::Yes)
 	{
 		static_assert(!std::is_same_v<CountType, bool>, "TObservableArray::RemoveAt: unexpected bool passed as the Count argument");
 		check((NumToRemove > 0) & (Index >= 0) & (Index + NumToRemove <= Num()));
@@ -336,7 +346,7 @@ public:
 			if (NumToRemove == 1)
 			{
 				ElementType RemovedElement = MoveTemp(Array[Index]);
-				Array.RemoveAt(Index, NumToRemove, bAllowShrinking);
+				Array.RemoveAt(Index, NumToRemove, AllowShrinking);
 				ArrayChangedDelegate.Broadcast(ObservableArrayChangedArgsType::MakeRemoveAction({ &RemovedElement, 1 }, Index));
 			}
 			else
@@ -350,19 +360,24 @@ public:
 					RemovedElements.Add(MoveTemp(Array[RemoveIndex + Index]));
 				}
 				
-				Array.RemoveAt(Index, NumToRemove, bAllowShrinking);
+				Array.RemoveAt(Index, NumToRemove, AllowShrinking);
 				ArrayChangedDelegate.Broadcast(ObservableArrayChangedArgsType::MakeRemoveAction({ RemovedElements.GetData(), NumToRemove }, Index));
 			}
 		}
 	}
+	UE_ALLOWSHRINKING_BOOL_DEPRECATED("RemoveAt")
+	FORCEINLINE void RemoveAt(SizeType Index, SizeType NumToRemove, bool bAllowShrinking)
+	{
+		RemoveAt(Index, NumToRemove, bAllowShrinking ? EAllowShrinking::Yes : EAllowShrinking::No);
+	}
 
 	void RemoveAtSwap(SizeType Index)
 	{
-		RemoveAtSwap(Index, 1, true);
+		RemoveAtSwap(Index, 1, EAllowShrinking::Yes);
 	}
 	
 	template <typename CountType>
-	void RemoveAtSwap(SizeType Index, CountType NumToRemove = 1, bool bAllowShrinking = true)
+	void RemoveAtSwap(SizeType Index, CountType NumToRemove = 1, EAllowShrinking AllowShrinking = EAllowShrinking::Yes)
 	{
 		static_assert(!std::is_same_v<CountType, bool>, "TObservableArray::RemoveAtSwap: unexpected bool passed as the Count argument");
 		check((NumToRemove > 0) & (Index >= 0) & (Index + NumToRemove <= Num()));
@@ -374,7 +389,7 @@ public:
 			{
 				ElementType RemovedElement = MoveTemp(Array[Index]);
 				
-				Array.RemoveAtSwap(Index, NumToRemove, bAllowShrinking);
+				Array.RemoveAtSwap(Index, NumToRemove, AllowShrinking);
 				ArrayChangedDelegate.Broadcast(ObservableArrayChangedArgsType::MakeRemoveSwapAction({ &RemovedElement, 1 }, Index, PreviousNum - SwapAmount));
 			}
 			else
@@ -388,11 +403,16 @@ public:
 					RemovedElements.Add(MoveTemp(Array[RemoveIndex + Index]));
 				}
 
-				Array.RemoveAtSwap(Index, NumToRemove, bAllowShrinking);
+				Array.RemoveAtSwap(Index, NumToRemove, AllowShrinking);
 				ArrayChangedDelegate.Broadcast(ObservableArrayChangedArgsType::MakeRemoveSwapAction({ RemovedElements.GetData(), NumToRemove }, Index, PreviousNum - SwapAmount));
 			}
 		}
 	}
+	UE_ALLOWSHRINKING_BOOL_DEPRECATED("RemoveAtSwap")
+	FORCEINLINE void RemoveAtSwap(SizeType Index, SizeType NumToRemove, bool bAllowShrinking)
+	{
+		RemoveAtSwap(Index, NumToRemove, bAllowShrinking ? EAllowShrinking::Yes : EAllowShrinking::No);
+	}
 
 	void Swap(SizeType FirstIndexToSwap, SizeType SecondIndexToSwap)
 	{
