Slate: Add test to prevent unwanted modifications while iterating on a the list with the FWidgetAttributeIterator.

diff --git a/Engine/Source/Runtime/SlateCore/Private/FastUpdate/SlateInvalidationWidgetList.cpp b/Engine/Source/Runtime/SlateCore/Private/FastUpdate/SlateInvalidationWidgetList.cpp
--- a/Engine/Source/Runtime/SlateCore/Private/FastUpdate/SlateInvalidationWidgetList.cpp
+++ b/Engine/Source/Runtime/SlateCore/Private/FastUpdate/SlateInvalidationWidgetList.cpp
@@ -53,9 +53,13 @@ FSlateInvalidationWidgetList::FWidgetAttributeIterator::FWidgetAttributeIterator
 	, MoveToWidgetIndexOnNextAdvance(FSlateInvalidationWidgetIndex::Invalid)
 	, bNeedsWidgetFixUp(false)
 {
+	++WidgetList.NumberOfLock;
+
 	int32 ArrayIndex = WidgetList.FirstArrayIndex;
 	while(ArrayIndex != INDEX_NONE)
 	{
+		check(WidgetList.Data.IsValidIndex(ArrayIndex));
+
 		const FArrayNode& ArrayNode = WidgetList.Data[ArrayIndex];
 		if (ArrayNode.ElementIndexList_WidgetWithRegisteredSlateAttribute.Num() > 0)
 		{
@@ -68,6 +72,13 @@ FSlateInvalidationWidgetList::FWidgetAttributeIterator::FWidgetAttributeIterator
 }
 
 
+FSlateInvalidationWidgetList::FWidgetAttributeIterator::~FWidgetAttributeIterator()
+{
+	--WidgetList.NumberOfLock;
+	check(WidgetList.NumberOfLock >= 0);
+}
+
+
 void FSlateInvalidationWidgetList::FWidgetAttributeIterator::PreChildRemove(const FIndexRange& Range)
 {
 	// The current widget will be removed or rebuild. Attributes will be updated in the Building loop.
@@ -148,6 +159,7 @@ void FSlateInvalidationWidgetList::FWidgetAttributeIterator::FixCurrentWidgetInd
 void FSlateInvalidationWidgetList::FWidgetAttributeIterator::Seek(FSlateInvalidationWidgetIndex SeekTo)
 {
 	check(SeekTo != FSlateInvalidationWidgetIndex::Invalid);
+	check(WidgetList.Data.IsValidIndex(SeekTo.ArrayIndex));
 
 	const FArrayNode& ArrayNode = WidgetList.Data[SeekTo.ArrayIndex];
 	AttributeIndex = ArrayNode.ElementIndexList_WidgetWithRegisteredSlateAttribute.FindLowerBound(SeekTo.ElementIndex);
@@ -167,7 +179,7 @@ void FSlateInvalidationWidgetList::FWidgetAttributeIterator::Seek(FSlateInvalida
 
 void FSlateInvalidationWidgetList::FWidgetAttributeIterator::Advance()
 {
-	check(MoveToWidgetIndexOnNextAdvance == FSlateInvalidationWidgetIndex::Invalid);
+	check(WidgetList.Data.IsValidIndex(CurrentWidgetIndex.ArrayIndex));
 
 	++AttributeIndex;
 
@@ -234,6 +246,8 @@ void FSlateInvalidationWidgetList::FWidgetAttributeIterator::AdvanceArrayIndex(i
 
 	while (ArrayIndex != INDEX_NONE)
 	{
+		check(WidgetList.Data.IsValidIndex(ArrayIndex));
+
 		const FArrayNode& NewArrayNode = WidgetList.Data[ArrayIndex];
 		if (NewArrayNode.ElementIndexList_WidgetWithRegisteredSlateAttribute.Num() > 0)
 		{
@@ -280,6 +294,8 @@ void FSlateInvalidationWidgetList::FWidgetVolatileUpdateIterator::Advance()
 
 void FSlateInvalidationWidgetList::FWidgetVolatileUpdateIterator::Internal_Advance()
 {
+	check(WidgetList.Data.IsValidIndex(CurrentWidgetIndex.ArrayIndex));
+
 	++AttributeIndex;
 
 	const FArrayNode& ArrayNode = WidgetList.Data[CurrentWidgetIndex.ArrayIndex];
@@ -319,6 +335,7 @@ void FSlateInvalidationWidgetList::FWidgetVolatileUpdateIterator::SkipToNextExpe
 void FSlateInvalidationWidgetList::FWidgetVolatileUpdateIterator::Seek(FSlateInvalidationWidgetIndex SeekTo)
 {
 	check(SeekTo != FSlateInvalidationWidgetIndex::Invalid);
+	check(WidgetList.Data.IsValidIndex(SeekTo.ArrayIndex));
 
 	const FArrayNode& ArrayNode = WidgetList.Data[SeekTo.ArrayIndex];
 	AttributeIndex = ArrayNode.ElementIndexList_VolatileUpdateWidget.FindLowerBound(SeekTo.ElementIndex);
@@ -339,6 +356,8 @@ void FSlateInvalidationWidgetList::FWidgetVolatileUpdateIterator::AdvanceArray(i
 {
 	while (ArrayIndex != INDEX_NONE)
 	{
+		check(WidgetList.Data.IsValidIndex(ArrayIndex));
+
 		const FArrayNode& NewArrayNode = WidgetList.Data[ArrayIndex];
 		if (NewArrayNode.ElementIndexList_VolatileUpdateWidget.Num() > 0)
 		{
@@ -483,6 +502,7 @@ FSlateInvalidationWidgetIndex FSlateInvalidationWidgetList::Internal_BuildWidget
 
 void FSlateInvalidationWidgetList::BuildWidgetList(const TSharedRef<SWidget>& InRoot)
 {
+	ensureMsgf(NumberOfLock == 0, TEXT("You are not allowed to modify the list while iterating on it."));
 	SCOPED_NAMED_EVENT(Slate_InvalidationList_ProcessBuild, FColorList::Blue);
 
 	Reset();
@@ -543,6 +563,7 @@ void FSlateInvalidationWidgetList::Internal_RebuildWidgetListTree(SWidget& Widge
 
 bool FSlateInvalidationWidgetList::ProcessChildOrderInvalidation(FSlateInvalidationWidgetIndex WidgetIndex, IProcessChildOrderInvalidationCallback& Callback)
 {
+	ensureMsgf(NumberOfLock == 0, TEXT("You are not allowed to modify the list while iterating on it."));
 	SCOPE_CYCLE_COUNTER(STAT_WidgetList_ProcessChildOrderInvalidation);
 
 	bool bIsInvalidationWidgetStillValid = true;
@@ -707,6 +728,8 @@ bool FSlateInvalidationWidgetList::ProcessChildOrderInvalidation(FSlateInvalidat
 
 void FSlateInvalidationWidgetList::ProcessAttributeRegistrationInvalidation(const InvalidationWidgetType& InvalidationWidget)
 {
+	ensureMsgf(NumberOfLock == 0, TEXT("You are not allowed to modify the list while iterating on it."));
+
 	SWidget* WidgetPtr = InvalidationWidget.GetWidget();
 	check(WidgetPtr);
 
@@ -723,6 +746,8 @@ void FSlateInvalidationWidgetList::ProcessAttributeRegistrationInvalidation(cons
 
 void FSlateInvalidationWidgetList::ProcessVolatileUpdateInvalidation(InvalidationWidgetType& InvalidationWidget)
 {
+	ensureMsgf(NumberOfLock == 0, TEXT("You are not allowed to modify the list while iterating on it."));
+
 	SWidget* WidgetPtr = InvalidationWidget.GetWidget();
 	check(WidgetPtr);
 
@@ -839,6 +864,8 @@ FSlateInvalidationWidgetIndex FSlateInvalidationWidgetList::FindNextSibling(FSla
 
 void FSlateInvalidationWidgetList::Empty()
 {
+	ensureMsgf(NumberOfLock == 0, TEXT("You are not allowed to modify the list while iterating on it."));
+
 	Data.Empty();
 	Root.Reset();
 	FirstArrayIndex = INDEX_NONE;
@@ -848,6 +875,8 @@ void FSlateInvalidationWidgetList::Empty()
 
 void FSlateInvalidationWidgetList::Reset()
 {
+	ensureMsgf(NumberOfLock == 0, TEXT("You are not allowed to modify the list while iterating on it."));
+
 	Data.Reset();
 	Root.Reset();
 	FirstArrayIndex = INDEX_NONE;
@@ -864,6 +893,8 @@ FSlateInvalidationWidgetList::IndexType FSlateInvalidationWidgetList::AddArrayNo
 			ensureAlwaysMsgf(false, TEXT("The widget array is split more time that we support. Widget will not be updated properly. Try to increase Slate.InvalidationList.MaxArrayElements"));
 			return (IndexType)LastArrayIndex;
 		}
+
+		check(NumberOfLock == 0);
 		const int32 Index = Data.Add(FArrayNode());
 		check(Index < std::numeric_limits<IndexType>::max());
 		if (bReserveElementList)
@@ -910,6 +941,7 @@ void FSlateInvalidationWidgetList::RebuildOrderIndex(IndexType StartFrom)
 	FSlateInvalidationWidgetList* Self = this;
 	auto SetValue = [Self](IndexType ArrayIndex, int32 NewSortOrder)
 		{
+			check(Self->NumberOfLock == 0);
 			ElementListType& ElementList = Self->Data[ArrayIndex].ElementList;
 			int32 ElementIndex = Self->Data[ArrayIndex].StartIndex;
 			const int32 PreviousSortOrder = Self->Data[ArrayIndex].SortOrder;
@@ -1025,6 +1057,7 @@ FSlateInvalidationWidgetList::IndexType FSlateInvalidationWidgetList::InsertData
 	}
 	else
 	{
+		check(NumberOfLock == 0);
 		check(AfterIndex != INDEX_NONE);
 
 		if (Data.Num() + 1 == FSlateInvalidationWidgetIndex::Invalid.ArrayIndex)
@@ -1072,6 +1105,7 @@ FSlateInvalidationWidgetList::IndexType FSlateInvalidationWidgetList::InsertData
 
 void FSlateInvalidationWidgetList::RemoveDataNode(IndexType Index)
 {
+	check(NumberOfLock == 0);
 	check(Index != INDEX_NONE && Index != std::numeric_limits<IndexType>::max());
 	FArrayNode& ArrayNode = Data[Index];
 	if (ArrayNode.PreviousArrayIndex != INDEX_NONE)
@@ -1336,6 +1370,8 @@ FSlateInvalidationWidgetList::FCutResult FSlateInvalidationWidgetList::Internal_
 	// N.B. We can cut/move anywhere. Cross family may occur. Fix up everything.
 	if (WhereToCut.ElementIndex < Data[WhereToCut.ArrayIndex].ElementList.Num() - 1)
 	{
+		check(NumberOfLock == 0);
+
 		//From where to where we are moving the item
 		const IndexType OldArrayIndex = WhereToCut.ArrayIndex;
 		const IndexType OldElementIndexStart = WhereToCut.ElementIndex + 1;
diff --git a/Engine/Source/Runtime/SlateCore/Private/FastUpdate/SlateInvalidationWidgetList.h b/Engine/Source/Runtime/SlateCore/Private/FastUpdate/SlateInvalidationWidgetList.h
--- a/Engine/Source/Runtime/SlateCore/Private/FastUpdate/SlateInvalidationWidgetList.h
+++ b/Engine/Source/Runtime/SlateCore/Private/FastUpdate/SlateInvalidationWidgetList.h
@@ -290,13 +290,21 @@ public:
 		bool bNeedsWidgetFixUp;
 
 	public:
+		~FWidgetAttributeIterator();
 		FWidgetAttributeIterator(const FSlateInvalidationWidgetList& InWidgetList);
+		FWidgetAttributeIterator(const FWidgetAttributeIterator&) = delete;
+		FWidgetAttributeIterator& operator= (const FWidgetAttributeIterator&) = delete;
 
 		//~ Handle operation
+		UE_DEPRECATED(5.4, "PreChildRemove is deprecated. It was unused.")
 		void PreChildRemove(const FIndexRange& Range);
+		UE_DEPRECATED(5.4, "ReIndexed is deprecated. It was unused.")
 		void ReIndexed(const IProcessChildOrderInvalidationCallback::FReIndexOperation& Operation);
+		UE_DEPRECATED(5.4, "PostResort is deprecated. It was unused.")
 		void PostResort();
+		UE_DEPRECATED(5.4, "ProxiesBuilt is deprecated. It was unused.")
 		void ProxiesBuilt(const FIndexRange& Range);
+		UE_DEPRECATED(5.4, "FixCurrentWidgetIndex is deprecated. It was unused.")
 		void FixCurrentWidgetIndex();
 		void Seek(FSlateInvalidationWidgetIndex SeekTo);
 
@@ -565,6 +573,7 @@ private:
 	TWeakPtr<SWidget> Root;
 	int32 FirstArrayIndex = INDEX_NONE;
 	int32 LastArrayIndex = INDEX_NONE;
+	mutable int32 NumberOfLock = 0;
 	IProcessChildOrderInvalidationCallback* CurrentInvalidationCallback = nullptr;
 	const FArguments WidgetListConfig;
 };
\ No newline at end of file
