Fix for HDR viewports not rendering properly when Slate.EnableGlobalInvalidation=1
When Slate.EnableGlobalInvalidation=1, we go through FSlateElementBatcher::AddCachedElements and do BatchData->AddCachedBatches(CachedElementData.GetCachedBatches());
HDR viewports actually need to choose between BatchDataHDR and BatchData based on whether we want to composite ui with the scene later on
Added a flag on the batchdraw flags to specify which batchdata the viewport should end up into.
Since the flag might change when toggling, we invalidate the window the viewport is associated with whenever we detect a viewport change in FSlateRHIRenderer::ConditionalResizeViewport
Add GetViewportDynamicRange on ISlateViewport to be able to distinguish SDR and HDR viewports: we can end up in a scenario where we have an HDR window because still have an SDR viewport, like the FMovieViewport
The assumption is that we only play SDR movies, the code will  have to be updated to support HDR movies

diff --git a/Engine/Source/Runtime/Engine/Private/Slate/SceneViewport.cpp b/Engine/Source/Runtime/Engine/Private/Slate/SceneViewport.cpp
--- a/Engine/Source/Runtime/Engine/Private/Slate/SceneViewport.cpp
+++ b/Engine/Source/Runtime/Engine/Private/Slate/SceneViewport.cpp
@@ -1282,6 +1282,11 @@ bool FSceneViewport::GetSceneHDREnabled() const
 	return bHDRViewport;
 }
 
+ESlateViewportDynamicRange FSceneViewport::GetViewportDynamicRange() const
+{
+	return bHDRViewport ? ESlateViewportDynamicRange::HDR : ESlateViewportDynamicRange::SDR;
+}
+
 void FSceneViewport::OnViewportDeactivated(const FWindowActivateEvent& InActivateEvent)
 {
 	// We backup if we have capture for us on activation, however we also maintain "true" if it's already true!
diff --git a/Engine/Source/Runtime/Engine/Public/Slate/SceneViewport.h b/Engine/Source/Runtime/Engine/Public/Slate/SceneViewport.h
--- a/Engine/Source/Runtime/Engine/Public/Slate/SceneViewport.h
+++ b/Engine/Source/Runtime/Engine/Public/Slate/SceneViewport.h
@@ -245,6 +245,7 @@ public:
 	ENGINE_API virtual EDisplayColorGamut GetDisplayColorGamut() const override;
 	ENGINE_API virtual EDisplayOutputFormat GetDisplayOutputFormat() const override;
 	ENGINE_API virtual bool GetSceneHDREnabled() const override;
+	ENGINE_API virtual ESlateViewportDynamicRange GetViewportDynamicRange() const override;
 
 	ENGINE_API void SetViewportSize(uint32 NewSizeX,uint32 NewSizeY);
 	ENGINE_API void SetFixedViewportSize(uint32 NewSizeX, uint32 NewSizeY);
diff --git a/Engine/Source/Runtime/SlateCore/Private/Rendering/DrawElementTypes.cpp b/Engine/Source/Runtime/SlateCore/Private/Rendering/DrawElementTypes.cpp
--- a/Engine/Source/Runtime/SlateCore/Private/Rendering/DrawElementTypes.cpp
+++ b/Engine/Source/Runtime/SlateCore/Private/Rendering/DrawElementTypes.cpp
@@ -651,6 +651,11 @@ void FSlateDrawElement::MakeViewport( FSlateWindowElementList& ElementList, uint
 	check(Element.RenderTargetResource == nullptr || !Element.RenderTargetResource->Debug_IsDestroyed());
 
 	Element.Init(ElementList, EElementType::ET_Viewport, InLayer, PaintGeometry, InDrawEffects);
+
+	if (Viewport->GetViewportDynamicRange() == ESlateViewportDynamicRange::HDR)
+	{
+		EnumAddFlags(Element.BatchFlags, ESlateBatchDrawFlag::HDR);
+	}
 }
 
 
diff --git a/Engine/Source/Runtime/SlateCore/Private/Rendering/ElementBatcher.cpp b/Engine/Source/Runtime/SlateCore/Private/Rendering/ElementBatcher.cpp
--- a/Engine/Source/Runtime/SlateCore/Private/Rendering/ElementBatcher.cpp
+++ b/Engine/Source/Runtime/SlateCore/Private/Rendering/ElementBatcher.cpp
@@ -119,6 +119,26 @@ void FSlateBatchData::AddCachedBatches(const TSparseArray<FSlateRenderBatch>& In
 	}
 }
 
+void FSlateBatchData::AddCachedBatchesToBatchData(FSlateBatchData* BatchDataSDR, FSlateBatchData* BatchDataHDR, const TSparseArray<FSlateRenderBatch>& InCachedBatches)
+{
+	TArray<FSlateRenderBatch>& RenderBatchesSDR = BatchDataSDR->RenderBatches;
+	TArray<FSlateRenderBatch>& RenderBatchesHDR = BatchDataHDR->RenderBatches;
+
+	RenderBatchesSDR.Reserve(RenderBatchesSDR.Num() + InCachedBatches.Num());
+	RenderBatchesHDR.Reserve(RenderBatchesHDR.Num() + InCachedBatches.Num());
+	for (const FSlateRenderBatch& CachedBatch : InCachedBatches)
+	{
+		if (EnumHasAnyFlags(CachedBatch.GetDrawFlags(), ESlateBatchDrawFlag::HDR))
+		{
+			RenderBatchesHDR.Add(CachedBatch);
+		}
+		else
+		{
+			RenderBatchesSDR.Add(CachedBatch);
+		}
+	}
+}
+
 void FSlateBatchData::FillBuffersFromNewBatch(FSlateRenderBatch& Batch, FSlateVertexArray& FinalVertices, FSlateIndexArray& FinalIndices)
 {
 	if(Batch.HasVertexData())
@@ -549,8 +569,7 @@ void FSlateElementBatcher::AddCachedElements(FSlateCachedElementData& CachedElem
 	CachedElementData.ListsWithNewData.Empty();
 
 	// Add the existing and new cached batches.
-	BatchData->AddCachedBatches(CachedElementData.GetCachedBatches());
-
+	FSlateBatchData::AddCachedBatchesToBatchData(BatchData, BatchDataHDR, CachedElementData.GetCachedBatches());
 	CachedElementData.CleanupUnusedClipStates();
 
 #if SLATE_CSV_TRACKER
@@ -2659,7 +2678,11 @@ void FSlateElementBatcher::AddViewportElement( const FSlateViewportElement& Draw
 		// This is a slight hack, but the grayscale font shader is the same as the general shader except it reads alpha only textures and doesn't support tiling
 		ShaderType = ESlateShader::GrayscaleFont;
 	}
-	FSlateBatchData* UsedBatchData = bCompositeHDRViewports ? BatchDataHDR : BatchData;
+
+	bool bIsHDRViewport = EnumHasAnyFlags(DrawFlags, ESlateBatchDrawFlag::HDR);
+	bool bUseBatchDataHDR = (bCompositeHDRViewports && bIsHDRViewport);
+
+	FSlateBatchData* UsedBatchData = bUseBatchDataHDR ? BatchDataHDR : BatchData;
 	FSlateRenderBatch& RenderBatch = CreateRenderBatch(UsedBatchData, Layer, FShaderParams(), ViewportResource, ESlateDrawPrimitive::TriangleList, ShaderType, InDrawEffects, DrawFlags, DrawElement);
 
 	// Tag this batch as requiring vsync if the viewport requires it.
@@ -2700,7 +2723,7 @@ void FSlateElementBatcher::AddViewportElement( const FSlateViewportElement& Draw
 	RenderBatch.AddIndex( IndexStart + 1 );
 	RenderBatch.AddIndex( IndexStart + 3 );
 
-	if (bCompositeHDRViewports)
+	if (bUseBatchDataHDR)
 	{
 		// used to poke a hole in the slate tree: in case HDR is enabled, we need to compose the hdr scene with the SDR ui based on UI alpha
 		// The problem is that in editor mode, there's already a few quads already drawn below the viewport that we actually don't want. If we had an easy way to split
diff --git a/Engine/Source/Runtime/SlateCore/Public/Rendering/ElementBatcher.h b/Engine/Source/Runtime/SlateCore/Public/Rendering/ElementBatcher.h
--- a/Engine/Source/Runtime/SlateCore/Public/Rendering/ElementBatcher.h
+++ b/Engine/Source/Runtime/SlateCore/Public/Rendering/ElementBatcher.h
@@ -200,6 +200,7 @@ public:
 
 	/** Adds a cached batch, used in retained rendering */
 	void AddCachedBatches(const TSparseArray<FSlateRenderBatch>& InCachedBatches);
+	static void AddCachedBatchesToBatchData(FSlateBatchData* BatchDataSDR, FSlateBatchData* BatchDataHDR, const TSparseArray<FSlateRenderBatch>& InCachedBatches);
 private:
 	void FillBuffersFromNewBatch(FSlateRenderBatch& Batch, FSlateVertexArray& FinalVertices, FSlateIndexArray& FinalIndices);
 	void CombineBatches(FSlateRenderBatch& FirstBatch, FSlateRenderBatch& SecondBatch, FSlateVertexArray& FinalVertices, FSlateIndexArray& FinalIndices);
diff --git a/Engine/Source/Runtime/SlateCore/Public/Rendering/RenderingCommon.h b/Engine/Source/Runtime/SlateCore/Public/Rendering/RenderingCommon.h
--- a/Engine/Source/Runtime/SlateCore/Public/Rendering/RenderingCommon.h
+++ b/Engine/Source/Runtime/SlateCore/Public/Rendering/RenderingCommon.h
@@ -107,7 +107,7 @@ enum class ESlateDrawEffect : uint8
 ENUM_CLASS_FLAGS(ESlateDrawEffect);
 
 /** Flags for drawing a batch */
-enum class ESlateBatchDrawFlag : uint8
+enum class ESlateBatchDrawFlag : uint16
 {
 	/** No draw flags */
 	None					= 0,
@@ -129,7 +129,9 @@ enum class ESlateBatchDrawFlag : uint8
 	/** The element should be tiled vertically */
 	TileV				= 1 << 6,
 	/** Reverse gamma correction */
-	ReverseGamma			 = 1 << 7
+	ReverseGamma		= 1 << 7,
+	/** Potentially apply to HDR batch when composition is active*/
+	HDR					= 1 << 8
 };
 
 ENUM_CLASS_FLAGS(ESlateBatchDrawFlag);
@@ -162,6 +164,12 @@ enum class ESlateVertexRounding : uint8
 	Enabled
 };
 
+enum class ESlateViewportDynamicRange : uint8
+{
+	SDR,
+	HDR
+};
+
 class FSlateRenderBatch;
 
 
@@ -484,6 +492,14 @@ public:
 		return false;
 	}
 
+	/**
+	 * Does the texture contain SDR/HDR information
+	 */
+	virtual ESlateViewportDynamicRange GetViewportDynamicRange() const
+	{
+		return ESlateViewportDynamicRange::SDR;
+	}
+
 	/**
 	 * Performs any ticking necessary by this handle                   
 	 */
diff --git a/Engine/Source/Runtime/SlateRHIRenderer/Private/SlateRHIRenderer.cpp b/Engine/Source/Runtime/SlateRHIRenderer/Private/SlateRHIRenderer.cpp
--- a/Engine/Source/Runtime/SlateRHIRenderer/Private/SlateRHIRenderer.cpp
+++ b/Engine/Source/Runtime/SlateRHIRenderer/Private/SlateRHIRenderer.cpp
@@ -548,6 +548,12 @@ void FSlateRHIRenderer::ConditionalResizeViewport(FViewportInfo* ViewInfo, uint3
 
 		// Reset texture streaming texture updates.
 		ResumeTextureStreamingRenderTasks();
+
+		// when the window's state for HDR changed, we need to invalidate the window to make sure the viewport will end up in the appropriate FSlateBatchData, see FSlateElementBatcher::AddViewportElement
+		if (bHDRStale)
+		{
+			Window->Invalidate(EInvalidateWidgetReason::Paint);
+		}
 	}
 }
 
