Font Face: add a setting in the UFontFace class to allow adjusting the height of the strike brush relative to the full height of the font
When calculating the strike brush position, make sure to get the height of font based on its layout method instead of using font metrics only.

diff --git a/Engine/Source/Runtime/Engine/Classes/Engine/FontFace.h b/Engine/Source/Runtime/Engine/Classes/Engine/FontFace.h
--- a/Engine/Source/Runtime/Engine/Classes/Engine/FontFace.h
+++ b/Engine/Source/Runtime/Engine/Classes/Engine/FontFace.h
@@ -52,6 +52,7 @@ public:
 	virtual int32 GetAscendOverriddenValue() const override;
 	virtual bool IsDescendOverridden() const override;
 	virtual int32 GetDescendOverriddenValue() const override;
+	virtual int32 GetStrikeBrushHeightPercentage() const override;
 	virtual FFontFaceDataConstRef GetFontFaceData() const override;
 	//~ End IFontFaceInterface interface
 
@@ -97,6 +98,12 @@ public:
 	UPROPERTY(EditAnywhere, Category=FontFace, AdvancedDisplay)
 	bool bIsDescendOverridden;
 
+	/** The percentage of the font height to draw the strike brush at.
+	 * 0% is the bottom, 100% is the top.
+	 */
+	UPROPERTY(EditAnywhere, BlueprintReadWrite, Category = FontFace, AdvancedDisplay, meta = (ClampMin = "0", ClampMax = "100", ForceUnits = "%"))
+	int32 StrikeBrushHeightPercentage;
+
 	/** The data associated with the font face. This should always be filled in providing the source filename is valid. CacheSubFaces should be called after manually changing this property. */
 	FFontFaceDataRef FontFaceData;
 
diff --git a/Engine/Source/Runtime/Engine/Private/FontFace.cpp b/Engine/Source/Runtime/Engine/Private/FontFace.cpp
--- a/Engine/Source/Runtime/Engine/Private/FontFace.cpp
+++ b/Engine/Source/Runtime/Engine/Private/FontFace.cpp
@@ -22,6 +22,7 @@ UFontFace::UFontFace()
 	, bIsAscendOverridden(false)
 	, DescendOverriddenValue(0)
 	, bIsDescendOverridden(false)
+	, StrikeBrushHeightPercentage(60)
 	, FontFaceData(FFontFaceData::MakeFontFaceData())
 
 {
@@ -236,6 +237,11 @@ int32 UFontFace::GetDescendOverriddenValue() const
 	return DescendOverriddenValue;
 }
 
+int32 UFontFace::GetStrikeBrushHeightPercentage() const
+{
+	return StrikeBrushHeightPercentage;
+}
+
 FFontFaceDataConstRef UFontFace::GetFontFaceData() const
 {
 	return FontFaceData;
diff --git a/Engine/Source/Runtime/Slate/Private/Framework/Text/SlateTextUnderlineLineHighlighter.cpp b/Engine/Source/Runtime/Slate/Private/Framework/Text/SlateTextUnderlineLineHighlighter.cpp
--- a/Engine/Source/Runtime/Slate/Private/Framework/Text/SlateTextUnderlineLineHighlighter.cpp
+++ b/Engine/Source/Runtime/Slate/Private/Framework/Text/SlateTextUnderlineLineHighlighter.cpp
@@ -60,7 +60,7 @@ int32 ISlateTextLineHighlighter::OnPaint(const FPaintArgs& Args, const FTextLayo
 				);
 		}
 
-		// Draw underline
+		// Draw underline/strike
 		FSlateDrawElement::MakeBox(
 			OutDrawElements,
 			++LayerId,
diff --git a/Engine/Source/Runtime/SlateCore/Private/Fonts/CompositeFont.cpp b/Engine/Source/Runtime/SlateCore/Private/Fonts/CompositeFont.cpp
--- a/Engine/Source/Runtime/SlateCore/Private/Fonts/CompositeFont.cpp
+++ b/Engine/Source/Runtime/SlateCore/Private/Fonts/CompositeFont.cpp
@@ -167,6 +167,16 @@ int32 FFontData::GetDescendOverriddenValue() const
 	return 0;
 }
 
+int32 FFontData::GetStrikeBrushHeightPercentage() const
+{
+	if (FontFaceAsset)
+	{
+		const IFontFaceInterface* FontFace = CastChecked<const IFontFaceInterface>(FontFaceAsset);
+		return FontFace->GetStrikeBrushHeightPercentage();
+	}
+	return 60;
+}
+
 FFontFaceDataConstPtr FFontData::GetFontFaceData() const
 {
 	if (FontFaceAsset)
diff --git a/Engine/Source/Runtime/SlateCore/Private/Fonts/SlateFontRenderer.cpp b/Engine/Source/Runtime/SlateCore/Private/Fonts/SlateFontRenderer.cpp
--- a/Engine/Source/Runtime/SlateCore/Private/Fonts/SlateFontRenderer.cpp
+++ b/Engine/Source/Runtime/SlateCore/Private/Fonts/SlateFontRenderer.cpp
@@ -264,8 +264,8 @@ void FSlateFontRenderer::GetStrikeMetrics(const FSlateFontInfo& InFontInfo, cons
 	{
 		FreeTypeUtils::ApplySizeAndScale(FontFace, InFontInfo.Size, InScale);
 
-		// Place the strike 2/5th down the text by height (the code below does 3/5th as it counts from the bottom, not the top)
-		OutStrikeLinePos = FreeTypeUtils::Convert26Dot6ToRoundedPixel<int16>(FT_MulFix(FT_MulFix(FT_DivFix(FontFace->height, 5), 3), FontFace->size->metrics.y_scale));
+		// Place the strike at the given percent of the font's full height (the default is 60% from the bottom).
+		OutStrikeLinePos = FreeTypeUtils::Convert26Dot6ToRoundedPixel<int16>(FT_MulFix(FT_DivFix(FT_MulFix(FreeTypeUtils::GetHeight(FontFace, FontData.GetLayoutMethod()), FontData.GetStrikeBrushHeightPercentage()), 100), FontFace->size->metrics.y_scale));
 		OutStrikeLineThickness = FreeTypeUtils::Convert26Dot6ToRoundedPixel<int16>(FT_MulFix(FontFace->underline_thickness, FontFace->size->metrics.y_scale));
 	}
 	else
diff --git a/Engine/Source/Runtime/SlateCore/Public/Fonts/CompositeFont.h b/Engine/Source/Runtime/SlateCore/Public/Fonts/CompositeFont.h
--- a/Engine/Source/Runtime/SlateCore/Public/Fonts/CompositeFont.h
+++ b/Engine/Source/Runtime/SlateCore/Public/Fonts/CompositeFont.h
@@ -179,6 +179,9 @@ struct FFontData
 	/** Returns the overridden value of the descend. This value will be used only if IsDescendOverridden returns true. */
 	SLATECORE_API int32 GetDescendOverriddenValue() const;
 
+	/** Gets the percentage of the font height that's used as the height of the strike brush. */
+	SLATECORE_API int32 GetStrikeBrushHeightPercentage() const;
+
 	/** Get the data buffer containing the data for the current font face. */
 	SLATECORE_API FFontFaceDataConstPtr GetFontFaceData() const;
 
diff --git a/Engine/Source/Runtime/SlateCore/Public/Fonts/FontFaceInterface.h b/Engine/Source/Runtime/SlateCore/Public/Fonts/FontFaceInterface.h
--- a/Engine/Source/Runtime/SlateCore/Public/Fonts/FontFaceInterface.h
+++ b/Engine/Source/Runtime/SlateCore/Public/Fonts/FontFaceInterface.h
@@ -53,6 +53,9 @@ class IFontFaceInterface
 	/** Returns the overridden value of the descend. This value will be used only if IsDescendOverridden returns true. */
 	virtual int32 GetDescendOverriddenValue() const = 0;
 
+	/** Returns the percentage of the font height that's used as the height the strike brush. */
+	virtual int32 GetStrikeBrushHeightPercentage() const {return 60;};
+
 	/** Get the data buffer containing the data for the current font face. */
 	virtual FFontFaceDataConstRef GetFontFaceData() const = 0;
 };
