Fixed issue where window would grow over time when launching PIE in a new editor window

diff --git a/Engine/Source/Runtime/SlateCore/Private/Widgets/SWindow.cpp b/Engine/Source/Runtime/SlateCore/Private/Widgets/SWindow.cpp
--- a/Engine/Source/Runtime/SlateCore/Private/Widgets/SWindow.cpp
+++ b/Engine/Source/Runtime/SlateCore/Private/Widgets/SWindow.cpp
@@ -875,8 +875,8 @@ FSlateRect SWindow::GetClientRectInScreen() const
 	}
 
 	return GetRectInScreen()
-		.InsetBy(GetWindowBorderSize())
-		.InsetBy(FMargin(0.0f, TitleBarSize, 0.0f, 0.0f));
+		.InsetBy(GetWindowBorderSize() * GetDPIScaleFactor())
+		.InsetBy(FMargin(0.0f, TitleBarSize * GetDPIScaleFactor(), 0.0f, 0.0f));
 }
 
 UE::Slate::FDeprecateVector2DResult SWindow::GetClientSizeInScreen() const
