Constraints: move constraints related styles from core to control rig plugin

#rnx

diff --git a/Engine/Plugins/Animation/ControlRig/Source/ControlRigEditor/Private/ControlRigEditorStyle.h b/Engine/Plugins/Animation/ControlRig/Source/ControlRigEditor/Private/ControlRigEditorStyle.h
--- a/Engine/Plugins/Animation/ControlRig/Source/ControlRigEditor/Private/ControlRigEditorStyle.h
+++ b/Engine/Plugins/Animation/ControlRig/Source/ControlRigEditor/Private/ControlRigEditorStyle.h
@@ -160,6 +160,27 @@ public:
 			Set("ControlRig.Schematic.Link", new IMAGE_BRUSH_SVG("Slate/Link_Schematic", Icon128x128));
 		}
 
+		// Constraint Manager Icons
+		{
+			const FButtonStyle ConstraintOptionButton = FButtonStyle()
+				.SetNormal(FSlateRoundedBoxBrush(FStyleColors::Transparent, 4.f))
+				.SetHovered(FSlateRoundedBoxBrush(FStyleColors::Hover, 4.f))
+				.SetHoveredForeground(FLinearColor::White)
+				.SetPressed(FSlateRoundedBoxBrush(FStyleColors::Hover, 4.f))
+				.SetPressedForeground(FLinearColor::White)
+				.SetPressedPadding(FMargin(0.0, 1.0, 0.0, 0.0));
+
+			const FString EngineSlateDir = FPaths::EngineContentDir() / TEXT("Slate");
+			FContentRootBracket Bracket(this, EngineSlateDir);		
+
+			FComboButtonStyle ConstraintComboButton = FComboButtonStyle()
+				.SetButtonStyle(ConstraintOptionButton)
+				.SetDownArrowImage(IMAGE_BRUSH_SVG("Starship/Common/ellipsis-vertical-narrow", FVector2f(6.f, 15.f)));
+			ConstraintComboButton.ButtonStyle = ConstraintOptionButton;
+			
+			Set("ConstraintManager.ComboButton", ConstraintComboButton);
+		}
+		
 		FSlateStyleRegistry::RegisterSlateStyle(*this);
 	}
 
diff --git a/Engine/Plugins/Animation/ControlRig/Source/ControlRigEditor/Private/Editor/Constraints/SConstraintsWidget.cpp b/Engine/Plugins/Animation/ControlRig/Source/ControlRigEditor/Private/Editor/Constraints/SConstraintsWidget.cpp
--- a/Engine/Plugins/Animation/ControlRig/Source/ControlRigEditor/Private/Editor/Constraints/SConstraintsWidget.cpp
+++ b/Engine/Plugins/Animation/ControlRig/Source/ControlRigEditor/Private/Editor/Constraints/SConstraintsWidget.cpp
@@ -225,7 +225,7 @@ void SConstraintMenuEntry::Construct(
 		.AutoWidth()
 		[
 			SNew(SComboButton)
-			.ComboButtonStyle(FAppStyle::Get(), "ConstraintManager.ComboButton")
+			.ComboButtonStyle(FControlRigEditorStyle::Get(), "ConstraintManager.ComboButton")
 			.MenuContent()
 			[
 				GenerateConstraintDefaultWidget()
diff --git a/Engine/Source/Runtime/SlateCore/Private/Styling/StarshipCoreStyle.cpp b/Engine/Source/Runtime/SlateCore/Private/Styling/StarshipCoreStyle.cpp
--- a/Engine/Source/Runtime/SlateCore/Private/Styling/StarshipCoreStyle.cpp
+++ b/Engine/Source/Runtime/SlateCore/Private/Styling/StarshipCoreStyle.cpp
@@ -438,25 +438,8 @@ TSharedRef<ISlateStyle> FStarshipCoreStyle::Create()
 		Style->Set("Icons.Toolbar.Settings", new IMAGE_BRUSH_SVG("Starship/Common/settings", Icon20x20));
 
 		// Constraint Manager Icons
-		{
-			Style->Set("Icons.ConstraintManager.LookAt", new IMAGE_BRUSH_SVG("Starship/Common/LookAt", Icon16x16));
-			Style->Set("Icons.ConstraintManager.ParentHierarchy", new IMAGE_BRUSH_SVG("Starship/Common/ParentHierarchy", Icon16x16));
-			
-			const FButtonStyle ConstraintOptionButton = FButtonStyle()
-				.SetNormal(FSlateRoundedBoxBrush(FStyleColors::Transparent, 4.f))
-				.SetHovered(FSlateRoundedBoxBrush(FStyleColors::Hover, 4.f))
-				.SetHoveredForeground(FLinearColor::White)
-				.SetPressed(FSlateRoundedBoxBrush(FStyleColors::Hover, 4.f))
-				.SetPressedForeground(FLinearColor::White)
-				.SetPressedPadding(FMargin(0.0, 1.0, 0.0, 0.0));
-			
-			FComboButtonStyle ConstraintComboButton = FComboButtonStyle()
-				.SetButtonStyle(ConstraintOptionButton)
-				.SetDownArrowImage(IMAGE_BRUSH_SVG("Starship/Common/ellipsis-vertical-narrow", FVector2f(6.f, 15.f)));
-				ConstraintComboButton.ButtonStyle = ConstraintOptionButton;
-			
-			Style->Set( "ConstraintManager.ComboButton", ConstraintComboButton);
-		}
+		Style->Set("Icons.ConstraintManager.LookAt", new IMAGE_BRUSH_SVG("Starship/Common/LookAt", Icon16x16));
+		Style->Set("Icons.ConstraintManager.ParentHierarchy", new IMAGE_BRUSH_SVG("Starship/Common/ParentHierarchy", Icon16x16));
 
 		Style->Set("Icons.Role", new IMAGE_BRUSH_SVG("Starship/Common/Role", Icon16x16));
 		Style->Set("Icons.Merge", new IMAGE_BRUSH_SVG("Starship/Common/Merge", Icon16x16));
