Fix Post Buffer updates not triggering with global invalidation due cached batch buffer use not being considered.

diff --git a/Engine/Source/Runtime/SlateCore/Private/Rendering/ElementBatcher.cpp b/Engine/Source/Runtime/SlateCore/Private/Rendering/ElementBatcher.cpp
--- a/Engine/Source/Runtime/SlateCore/Private/Rendering/ElementBatcher.cpp
+++ b/Engine/Source/Runtime/SlateCore/Private/Rendering/ElementBatcher.cpp
@@ -568,6 +568,14 @@ void FSlateElementBatcher::AddCachedElements(FSlateCachedElementData& CachedElem
 	}
 	CachedElementData.ListsWithNewData.Empty();
 
+	for (const FSlateRenderBatch& CachedBatch : CachedElementData.GetCachedBatches())
+	{
+		if (const FSlateShaderResource* ShaderResource = CachedBatch.GetShaderResource())
+		{
+			UsedSlatePostBuffers |= ShaderResource->GetUsedSlatePostBuffers();
+		}
+	}
+
 	// Add the existing and new cached batches.
 	FSlateBatchData::AddCachedBatchesToBatchData(BatchData, BatchDataHDR, CachedElementData.GetCachedBatches());
 	CachedElementData.CleanupUnusedClipStates();
