Font settings: hide monospace width when monospacing is disabled.
#tests Played with bForceMonospaced (aka Monospacing) bool (from text, common text, or rich text style in a datatable), MonospacedWidth visibility was updated accordingly.

diff --git a/Engine/Source/Runtime/SlateCore/Public/Fonts/SlateFontInfo.h b/Engine/Source/Runtime/SlateCore/Public/Fonts/SlateFontInfo.h
--- a/Engine/Source/Runtime/SlateCore/Public/Fonts/SlateFontInfo.h
+++ b/Engine/Source/Runtime/SlateCore/Public/Fonts/SlateFontInfo.h
@@ -192,7 +192,7 @@ struct FSlateFontInfo
 	bool bForceMonospaced = false;
 
 	/** The uniform width to apply to all characters when bForceMonospaced is enabled, proportional of the font Size. */
-	UPROPERTY(EditAnywhere, BlueprintReadWrite, Category=SlateStyleRules, meta=(ClampMin=0))
+	UPROPERTY(EditAnywhere, BlueprintReadWrite, Category=SlateStyleRules, meta=(ClampMin=0, EditCondition = "bForceMonospaced==true", EditConditionHides))
 	float MonospacedWidth = 1.0f;
 
 #if WITH_EDITORONLY_DATA
