Added layout new layout sizing rule SizeRule_StretchContent
- widgets size starts at desired size and is adjusted proportionally based in available space

diff --git a/Engine/Source/Runtime/SlateCore/Public/Layout/LayoutUtils.h b/Engine/Source/Runtime/SlateCore/Public/Layout/LayoutUtils.h
--- a/Engine/Source/Runtime/SlateCore/Public/Layout/LayoutUtils.h
+++ b/Engine/Source/Runtime/SlateCore/Public/Layout/LayoutUtils.h
@@ -247,139 +247,307 @@ static void ArrangeSingleChild(EFlowDirection InFlowDirection, const FGeometry&
 }
 
 template<EOrientation Orientation, typename SlotType>
-static void ArrangeChildrenInStack(EFlowDirection InLayoutFlow, const TPanelChildren<SlotType>& Children, const FGeometry& AllottedGeometry, FArrangedChildren& ArrangedChildren, float InOffset, bool InAllowShrink)
+static void ArrangeChildrenInStack(EFlowDirection InLayoutFlow, const TPanelChildren<SlotType>& Children, const FGeometry& AllottedGeometry, FArrangedChildren& ArrangedChildren, float InOffset, bool bInAllowShrink)
 {
+
+	if (Children.Num() == 0)
+	{
+		return;
+	}
+	
 	// Allotted space will be given to fixed-size children first.
-	// Remaining space will be proportionately divided between stretch children (SizeRule_Stretch)
-	// based on their stretch coefficient
+	// Remaining space will be proportionately divided between stretch children (SizeRule_Stretch and SizeRule_StretchContent)
+	// based on their stretch coefficient.
 
-	if (Children.Num() > 0)
+	// Helper function to clamp to max size, if the constraint is set.
+	auto ClampToMaxSize = [](const float Size, const float MaxSize)
 	{
-		float StretchCoefficientTotal = 0.0f;
-		float FixedTotal = 0.0f;
-		float StretchSizeTotal = 0.0f;
-
-		bool bAnyChildVisible = false;
-		// Compute the sum of stretch coefficients (SizeRule_Stretch) and space required by fixed-size widgets (SizeRule_Auto),
-		// as well as the total desired size.
-		for (int32 ChildIndex = 0; ChildIndex < Children.Num(); ++ChildIndex)
+		return MaxSize > 0.0f ? FMath::Min(MaxSize, Size) : Size; 
+	};
+	
+	float GrowStretchCoefficientTotal = 0.0f;
+	float ShrinkStretchCoefficientTotal = 0.0f;
+	float FixedSizeTotal = 0.0f;
+	float StretchSizeTotal = 0.0f;
+
+	struct FStretchItem
+	{
+		// Size of the item
+		float Size = 0.0f;
+		// Initial size of the item 
+		float BasisSize = 0.0f;
+		// Max size constraint of the item.
+		float MaxSize = 0.0f;
+		// Stretch coefficient when the items are growing.
+		float GrowStretchValue = 0.0f;
+		// Stretch coefficient when the items are shrinking.
+		float ShrinkStretchValue = 0.0f;
+		// True if the constraints of the item has been satisfied.
+		bool bFrozen = false;
+		// Sizing rule for the item.
+		FSizeParam::ESizeRule SizeRule = FSizeParam::ESizeRule::SizeRule_Auto;
+	};
+	TArray<FStretchItem, TInlineAllocator<16>> StretchItems;
+	StretchItems.Init({}, Children.Num());
+	
+	bool bAnyChildVisible = false;
+	bool bAnyStretchContentItems = false;
+	bool bAnyStretchItems = false;
+	
+	// Compute the sum of stretch coefficients (SizeRule_Stretch & SizeRule_StretchContent) and space required by fixed-size widgets (SizeRule_Auto),
+	// as well as the total desired size.
+	for (int32 ChildIndex = 0; ChildIndex < Children.Num(); ++ChildIndex)
+	{
+		const SlotType& CurChild = Children[ChildIndex];
+
+		if (CurChild.GetWidget()->GetVisibility() != EVisibility::Collapsed)
 		{
-			const SlotType& CurChild = Children[ChildIndex];
+			bAnyChildVisible = true;
+			
+			// All widgets contribute their margin to the fixed space requirement
+			FixedSizeTotal += CurChild.GetPadding().template GetTotalSpaceAlong<Orientation>();
 
-			if (CurChild.GetWidget()->GetVisibility() != EVisibility::Collapsed)
-			{
-				bAnyChildVisible = true;
-				// All widgets contribute their margin to the fixed space requirement
-				FixedTotal += CurChild.GetPadding().template GetTotalSpaceAlong<Orientation>();
+			FVector2f ChildDesiredSize = CurChild.GetWidget()->GetDesiredSize();
+
+			// Auto-sized children contribute their desired size to the fixed space requirement
+			float ChildSize = (Orientation == Orient_Vertical)
+				? ChildDesiredSize.Y
+				: ChildDesiredSize.X;
 
-				FVector2f ChildDesiredSize = CurChild.GetWidget()->GetDesiredSize();
+			const float MaxSize = CurChild.GetMaxSize();
 
-				// Auto-sized children contribute their desired size to the fixed space requirement
-				const float ChildSize = (Orientation == Orient_Vertical)
-					? ChildDesiredSize.Y
-					: ChildDesiredSize.X;
+			FStretchItem& Item = StretchItems[ChildIndex]; 
+			Item.MaxSize = MaxSize;
+			Item.SizeRule = CurChild.GetSizeRule();
 
-				if (CurChild.GetSizeRule() == FSizeParam::SizeRule_Stretch)
+			// Clamp to the max size if it was specified
+			ChildSize = ClampToMaxSize(ChildSize, MaxSize);
+
+			if (CurChild.GetSizeRule() == FSizeParam::SizeRule_Stretch)
+			{
+				// Using same shrink and grow since otherwise the transition would be discontinuous as (reference) basis size is 0.
+				Item.GrowStretchValue = CurChild.GetSizeValue();
+				Item.ShrinkStretchValue = Item.GrowStretchValue;
+				Item.Size = 0.0f;
+				Item.BasisSize = 0.0f;
+
+				// For stretch children we save sum up the stretch coefficients
+				GrowStretchCoefficientTotal += Item.GrowStretchValue;
+				ShrinkStretchCoefficientTotal += Item.ShrinkStretchValue;
+				StretchSizeTotal += ChildSize;
+
+				bAnyStretchItems = true;
+			}
+			else if (CurChild.GetSizeRule() == FSizeParam::SizeRule_StretchContent)
+			{
+				// Allow separate values from grow and shrink, as the adjustment is relative to the child size. 
+				Item.GrowStretchValue = CurChild.GetSizeValue();
+				if constexpr (requires { CurChild.GetShrinkSizeValue(); })
 				{
-					// for stretch children we save sum up the stretch coefficients
-					StretchCoefficientTotal += CurChild.GetSizeValue();
-					StretchSizeTotal += ChildSize;
+					Item.ShrinkStretchValue = CurChild.GetShrinkSizeValue();
 				}
 				else
 				{
-
-					// Clamp to the max size if it was specified
-					float MaxSize = CurChild.GetMaxSize();
-					FixedTotal += MaxSize > 0.0f ? FMath::Min(MaxSize, ChildSize) : ChildSize;
+					Item.ShrinkStretchValue = Item.GrowStretchValue;
 				}
+				Item.Size = ChildSize;
+				Item.BasisSize = ChildSize;
+
+				// For sized stretch we sum to coefficients, but also treat the size as fixed.
+				GrowStretchCoefficientTotal += Item.GrowStretchValue;
+				ShrinkStretchCoefficientTotal += Item.ShrinkStretchValue;
+				StretchSizeTotal += ChildSize;
+
+				bAnyStretchContentItems = true;
+			}
+			else
+			{
+				FixedSizeTotal += ChildSize;
+
+				Item.GrowStretchValue = 0.0f;
+				Item.ShrinkStretchValue = 0.0f;
+				Item.Size = ChildSize;
+				Item.BasisSize = ChildSize;
 			}
 		}
+	}
 
-		if (!bAnyChildVisible)
+	if (!bAnyChildVisible)
+	{
+		return;
+	}
+
+	// When shrink is not allowed, we'll ensure to use all the space desired by the stretchable widgets.
+	const float MinAvailableSpace = bInAllowShrink ? 0.0f : StretchSizeTotal;
+
+	const float AllottedSize = Orientation == Orient_Vertical
+		? AllottedGeometry.GetLocalSize().Y
+		: AllottedGeometry.GetLocalSize().X;
+	
+	// The space available for SizeRule_Stretch and SizeRule_StretchContent widgets is any space that wasn't taken up by fixed-sized widgets.
+	float AvailableSpace = FMath::Max(MinAvailableSpace, AllottedSize - FixedSizeTotal);
+
+	// Apply SizeRule_Stretch.
+	if (bAnyStretchItems && GrowStretchCoefficientTotal > 0.0f)
+	{
+		// Distribute available space amongst the SizeRule_Stretch items proportional to the their stretch coefficient.
+		float UsedSpace = 0.0f;
+		for (FStretchItem& Item : StretchItems)
 		{
-			return;
+			if (Item.SizeRule == FSizeParam::SizeRule_Stretch)
+			{
+				// Stretch widgets get a fraction of the space remaining after all the fixed-space requirements are met.
+				// Supporting only one stretch value since otherwise the transition would be discontinuous as (reference) basis size is 0.
+				const float Size = AvailableSpace * Item.GrowStretchValue / GrowStretchCoefficientTotal;
+				
+				Item.Size = ClampToMaxSize(Size, Item.MaxSize);
+				
+				UsedSpace += Item.Size;
+			}
 		}
+		AvailableSpace -= UsedSpace;
+	}
 
-		//When shrink is not allowed, we'll ensure to use all the space desired by the stretchable widgets.
-		const float MinSize = InAllowShrink ? 0.0f : StretchSizeTotal;
+	// Apply SizeRule_StretchContent.
+	const bool bIsGrowing = AvailableSpace > StretchSizeTotal;
 
-		// The space available for SizeRule_Stretch widgets is any space that wasn't taken up by fixed-sized widgets.
-		const float NonFixedSpace = FMath::Max(MinSize, (Orientation == Orient_Vertical)
-			? AllottedGeometry.GetLocalSize().Y - FixedTotal
-			: AllottedGeometry.GetLocalSize().X - FixedTotal);
+	const bool bCanStretch = bIsGrowing
+		? (GrowStretchCoefficientTotal > 0.0f)
+		: (ShrinkStretchCoefficientTotal > 0.0f);
 
-		float PositionSoFar = 0.0f;
+	if (bAnyStretchContentItems && bCanStretch)
+	{
+		// Each StretchContent item starts at desired size and shrinks or grows based on available size.
+		// First, consume each items desired size from the available space.
+		// The remainder is corrected by growing ot shrinking the items.
+		int32 NumStretchContentItems = 0;
+		for (const FStretchItem& Item : StretchItems)
+		{
+			if (Item.SizeRule == FSizeParam::SizeRule_StretchContent)
+			{
+				AvailableSpace -= Item.Size;
+				NumStretchContentItems++;
+			}
+		}
 
-		// Now that we have the total fixed-space requirement and the total stretch coefficients we can
-		// arrange widgets top-to-bottom or left-to-right (depending on the orientation).
-		for (TPanelChildrenConstIterator<SlotType> It(Children, Orientation, InLayoutFlow); It; ++It)
+		// Run number of passes to satisfy the StretchContent constraints.
+		// On each pass distribute the available space to non-frozen items.
+		// An item gets frozen if it's (min/max) constraints are violated.
+		// This makes sure that we distribute all of the available space, event if small items collapse or if items clamp to max size.
+		// Each iteration should solve at least one constraint.
+		// In practice most layouts solve in 2 passes, we're capping to 5 iterations to keep things in fixed budget.
+		const int32 MaxPasses = FMath::Min(NumStretchContentItems, 5);
+		for (int32 Pass = 0; Pass < MaxPasses; Pass++)
 		{
-			const SlotType& CurChild = *It;
-			const EVisibility ChildVisibility = CurChild.GetWidget()->GetVisibility();
+			// On each pass calculate the total coefficients for valid items.
+			GrowStretchCoefficientTotal = 0.0f;
+			ShrinkStretchCoefficientTotal = 0.0f;
 
-			// Figure out the area allocated to the child in the direction of BoxPanel
-			// The area allocated to the slot is ChildSize + the associated margin.
-			float ChildSize = 0.0f;
-			if (ChildVisibility != EVisibility::Collapsed)
+			for (const FStretchItem& Item : StretchItems)
 			{
-				// The size of the widget depends on its size type
-				if (CurChild.GetSizeRule() == FSizeParam::SizeRule_Stretch)
+				if (Item.SizeRule == FSizeParam::SizeRule_StretchContent
+					&& !Item.bFrozen)
 				{
-					if (StretchCoefficientTotal > 0.0f)
-					{
-						// Stretch widgets get a fraction of the space remaining after all the fixed-space requirements are met
-						ChildSize = NonFixedSpace * CurChild.GetSizeValue() / StretchCoefficientTotal;
-					}
+					// Items are grown proportional to their stretch value.
+					GrowStretchCoefficientTotal += Item.GrowStretchValue;
+					// Items are shrank proportional to their stretch value and size. This is to emulate the flexbox behavior.
+					ShrinkStretchCoefficientTotal += Item.ShrinkStretchValue * Item.BasisSize;
 				}
-				else
-				{
-					const FVector2f ChildDesiredSize = CurChild.GetWidget()->GetDesiredSize();
+			}
 
-					// Auto-sized widgets get their desired-size value
-					ChildSize = (Orientation == Orient_Vertical)
-						? ChildDesiredSize.Y
-						: ChildDesiredSize.X;
-				}
+			const float StretchCoefficientTotal = bIsGrowing
+				? GrowStretchCoefficientTotal
+				: ShrinkStretchCoefficientTotal;
+			
+			// If all items are frozen, or no space to distribute, stop.
+			if (StretchCoefficientTotal < UE_KINDA_SMALL_NUMBER
+				|| FMath::Abs(AvailableSpace) < UE_KINDA_SMALL_NUMBER)
+			{
+				break;
+			}
 
-				// Clamp to the max size if it was specified
-				float MaxSize = CurChild.GetMaxSize();
-				if (MaxSize > 0.0f)
+			for (FStretchItem& Item : StretchItems)
+			{
+				if (Item.SizeRule == FSizeParam::SizeRule_StretchContent
+					&& !Item.bFrozen)
 				{
-					ChildSize = FMath::Min(MaxSize, ChildSize);
+					const float SizeAdjust = bIsGrowing
+						? (AvailableSpace * (Item.GrowStretchValue / GrowStretchCoefficientTotal))
+						: (AvailableSpace * (Item.ShrinkStretchValue * Item.BasisSize / ShrinkStretchCoefficientTotal));
+					
+					constexpr float MinSize = 0.0f;
+					const float MaxSize = Item.MaxSize;
+					const bool bHasMaxConstraint = MaxSize > 0.0f; 
+
+					if ((Item.Size + SizeAdjust) <= MinSize)
+					{
+						// Adjustment goes past min constraint, apply what we can and freeze since the item cannot change anymore.
+						AvailableSpace -= MinSize - Item.Size;
+						Item.Size = MinSize;
+						Item.bFrozen = true;
+					}
+					else if (bHasMaxConstraint
+						&& (Item.Size + SizeAdjust) >= MaxSize)
+					{
+						// Adjustment goes past max constraint, apply what we can and freeze since the item cannot change anymore.
+						AvailableSpace -= MaxSize - Item.Size;
+						Item.Size = MaxSize;
+						Item.bFrozen = true;
+					}
+					else
+					{
+						// Within constraints, adjust.
+						AvailableSpace -= SizeAdjust;
+						Item.Size += SizeAdjust;
+					}
 				}
 			}
+		}
+	}
+	
+	// Now that we have the satisfied size requirements we can
+	// arrange widgets top-to-bottom or left-to-right (depending on the orientation).
+	float PositionSoFar = 0.0f;
 
-			const FMargin SlotPadding(LayoutPaddingWithFlow(InLayoutFlow, CurChild.GetPadding()));
+	for (TPanelChildrenConstIterator<SlotType> It(Children, Orientation, InLayoutFlow); It; ++It)
+	{
+		const SlotType& CurChild = *It;
+		const EVisibility ChildVisibility = CurChild.GetWidget()->GetVisibility();
 
-			FVector2f SlotSize = (Orientation == Orient_Vertical)
-				? FVector2f(AllottedGeometry.GetLocalSize().X, ChildSize + SlotPadding.template GetTotalSpaceAlong<Orient_Vertical>())
-				: FVector2f(ChildSize + SlotPadding.template GetTotalSpaceAlong<Orient_Horizontal>(), AllottedGeometry.GetLocalSize().Y);
+		// Figure out the area allocated to the child in the direction of BoxPanel
+		// The area allocated to the slot is ChildSize + the associated margin.
+		const float ChildSize = StretchItems[It.GetIndex()].Size;
 
-			// Figure out the size and local position of the child within the slot			
-			AlignmentArrangeResult XAlignmentResult = AlignChild<Orient_Horizontal>(InLayoutFlow, SlotSize.X, CurChild, SlotPadding);
-			AlignmentArrangeResult YAlignmentResult = AlignChild<Orient_Vertical>(SlotSize.Y, CurChild, SlotPadding);
+		const FMargin SlotPadding(LayoutPaddingWithFlow(InLayoutFlow, CurChild.GetPadding()));
 
-			const FVector2f LocalPosition = (Orientation == Orient_Vertical)
-				? FVector2f(XAlignmentResult.Offset, PositionSoFar + YAlignmentResult.Offset + InOffset)
-				: FVector2f(PositionSoFar + XAlignmentResult.Offset + InOffset, YAlignmentResult.Offset);
+		FVector2f SlotSize = (Orientation == Orient_Vertical)
+			? FVector2f(AllottedGeometry.GetLocalSize().X, ChildSize + SlotPadding.template GetTotalSpaceAlong<Orient_Vertical>())
+			: FVector2f(ChildSize + SlotPadding.template GetTotalSpaceAlong<Orient_Horizontal>(), AllottedGeometry.GetLocalSize().Y);
 
-			const FVector2f LocalSize = FVector2f(XAlignmentResult.Size, YAlignmentResult.Size);
+		// Figure out the size and local position of the child within the slot			
+		const AlignmentArrangeResult XAlignmentResult = AlignChild<Orient_Horizontal>(InLayoutFlow, SlotSize.X, CurChild, SlotPadding);
+		const AlignmentArrangeResult YAlignmentResult = AlignChild<Orient_Vertical>(SlotSize.Y, CurChild, SlotPadding);
 
-			// Add the information about this child to the output list (ArrangedChildren)
-			ArrangedChildren.AddWidget(ChildVisibility, AllottedGeometry.MakeChild(
-				// The child widget being arranged
-				CurChild.GetWidget(),
-				// Child's local position (i.e. position within parent)
-				LocalPosition,
-				// Child's size
-				LocalSize
-			));
+		const FVector2f LocalPosition = (Orientation == Orient_Vertical)
+			? FVector2f(XAlignmentResult.Offset, PositionSoFar + YAlignmentResult.Offset + InOffset)
+			: FVector2f(PositionSoFar + XAlignmentResult.Offset + InOffset, YAlignmentResult.Offset);
 
-			if (ChildVisibility != EVisibility::Collapsed)
-			{
-				// Offset the next child by the size of the current child and any post-child (bottom/right) margin
-				PositionSoFar += (Orientation == Orient_Vertical) ? SlotSize.Y : SlotSize.X;
-			}
+		const FVector2f LocalSize = FVector2f(XAlignmentResult.Size, YAlignmentResult.Size);
+
+		// Add the information about this child to the output list (ArrangedChildren)
+		ArrangedChildren.AddWidget(ChildVisibility, AllottedGeometry.MakeChild(
+			// The child widget being arranged
+			CurChild.GetWidget(),
+			// Child's local position (i.e. position within parent)
+			LocalPosition,
+			// Child's size
+			LocalSize
+		));
+
+		if (ChildVisibility != EVisibility::Collapsed)
+		{
+			// Offset the next child by the size of the current child and any post-child (bottom/right) margin
+			PositionSoFar += (Orientation == Orient_Vertical) ? SlotSize.Y : SlotSize.X;
 		}
 	}
 }
diff --git a/Engine/Source/Runtime/SlateCore/Public/Types/SlateStructs.h b/Engine/Source/Runtime/SlateCore/Public/Types/SlateStructs.h
--- a/Engine/Source/Runtime/SlateCore/Public/Types/SlateStructs.h
+++ b/Engine/Source/Runtime/SlateCore/Public/Types/SlateStructs.h
@@ -73,19 +73,22 @@ private:
  *
  * Describes a way in which a parent widget allocates available space to its child widgets.
  *
- * When SizeRule is SizeRule_Auto, the widget's DesiredSize will be used as the space required.
- * When SizeRule is SizeRule_Stretch, the available space will be distributed proportionately between
- * peer Widgets depending on the Value property. Available space is space remaining after all the
- * peers' SizeRule_Auto requirements have been satisfied.
+ * When SizeRule is SizeRule_Auto, the required space is the widget's DesiredSize.
+ * When SizeRule is SizeRule_Stretch, the required space is the available space distributed proportionately between peer Widgets.
+ * When SizeRule is SizeRule_StretchContent, the required space is widget's content size adjusted proportionally to fit the available space.
  *
- * FSizeParam cannot be constructed directly - see FStretch, FAuto, and FAspectRatio
+ * Available space is space remaining after all the peers' SizeRule_Auto requirements have been satisfied.
+ * The available space is distributed proportionally between the peer widgets depending on the Value property.
+ *
+ * FSizeParam cannot be constructed directly - see FStretch, FStretchContent, FAuto, and FAspectRatio
  */
 struct FSizeParam
 {
 	enum ESizeRule
 	{
 		SizeRule_Auto,
-		SizeRule_Stretch
+		SizeRule_Stretch,
+		SizeRule_StretchContent,
 	};
 	
 	/** The sizing rule to use. */
@@ -94,10 +97,18 @@ struct FSizeParam
 	/**
 	 * The actual value this size parameter stores.
 	 *
-	 * This value can be driven by a delegate. It is only used for the Stretch mode.
+	 * This value can be driven by a delegate. It is only used for the Stretch and StretchContent modes.
 	 */
 	TAttribute<float> Value;
-	
+
+	/**
+	 * The actual value this size parameter stores, used for shrinking.
+	 * Treated as unused, if set to negative value.
+	 *
+	 * This value can be driven by a delegate. It is only used for the StretchContent mode.
+	 */
+	TAttribute<float> ShrinkValue;
+
 protected:
 
 	/**
@@ -107,9 +118,10 @@ protected:
 	 *
 	 * @see FAspectRatio, FAuto, FStretch
 	 */
-	FSizeParam( ESizeRule InTypeOfSize, const TAttribute<float>& InValue )
+	FSizeParam( ESizeRule InTypeOfSize, const TAttribute<float>& InValue, const TAttribute<float>& InShrinkValue )
 		: SizeRule(InTypeOfSize)
 		, Value(InValue)
+		, ShrinkValue(InShrinkValue)
 	{ }
 };
 
@@ -117,31 +129,50 @@ protected:
 /**
  * Structure for size parameters with SizeRule = SizeRule_Stretch.
  *
- * @see FAspectRatio
- * @see FAuto
+ * @see FStretchContent, FAspectRatio, FAuto
  */
 struct FStretch
 	: public FSizeParam
 {
 	FStretch( const TAttribute<float>& StretchAmount )
-		: FSizeParam(SizeRule_Stretch, StretchAmount)
+		: FSizeParam(SizeRule_Stretch, StretchAmount, StretchAmount)
 	{ }
 
 	FStretch( )
-		: FSizeParam(SizeRule_Stretch, 1.0f)
+		: FSizeParam(SizeRule_Stretch, 1.0f, 1.0f)
 	{ }
 };
 
+/**
+ * Structure for size parameters with SizeRule = SizeRule_StretchContent.
+ *
+ * @see FStretch, FAspectRatio, FAuto
+ */
+struct FStretchContent
+	: public FSizeParam
+{
+	FStretchContent( const TAttribute<float>& StretchAmount )
+		: FSizeParam(SizeRule_StretchContent, StretchAmount, StretchAmount)
+	{ }
+
+	FStretchContent( const TAttribute<float>& GrowStretchAmount, const TAttribute<float>& ShrinkStretchAmount )
+   		: FSizeParam(SizeRule_StretchContent, GrowStretchAmount, ShrinkStretchAmount)
+   	{ }
+
+	FStretchContent( )
+		: FSizeParam(SizeRule_StretchContent, 1.0f, 1.0f)
+	{ }
+};
 
 /**
  * Structure for size parameters with SizeRule = SizeRule_Auto.
  *
- * @see FAspectRatio, FStretch
+ * @see FAspectRatio, FStretch, FStretchContent
  */
 struct FAuto
 	: public FSizeParam
 {
 	FAuto()
-		: FSizeParam(SizeRule_Auto, 0.0f)
+		: FSizeParam(SizeRule_Auto, 0.0f, 0.0f)
 	{ }
 };
diff --git a/Engine/Source/Runtime/SlateCore/Public/Widgets/SBoxPanel.h b/Engine/Source/Runtime/SlateCore/Public/Widgets/SBoxPanel.h
--- a/Engine/Source/Runtime/SlateCore/Public/Widgets/SBoxPanel.h
+++ b/Engine/Source/Runtime/SlateCore/Public/Widgets/SBoxPanel.h
@@ -47,6 +47,7 @@ protected:
 			: TBasicLayoutWidgetSlot<SlotType>(HAlign_Fill, VAlign_Fill)
 			, SizeRule(FSizeParam::SizeRule_Stretch)
 			, SizeValue(*this, 1.f)
+			, ShrinkSizeValue(*this, -1.f) // Negative means unset, SizeValue is used instead.
 			, MaxSize(*this, 0.0f)
 		{ }
 
@@ -70,6 +71,8 @@ protected:
 			SLATE_ADD_SLOT_ATTRIBUTE_DEFINITION_WITH_NAME(TSlot<SlotType>, AttributeInitializer, "Slot.MaxSize", MaxSize, EInvalidateWidgetReason::Layout);
 			SLATE_ADD_SLOT_ATTRIBUTE_DEFINITION_WITH_NAME(TSlot<SlotType>, AttributeInitializer, "Slot.SizeValue", SizeValue, EInvalidateWidgetReason::Layout)
 				.UpdatePrerequisite("Slot.MaxSize");
+			SLATE_ADD_SLOT_ATTRIBUTE_DEFINITION_WITH_NAME(TSlot<SlotType>, AttributeInitializer, "Slot.ShrinkSizeValue", ShrinkSizeValue, EInvalidateWidgetReason::Layout)
+				.UpdatePrerequisite("Slot.MaxSize");
 		}
 
 		/** Get the space rule this slot should occupy along panel's direction. */
@@ -78,12 +81,25 @@ protected:
 			return SizeRule;
 		}
 
-		/** Get the space rule value this slot should occupy along panel's direction. */
+		/**
+		 * Get the size parameter for the space rule.
+		 * Used for size rule SizeRule_Stretch and SizeRule_StretchContent. 
+		 */
 		float GetSizeValue() const
 		{
 			return SizeValue.Get();
 		}
 
+		/**
+		 * Get the size parameter for the space rule, used when the slot size is shrinking below desired size.
+		 * Used for size rule SizeRule_StretchContent.
+		 */
+		float GetShrinkSizeValue() const
+		{
+			const float Value = ShrinkSizeValue.Get();
+			return Value < 0.0f ? GetSizeValue() : Value;
+		}
+
 		/** Get the max size the slot can be.*/
 		float GetMaxSize() const
 		{
@@ -91,23 +107,49 @@ protected:
 		}
 
 	public:
-		/** Set the size Param of the slot, It could be a FStretch or a FAuto. */
+		/** Set the size Param of the slot, It could be a FStretch, FStretchContent, or a FAuto. */
 		void SetSizeParam(FSizeParam InSizeParam)
 		{
 			SizeRule = InSizeParam.SizeRule;
 			SizeValue.Assign(*this, MoveTemp(InSizeParam.Value));
+			if (InSizeParam.ShrinkValue.IsSet())
+			{
+				ShrinkSizeValue.Assign(*this, MoveTemp(InSizeParam.ShrinkValue));
+			}
+			else
+			{
+				// Negative value means to use size value for shrink and grow.
+				ShrinkSizeValue.Set(*this, -1.f);
+			}
 		}
 
-		/** The widget's DesiredSize will be used as the space required. */
+		/**
+		 * The widget's DesiredSize will be used as the space required.
+		 */
 		void SetSizeToAuto()
 		{
 			SetSizeParam(FAuto());
 		}
 
-		/** The available space will be distributed proportionately. */
-		void SetSizeToStretch(TAttribute<float> StretchCoefficient)
+		/**
+		 * The available space will be distributed proportionately to each slots stretch coefficient.
+		 * A slot with coefficient of 2 will get assigned twice as much available space as slot with coefficient 1. 
+		 * @param InStretchCoefficient Stretch coefficient for this slot.  
+		 */
+		void SetSizeToStretch(TAttribute<float> InStretchCoefficient)
+		{
+			SetSizeParam(FStretch(MoveTemp(InStretchCoefficient)));
+		}
+
+		/**
+		 * The widget's content size is adjusted proportionally to fit the available space.
+		 * The slots size starts at DesiredSize, and a slot with coefficient of 2 will get adjusted twice as much as slot with coefficient 1 to fit the available space.
+		 * @param InStretchCoefficient Stretch coefficient for this slot.
+		 * @param InShrinkStretchCoefficient If specified, this stretch coefficient is used when the slots is shrinking below desired size. Otherwise InStretchCoefficient is used for both shrink and grow.  
+		 */
+		void SetSizeToStretchContent(TAttribute<float> InStretchCoefficient, TAttribute<float> InShrinkStretchCoefficient = TAttribute<float>())
 		{
-			SetSizeParam(FStretch(MoveTemp(StretchCoefficient)));
+			SetSizeParam(FStretchContent(MoveTemp(InStretchCoefficient), MoveTemp(InShrinkStretchCoefficient)));
 		}
 
 		/** Set the max size in SlateUnit this slot can be. */
@@ -117,20 +159,15 @@ protected:
 		}
 
 	private:
-		/**
-		 * How much space this slot should occupy along panel's direction.
-		 * When SizeRule is SizeRule_Auto, the widget's DesiredSize will be used as the space required.
-		 * When SizeRule is SizeRule_Stretch, the available space will be distributed proportionately between
-		 * peer Widgets depending on the Value property. Available space is space remaining after all the
-		 * peers' SizeRule_Auto requirements have been satisfied.
-		 */
-
-		/** The sizing rule to use. */
+		/** The sizing rule to use, see ESizeRule for more info how the different rules work. */
 		FSizeParam::ESizeRule SizeRule;
 
 		/** The actual value this size parameter stores. */
 		typename TBasicLayoutWidgetSlot<SlotType>::template TSlateSlotAttribute<float> SizeValue;
 
+		/** The actual value this size parameter stores, used for shrinking (negative if not defined, use SizeValue). */
+		typename TBasicLayoutWidgetSlot<SlotType>::template TSlateSlotAttribute<float> ShrinkSizeValue;
+
 		/** The max size that this slot can be (0 if no max) */
 		typename TBasicLayoutWidgetSlot<SlotType>::template TSlateSlotAttribute<float> MaxSize;
 	};
@@ -235,18 +272,38 @@ public:
 	{
 	public:
 		SLATE_SLOT_BEGIN_ARGS(FSlot, SBoxPanel::TSlot<FSlot>)
-			/** The widget's DesiredSize will be used as the space required. */
+			/**
+			 * The widget's DesiredSize will be used as the space required.
+			 */
 			FSlotArguments& AutoWidth()
 			{
 				_SizeParam = FAuto();
 				return Me();
 			}
-			/** The available space will be distributed proportionately. */
+		
+			/**
+			 * The available space will be distributed proportionately to each slots stretch coefficient.
+			 * A slot with coefficient of 2 will get assigned twice as much available space as slot with coefficient 1. 
+			 * @param InStretchCoefficient Stretch coefficient for this slot.  
+			 */
 			FSlotArguments& FillWidth(TAttribute<float> InStretchCoefficient)
 			{
 				_SizeParam = FStretch(MoveTemp(InStretchCoefficient));
 				return Me();
 			}
+		
+			/**
+			 * The widget's content size is adjusted proportionally to fit the available space.
+			 * The slots size starts at DesiredSize, and a slot with coefficient of 2 will get adjusted twice as much as slot with coefficient 1 to fit the available space.
+			 * @param InStretchCoefficient Stretch coefficient for this slot.
+			 * @param InShrinkStretchCoefficient If specified, this stretch coefficient is used when the slots is shrinking below desired size. Otherwise InStretchCoefficient is used for both shrink and grow.  
+			 */
+			FSlotArguments& FillContentWidth(TAttribute<float> InStretchCoefficient, TAttribute<float> InShrinkStretchCoefficient = TAttribute<float>())
+			{
+				_SizeParam = FStretchContent(MoveTemp(InStretchCoefficient), MoveTemp(InShrinkStretchCoefficient));
+				return Me();
+			}
+		
 			/** Set the max size in SlateUnit this slot can be. */
 			FSlotArguments& MaxWidth(TAttribute<float> InMaxWidth)
 			{
@@ -255,18 +312,35 @@ public:
 			}
 		SLATE_SLOT_END_ARGS()
 
-		/** The widget's DesiredSize will be used as the space required. */
+		/**
+		 * The widget's DesiredSize will be used as the space required.
+		 */
 		void SetAutoWidth()
 		{
 			SetSizeToAuto();
 		}
 
-		/** The available space will be distributed proportionately. */
+		/**
+		 * The available space will be distributed proportionately to each slots stretch coefficient.
+		 * A slot with coefficient of 2 will get assigned twice as much available space as slot with coefficient 1. 
+		 * @param InStretchCoefficient Stretch coefficient for this slot.  
+		 */
 		void SetFillWidth(TAttribute<float> InStretchCoefficient)
 		{
 			SetSizeToStretch(MoveTemp(InStretchCoefficient));
 		}
 
+		/**
+		 * The widget's content size is adjusted proportionally to fit the available space.
+		 * The slots size starts at DesiredSize, and a slot with coefficient of 2 will get adjusted twice as much as slot with coefficient 1 to fit the available space.
+		 * @param InStretchCoefficient Stretch coefficient for this slot.
+		 * @param InShrinkStretchCoefficient If specified, this stretch coefficient is used when the slots is shrinking below desired size. Otherwise InStretchCoefficient is used for both shrink and grow.  
+		 */
+		void SetFillContentWidth(TAttribute<float> InStretchCoefficient, TAttribute<float> InShrinkStretchCoefficient = TAttribute<float>())
+		{
+			SetSizeToStretchContent(MoveTemp(InStretchCoefficient), MoveTemp(InShrinkStretchCoefficient));
+		}
+
 		/** Set the max size in SlateUnit this slot can be. */
 		void SetMaxWidth(TAttribute<float> InMaxWidth)
 		{
@@ -338,18 +412,39 @@ public:
 	{
 	public:
 		SLATE_SLOT_BEGIN_ARGS(FSlot, SBoxPanel::TSlot<FSlot>)
-			/** The widget's DesiredSize will be used as the space required. */
+		
+			/**
+			 * The widget's DesiredSize will be used as the space required.
+			 */
 			FSlotArguments& AutoHeight()
 			{
 				_SizeParam = FAuto();
 				return Me();
 			}
-			/** The available space will be distributed proportionately. */
+		
+			/**
+			 * The available space will be distributed proportionately to each slots stretch coefficient.
+			 * A slot with coefficient of 2 will get assigned twice as much available space as slot with coefficient 1. 
+			 * @param InStretchCoefficient Stretch coefficient for this slot.  
+			 */
 			FSlotArguments& FillHeight(TAttribute<float> InStretchCoefficient)
 			{
 				_SizeParam = FStretch(MoveTemp(InStretchCoefficient));
 				return Me();
 			}
+		
+			/**
+			 * The widget's content size is adjusted proportionally to fit the available space.
+			 * The slots size starts at DesiredSize, and a slot with coefficient of 2 will get adjusted twice as much as slot with coefficient 1 to fit the available space.
+			 * @param InStretchCoefficient Stretch coefficient for this slot.
+			 * @param InShrinkStretchCoefficient If specified, this stretch coefficient is used when the slots is shrinking below desired size. Otherwise InStretchCoefficient is used for both shrink and grow.  
+			 */
+			FSlotArguments& FillContentHeight(TAttribute<float> InStretchCoefficient, TAttribute<float> InShrinkStretchCoefficient = TAttribute<float>())
+			{
+				_SizeParam = FStretchContent(MoveTemp(InStretchCoefficient), MoveTemp(InShrinkStretchCoefficient));
+				return Me();
+			}
+			
 			/** Set the max size in SlateUnit this slot can be. */
 			FSlotArguments& MaxHeight(TAttribute<float> InMaxHeight)
 			{
@@ -358,18 +453,35 @@ public:
 			}
 		SLATE_SLOT_END_ARGS()
 
-		/** The widget's DesiredSize will be used as the space required. */
+		/**
+		 * The widget's DesiredSize will be used as the space required.
+		 */
 		void SetAutoHeight()
 		{
 			SetSizeToAuto();
 		}
 
-		/** The available space will be distributed proportionately. */
+		/**
+		 * The available space will be distributed proportionately to each slots stretch coefficient.
+		 * A slot with coefficient of 2 will get assigned twice as much available space as slot with coefficient 1. 
+		 * @param InStretchCoefficient Stretch coefficient for this slot.  
+		 */
 		void SetFillHeight(TAttribute<float> InStretchCoefficient)
 		{
 			SetSizeToStretch(MoveTemp(InStretchCoefficient));
 		}
 
+		/**
+		 * The widget's content size is adjusted proportionally to fit the available space.
+		 * The slots size starts at DesiredSize, and a slot with coefficient of 2 will get adjusted twice as much as slot with coefficient 1 to fit the available space.
+		 * @param InStretchCoefficient Stretch coefficient for this slot.
+		 * @param InShrinkStretchCoefficient If specified, this stretch coefficient is used when the slots is shrinking below desired size. Otherwise InStretchCoefficient is used for both shrink and grow.  
+		 */
+		void SetFillContentHeight(TAttribute<float> InStretchCoefficient, TAttribute<float> InShrinkStretchCoefficient = TAttribute<float>())
+		{
+			SetSizeToStretchContent(MoveTemp(InStretchCoefficient), MoveTemp(InShrinkStretchCoefficient));
+		}
+
 		/** Set the max size in SlateUnit this slot can be. */
 		void SetMaxHeight(TAttribute<float> InMaxHeight)
 		{
@@ -442,18 +554,37 @@ public:
 	{
 	public:
 		SLATE_SLOT_BEGIN_ARGS(FSlot, SBoxPanel::TSlot<FSlot>)
-			/** The widget's DesiredSize will be used as the space required. */
+			/**
+			 * The widget's DesiredSize will be used as the space required.
+			 */
 			FSlotArguments& AutoSize()
 			{
 				_SizeParam = FAuto();
 				return Me();
 			}
-			/** The available space will be distributed proportionately. */
+		
+			/**
+			 * The available space will be distributed proportionately to each slots stretch coefficient.
+			 * A slot with coefficient of 2 will get assigned twice as much available space as slot with coefficient 1. 
+			 * @param InStretchCoefficient Stretch coefficient for this slot.  
+			 */
 			FSlotArguments& FillSize(TAttribute<float> InStretchCoefficient)
 			{
 				_SizeParam = FStretch(MoveTemp(InStretchCoefficient));
 				return Me();
 			}
+		
+			/**
+			 * The widget's content size is adjusted proportionally to fit the available space.
+			 * The slots size starts at DesiredSize, and a slot with coefficient of 2 will get adjusted twice as much as slot with coefficient 1 to fit the available space.
+			 * @param InStretchCoefficient Stretch coefficient for this slot.
+			 * @param InShrinkStretchCoefficient If specified, this stretch coefficient is used when the slots is shrinking below desired size. Otherwise InStretchCoefficient is used for both shrink and grow.  
+			 */
+			FSlotArguments& FillContentSize(TAttribute<float> InStretchCoefficient, TAttribute<float> InShrinkStretchCoefficient = TAttribute<float>())
+			{
+				_SizeParam = FStretchContent(MoveTemp(InStretchCoefficient), MoveTemp(InShrinkStretchCoefficient));
+				return Me();
+			}
 			/** Set the max size in SlateUnit this slot can be. */
 			FSlotArguments& MaxSize(TAttribute<float> InMaxHeight)
 			{
