Reduce vertical space consumed by title bar

The intent of this change is twofold. First, it fixes a bug where the min/max/close buttons overlapped the row of tabs in windows with a main menu. Second, it reduces the vertical space consumed by our windows in the area around the title bar, main menu, and tabs.

* Reduced height of SVG images used for the min/max/restore/close window buttons from 34 pixels to 23 pixels. This avoids an overlap of the min/max/close buttons and the row of tabs in windows with a main menu.
* Rebalanced padding around the Unreal icon in windows with or without a main menu. There is now at most 1 pixel more spacing below the logo than above where before it could be 3 to 8 pixels more below than above.
* Reduced padding between tabs and main menu.

Note that these changes are all within #if !PLATFORM_MAC, so the Mac should not be affected.

diff --git a/Engine/Source/Runtime/Slate/Private/Framework/Application/SWindowTitleBar.cpp b/Engine/Source/Runtime/Slate/Private/Framework/Application/SWindowTitleBar.cpp
--- a/Engine/Source/Runtime/Slate/Private/Framework/Application/SWindowTitleBar.cpp
+++ b/Engine/Source/Runtime/Slate/Private/Framework/Application/SWindowTitleBar.cpp
@@ -199,6 +199,7 @@ void SWindowTitleBar::MakeTitleBarContentWidgets( TSharedPtr< SWidget >& OutLeft
 			.AutoWidth()
 			.HAlign(HAlign_Left)
 			.VAlign(VAlign_Top)
+			.Padding(FMargin(0, 0, 0, 4.0f))
 			[
 				SAssignNew(AppIconWidget, SAppIconWidget)
 				.IconColorAndOpacity(this, &SWindowTitleBar::GetWindowTitleContentColor)
diff --git a/Engine/Source/Runtime/Slate/Private/Framework/Docking/SDockingTabStack.cpp b/Engine/Source/Runtime/Slate/Private/Framework/Docking/SDockingTabStack.cpp
--- a/Engine/Source/Runtime/Slate/Private/Framework/Docking/SDockingTabStack.cpp
+++ b/Engine/Source/Runtime/Slate/Private/Framework/Docking/SDockingTabStack.cpp
@@ -883,7 +883,8 @@ void SDockingTabStack::ReserveSpaceForWindowChrome(EChromeElement Element, bool
 		ControlsPadding = FMargin(67.0f, 0, 0, 0);
 	}
 #else
-	static const float TopPaddingForMenuBar = 27.0f;
+	static const float TopPaddingForMenuBar = 25.0f;
+
 	static const float LeftPaddingForIcon = FSlateApplication::Get().GetAppIcon()->GetImageSize().X;
 	// If we are including top padding for the menu bar we do not need to pad the outer sides since we will be below the left icon and the right controls.
 	if (bIncludePaddingForMenuBar)
@@ -893,7 +894,7 @@ void SDockingTabStack::ReserveSpaceForWindowChrome(EChromeElement Element, bool
 	}
 	else
 	{
-		ControlsPadding = FMargin(8.f, 10.f, 128.f, 0.f);
+		ControlsPadding = FMargin(8.f, 2.f, 128.f, 0.f);
 		IconPadding = FMargin(25.f, bOnlyMinorTabs ? 5.f : 0.f, 0.f, 0.f);
 	}
 #endif
diff --git a/Engine/Source/Runtime/SlateCore/Private/Styling/StarshipCoreStyle.cpp b/Engine/Source/Runtime/SlateCore/Private/Styling/StarshipCoreStyle.cpp
--- a/Engine/Source/Runtime/SlateCore/Private/Styling/StarshipCoreStyle.cpp
+++ b/Engine/Source/Runtime/SlateCore/Private/Styling/StarshipCoreStyle.cpp
@@ -742,24 +742,24 @@ TSharedRef<ISlateStyle> FStarshipCoreStyle::Create()
 	{
 #if !PLATFORM_MAC
 		const FButtonStyle MinimizeButtonStyle = FButtonStyle(NoBorder)
-			.SetNormal( IMAGE_BRUSH_SVG("Starship/CoreWidgets/Window/minimize", FVector2f(42.0f, 34.0f), FStyleColors::Foreground))
-			.SetHovered(IMAGE_BRUSH_SVG("Starship/CoreWidgets/Window/minimize", FVector2f(42.0f, 34.0f), FStyleColors::ForegroundHover))
-			.SetPressed(IMAGE_BRUSH_SVG("Starship/CoreWidgets/Window/minimize", FVector2f(42.0f, 34.0f), FStyleColors::Foreground));
+			.SetNormal( IMAGE_BRUSH_SVG("Starship/CoreWidgets/Window/minimize", FVector2f(42.0f, 23.0f), FStyleColors::Foreground))
+			.SetHovered(IMAGE_BRUSH_SVG("Starship/CoreWidgets/Window/minimize", FVector2f(42.0f, 23.0f), FStyleColors::ForegroundHover))
+			.SetPressed(IMAGE_BRUSH_SVG("Starship/CoreWidgets/Window/minimize", FVector2f(42.0f, 23.0f), FStyleColors::Foreground));
 
 		const FButtonStyle MaximizeButtonStyle = FButtonStyle(NoBorder)
-			.SetNormal( IMAGE_BRUSH_SVG("Starship/CoreWidgets/Window/maximize", FVector2f(42.0f, 34.0f), FStyleColors::Foreground))
-			.SetHovered(IMAGE_BRUSH_SVG("Starship/CoreWidgets/Window/maximize", FVector2f(42.0f, 34.0f), FStyleColors::ForegroundHover))
-			.SetPressed(IMAGE_BRUSH_SVG("Starship/CoreWidgets/Window/maximize", FVector2f(42.0f, 34.0f), FStyleColors::Foreground));
+			.SetNormal( IMAGE_BRUSH_SVG("Starship/CoreWidgets/Window/maximize", FVector2f(42.0f, 23.0f), FStyleColors::Foreground))
+			.SetHovered(IMAGE_BRUSH_SVG("Starship/CoreWidgets/Window/maximize", FVector2f(42.0f, 23.0f), FStyleColors::ForegroundHover))
+			.SetPressed(IMAGE_BRUSH_SVG("Starship/CoreWidgets/Window/maximize", FVector2f(42.0f, 23.0f), FStyleColors::Foreground));
 
 		const FButtonStyle RestoreButtonStyle = FButtonStyle(NoBorder)
-			.SetNormal( IMAGE_BRUSH_SVG("Starship/CoreWidgets/Window/restore", FVector2f(42.0f, 34.0f), FStyleColors::Foreground))
-			.SetHovered(IMAGE_BRUSH_SVG("Starship/CoreWidgets/Window/restore", FVector2f(42.0f, 34.0f), FStyleColors::ForegroundHover))
-			.SetPressed(IMAGE_BRUSH_SVG("Starship/CoreWidgets/Window/restore", FVector2f(42.0f, 34.0f), FStyleColors::Foreground));
+			.SetNormal( IMAGE_BRUSH_SVG("Starship/CoreWidgets/Window/restore", FVector2f(42.0f, 23.0f), FStyleColors::Foreground))
+			.SetHovered(IMAGE_BRUSH_SVG("Starship/CoreWidgets/Window/restore", FVector2f(42.0f, 23.0f), FStyleColors::ForegroundHover))
+			.SetPressed(IMAGE_BRUSH_SVG("Starship/CoreWidgets/Window/restore", FVector2f(42.0f, 23.0f), FStyleColors::Foreground));
 
 		const FButtonStyle CloseButtonStyle = FButtonStyle(NoBorder)
-			.SetNormal( IMAGE_BRUSH_SVG("Starship/CoreWidgets/Window/close", FVector2f(42.0f, 34.0f), FStyleColors::Foreground))
-			.SetHovered(IMAGE_BRUSH_SVG("Starship/CoreWidgets/Window/close", FVector2f(42.0f, 34.0f), FStyleColors::ForegroundHover))
-			.SetPressed(IMAGE_BRUSH_SVG("Starship/CoreWidgets/Window/close", FVector2f(42.0f, 34.0f), FStyleColors::Foreground));
+			.SetNormal( IMAGE_BRUSH_SVG("Starship/CoreWidgets/Window/close", FVector2f(42.0f, 23.0f), FStyleColors::Foreground))
+			.SetHovered(IMAGE_BRUSH_SVG("Starship/CoreWidgets/Window/close", FVector2f(42.0f, 23.0f), FStyleColors::ForegroundHover))
+			.SetPressed(IMAGE_BRUSH_SVG("Starship/CoreWidgets/Window/close", FVector2f(42.0f, 23.0f), FStyleColors::Foreground));
 #endif
 
 		FWindowStyle Window = FWindowStyle()
