LayoutUtils: Change use of c++ concept to a template

diff --git a/Engine/Source/Runtime/SlateCore/Public/Layout/LayoutUtils.h b/Engine/Source/Runtime/SlateCore/Public/Layout/LayoutUtils.h
--- a/Engine/Source/Runtime/SlateCore/Public/Layout/LayoutUtils.h
+++ b/Engine/Source/Runtime/SlateCore/Public/Layout/LayoutUtils.h
@@ -10,6 +10,7 @@
 #include "SlateRect.h"
 #include "ArrangedChildren.h"
 #include "FlowDirection.h"
+#include "Templates/UnrealTypeTraits.h"
 
 struct AlignmentArrangeResult
 {
@@ -246,6 +247,9 @@ static void ArrangeSingleChild(EFlowDirection InFlowDirection, const FGeometry&
 	}
 }
 
+// Defines THasMemberFunction_GetShrinkSizeValue<T>
+GENERATE_MEMBER_FUNCTION_CHECK(GetShrinkSizeValue, float, const);
+
 template<EOrientation Orientation, typename SlotType>
 static void ArrangeChildrenInStack(EFlowDirection InLayoutFlow, const TPanelChildren<SlotType>& Children, const FGeometry& AllottedGeometry, FArrangedChildren& ArrangedChildren, float InOffset, bool bInAllowShrink)
 {
@@ -342,7 +346,8 @@ static void ArrangeChildrenInStack(EFlowDirection InLayoutFlow, const TPanelChil
 			{
 				// Allow separate values from grow and shrink, as the adjustment is relative to the child size. 
 				Item.GrowStretchValue = CurChild.GetSizeValue();
-				if constexpr (requires { CurChild.GetShrinkSizeValue(); })
+
+				if constexpr (THasMemberFunction_GetShrinkSizeValue<SlotType>::Value)
 				{
 					Item.ShrinkStretchValue = CurChild.GetShrinkSizeValue();
 				}
