Fix missing virtual destructor on ISlateBrushSource

#rnx

diff --git a/Engine/Source/Runtime/SlateCore/Public/Styling/SlateBrush.h b/Engine/Source/Runtime/SlateCore/Public/Styling/SlateBrush.h
--- a/Engine/Source/Runtime/SlateCore/Public/Styling/SlateBrush.h
+++ b/Engine/Source/Runtime/SlateCore/Public/Styling/SlateBrush.h
@@ -526,5 +526,6 @@ protected:
 class ISlateBrushSource
 {
 public:
+	virtual ~ISlateBrushSource() = default;
 	virtual const FSlateBrush* GetSlateBrush() const = 0;
 };
