Fixed FScopeCycleCounterSWidget not working in insights captures.

diff --git a/Engine/Source/Runtime/Core/Public/Stats/Stats2.h b/Engine/Source/Runtime/Core/Public/Stats/Stats2.h
--- a/Engine/Source/Runtime/Core/Public/Stats/Stats2.h
+++ b/Engine/Source/Runtime/Core/Public/Stats/Stats2.h
@@ -1715,6 +1715,17 @@ public:
 #endif
 	}
 
+	FORCEINLINE_STATS void StartTrace(const FName Name, const TCHAR* Desc)
+	{
+#if CPUPROFILERTRACE_ENABLED
+		if (UE_TRACE_CHANNELEXPR_IS_ENABLED(CpuChannel))
+		{
+			FCpuProfilerTrace::OutputBeginDynamicEventWithId(Name, Desc);
+			EmittedEvent |= TraceEvent;
+		}
+#endif
+	}
+
 	/**
 	 * Stops the capturing and stores the result
 	 */
diff --git a/Engine/Source/Runtime/SlateCore/Public/Widgets/SWidget.h b/Engine/Source/Runtime/SlateCore/Public/Widgets/SWidget.h
--- a/Engine/Source/Runtime/SlateCore/Public/Widgets/SWidget.h
+++ b/Engine/Source/Runtime/SlateCore/Public/Widgets/SWidget.h
@@ -822,11 +822,11 @@ private:
 
 public:
 
-	FORCEINLINE TStatId GetStatID() const
+	FORCEINLINE TStatId GetStatID(bool bForDeferredUse = false) const
 	{
 #if STATS
 		// this is done to avoid even registering stats for a disabled group (unless we plan on using it later)
-		if (FThreadStats::IsCollectingData())
+		if (bForDeferredUse || FThreadStats::IsCollectingData())
 		{
 			if (!StatID.IsValidStat())
 			{
@@ -836,7 +836,7 @@ public:
 		}
 		return TStatId(); // not doing stats at the moment, or ever
 #elif ENABLE_STATNAMEDEVENTS
-		if (!StatID.IsValidStat() && GCycleStatsShouldEmitNamedEvents)
+		if (!StatID.IsValidStat() && (bForDeferredUse || GCycleStatsShouldEmitNamedEvents))
 		{
 			CreateStatID();
 		}
diff --git a/Engine/Source/Runtime/SlateCore/Public/Widgets/SWidgetUtils.h b/Engine/Source/Runtime/SlateCore/Public/Widgets/SWidgetUtils.h
--- a/Engine/Source/Runtime/SlateCore/Public/Widgets/SWidgetUtils.h
+++ b/Engine/Source/Runtime/SlateCore/Public/Widgets/SWidgetUtils.h
@@ -17,11 +17,25 @@ struct FScopeCycleCounterSWidget : public FCycleCounter
 	{
 		if (Widget)
 		{
-			TStatId WidgetStatId = Widget->GetStatID();
+#if CPUPROFILERTRACE_ENABLED
+			const bool bCpuChannelEnabled = UE_TRACE_CHANNELEXPR_IS_ENABLED(CpuChannel);
+#else
+			const bool bCpuChannelEnabled = false;
+#endif
+			bool bStarted = false;
+			TStatId WidgetStatId = Widget->GetStatID(bCpuChannelEnabled);
 			if (FThreadStats::IsCollectingData(WidgetStatId))
 			{
 				Start(WidgetStatId);
+				bStarted = true;
+			}
+
+#if CPUPROFILERTRACE_ENABLED
+			if (!bStarted && bCpuChannelEnabled && WidgetStatId.IsValidStat())
+			{
+				StartTrace(WidgetStatId.GetName(), WidgetStatId.GetStatDescriptionWIDE());
 			}
+#endif
 		}
 	}
 
@@ -42,9 +56,34 @@ struct FScopeCycleCounterSWidget : public FCycleCounter
 struct FScopeCycleCounterSWidget
 {
 	FScopeCycleCounter ScopeCycleCounter;
+#if CPUPROFILERTRACE_ENABLED
+	bool bPop = false;
+#endif
+
 	FORCEINLINE FScopeCycleCounterSWidget(const SWidget* Widget)
 		: ScopeCycleCounter(Widget ? Widget->GetStatID().StatString : nullptr)
 	{
+#if CPUPROFILERTRACE_ENABLED
+		if (GCycleStatsShouldEmitNamedEvents && UE_TRACE_CHANNELEXPR_IS_ENABLED(CpuChannel) && Widget)
+		{
+			const TStatId StatId = Widget->GetStatID();
+			if (StatId.IsValidStat())
+			{
+				bPop = true;
+				FCpuProfilerTrace::OutputBeginDynamicEvent(StatId.StatString);
+			}
+		}
+#endif
+	}
+
+	FORCEINLINE ~FScopeCycleCounterSWidget()
+	{
+#if CPUPROFILERTRACE_ENABLED
+		if (bPop)
+		{
+			FCpuProfilerTrace::OutputEndEvent();
+		}
+#endif
 	}
 };
 
