Added profiler scope around new FSlateElementBatcher::AddCachedElements logic from 31164434.

diff --git a/Engine/Source/Runtime/SlateCore/Private/Rendering/ElementBatcher.cpp b/Engine/Source/Runtime/SlateCore/Private/Rendering/ElementBatcher.cpp
--- a/Engine/Source/Runtime/SlateCore/Private/Rendering/ElementBatcher.cpp
+++ b/Engine/Source/Runtime/SlateCore/Private/Rendering/ElementBatcher.cpp
@@ -568,11 +568,16 @@ void FSlateElementBatcher::AddCachedElements(FSlateCachedElementData& CachedElem
 	}
 	CachedElementData.ListsWithNewData.Empty();
 
-	for (const FSlateRenderBatch& CachedBatch : CachedElementData.GetCachedBatches())
+	const TSparseArray<FSlateRenderBatch>& CachedBatches = CachedElementData.GetCachedBatches();
+	if (!CachedBatches.IsEmpty())
 	{
-		if (const FSlateShaderResource* ShaderResource = CachedBatch.GetShaderResource())
+		QUICK_SCOPE_CYCLE_COUNTER(STAT_UpdateUsedSlatePostBuffers);
+		for (const FSlateRenderBatch& CachedBatch : CachedElementData.GetCachedBatches())
 		{
-			UsedSlatePostBuffers |= ShaderResource->GetUsedSlatePostBuffers();
+			if (const FSlateShaderResource* ShaderResource = CachedBatch.GetShaderResource())
+			{
+				UsedSlatePostBuffers |= ShaderResource->GetUsedSlatePostBuffers();
+			}
 		}
 	}
 
