FSlateDrawElement::MakeGradient - pass along InGradientStops array to FSlateGradientElement::SetGradient to avoid a copy.

diff --git a/Engine/Source/Runtime/SlateCore/Private/Rendering/DrawElementTypes.cpp b/Engine/Source/Runtime/SlateCore/Private/Rendering/DrawElementTypes.cpp
--- a/Engine/Source/Runtime/SlateCore/Private/Rendering/DrawElementTypes.cpp
+++ b/Engine/Source/Runtime/SlateCore/Private/Rendering/DrawElementTypes.cpp
@@ -516,7 +516,7 @@ void FSlateDrawElement::MakeGradient( FSlateWindowElementList& ElementList, uint
 
 	FSlateGradientElement& Element = ElementList.AddUninitialized<EElementType::ET_Gradient>();
 
-	Element.SetGradient(InGradientStops, InGradientType, CornerRadius);
+	Element.SetGradient(MoveTemp(InGradientStops), InGradientType, CornerRadius);
 
 	Element.Init(ElementList, EElementType::ET_Gradient, InLayer, PaintGeometry, InDrawEffects);
 }
