Removing unused inclusive Slate CSV metrics

diff --git a/Engine/Source/Runtime/Slate/Private/Framework/Application/SlateApplication.cpp b/Engine/Source/Runtime/Slate/Private/Framework/Application/SlateApplication.cpp
--- a/Engine/Source/Runtime/Slate/Private/Framework/Application/SlateApplication.cpp
+++ b/Engine/Source/Runtime/Slate/Private/Framework/Application/SlateApplication.cpp
@@ -1269,7 +1269,6 @@ static void PrepassWindowAndChildren(TSharedRef<SWindow> WindowToPrepass, const
 void FSlateApplication::DrawPrepass( TSharedPtr<SWindow> DrawOnlyThisWindow )
 {
 	SCOPED_NAMED_EVENT_TEXT("Slate::Prepass", FColor::Magenta);
-	CSV_SCOPED_TIMING_STAT(Slate, Prepass);
 
 	TSharedPtr<SWindow> CurrentDebuggingWindowPinned = CurrentDebuggingWindow.Pin();
 
@@ -1517,7 +1516,6 @@ void FSlateApplication::Tick(ESlateTickType TickType)
 #endif
 
 	SCOPED_NAMED_EVENT_F(TEXT("Slate::Tick (%s)"), FColor::Magenta, LexToString(TickType));
-	CSV_SCOPED_TIMING_STAT(Slate, Tick);
 
 	{
 		SCOPE_CYCLE_COUNTER(STAT_SlateTickTime);
diff --git a/Engine/Source/Runtime/SlateCore/Private/FastUpdate/SlateInvalidationRoot.cpp b/Engine/Source/Runtime/SlateCore/Private/FastUpdate/SlateInvalidationRoot.cpp
--- a/Engine/Source/Runtime/SlateCore/Private/FastUpdate/SlateInvalidationRoot.cpp
+++ b/Engine/Source/Runtime/SlateCore/Private/FastUpdate/SlateInvalidationRoot.cpp
@@ -700,7 +700,6 @@ bool FSlateInvalidationRoot::PaintFastPath_UpdateNextWidget(const FSlateInvalida
 bool FSlateInvalidationRoot::PaintFastPath(const FSlateInvalidationContext& Context)
 {
 	SCOPED_NAMED_EVENT(SWidget_FastPathUpdate, FColor::Green);
-	CSV_SCOPED_TIMING_STAT(Slate, FastPathUpdate);
 
 	check(!bNeedsSlowPath);
 
@@ -1205,7 +1204,6 @@ bool FSlateInvalidationRoot::ProcessPostUpdate()
 bool FSlateInvalidationRoot::ProcessInvalidation()
 {
 	SCOPED_NAMED_EVENT(Slate_InvalidationProcessing, FColor::Blue);
-	CSV_SCOPED_TIMING_STAT(Slate, InvalidationProcessing);
 #if WITH_SLATE_DEBUGGING
 	PerformanceStat = FPerformanceStat();
 	FScopedDurationTimer TmpPerformance_ProcessInvalidation(PerformanceStat.InvalidationProcessing);
diff --git a/Engine/Source/Runtime/SlateCore/Private/Rendering/ElementBatcher.cpp b/Engine/Source/Runtime/SlateCore/Private/Rendering/ElementBatcher.cpp
--- a/Engine/Source/Runtime/SlateCore/Private/Rendering/ElementBatcher.cpp
+++ b/Engine/Source/Runtime/SlateCore/Private/Rendering/ElementBatcher.cpp
@@ -532,7 +532,6 @@ void FSlateElementBatcher::AddElementsInternal(const FSlateDrawElementMap& DrawE
 
 void FSlateElementBatcher::AddCachedElements(FSlateCachedElementData& CachedElementData, FVector2f ViewportSize)
 {
-	CSV_SCOPED_TIMING_STAT(Slate, AddCachedElements);
 	SCOPED_NAMED_EVENT_TEXT("Slate::AddCachedElements", FColor::Magenta);
 
 #if SLATE_CSV_TRACKER
diff --git a/Engine/Source/Runtime/UMG/Private/UserWidget.cpp b/Engine/Source/Runtime/UMG/Private/UserWidget.cpp
--- a/Engine/Source/Runtime/UMG/Private/UserWidget.cpp
+++ b/Engine/Source/Runtime/UMG/Private/UserWidget.cpp
@@ -2421,8 +2421,6 @@ UUserWidget* UUserWidget::CreateInstanceInternal(UObject* Outer, TSubclassOf<UUs
 {
 	LLM_SCOPE_BYTAG(UI_UMG);
 
-	//CSV_SCOPED_TIMING_STAT(Slate, CreateWidget);
-
 #if !(UE_BUILD_SHIPPING || UE_BUILD_TEST)
 	// Only do this on a non-shipping or test build.
 	if (!CreateWidgetHelpers::ValidateUserWidgetClass(UserWidgetClass))
