Migrate cached batch post buffer usage checks to single loop

#tests Minimal code build + RR

diff --git a/Engine/Source/Runtime/SlateCore/Private/Rendering/ElementBatcher.cpp b/Engine/Source/Runtime/SlateCore/Private/Rendering/ElementBatcher.cpp
--- a/Engine/Source/Runtime/SlateCore/Private/Rendering/ElementBatcher.cpp
+++ b/Engine/Source/Runtime/SlateCore/Private/Rendering/ElementBatcher.cpp
@@ -567,21 +567,36 @@ void FSlateElementBatcher::AddCachedElements(FSlateCachedElementData& CachedElem
 	}
 	CachedElementData.ListsWithNewData.Empty();
 
+	// Add the existing and new cached batches.
 	const TSparseArray<FSlateRenderBatch>& CachedBatches = CachedElementData.GetCachedBatches();
-	if (!CachedBatches.IsEmpty())
+	TArray<FSlateRenderBatch>& RenderBatchesSDR = BatchData->RenderBatches;
+	TArray<FSlateRenderBatch>& RenderBatchesHDR = BatchDataHDR->RenderBatches;
+
+	auto AddBatch = [&](const FSlateRenderBatch& CachedBatch)
 	{
-		QUICK_SCOPE_CYCLE_COUNTER(STAT_UpdateUsedSlatePostBuffers);
-		for (const FSlateRenderBatch& CachedBatch : CachedElementData.GetCachedBatches())
+		if (EnumHasAnyFlags(CachedBatch.GetDrawFlags(), ESlateBatchDrawFlag::HDR))
 		{
-			if (const FSlateShaderResource* ShaderResource = CachedBatch.GetShaderResource())
-			{
-				UsedSlatePostBuffers |= ShaderResource->GetUsedSlatePostBuffers();
-			}
+			RenderBatchesHDR.Add(CachedBatch);
 		}
+		else
+		{
+			RenderBatchesSDR.Add(CachedBatch);
+		}
+
+		if (const FSlateShaderResource* ShaderResource = CachedBatch.GetShaderResource())
+		{
+			UsedSlatePostBuffers |= ShaderResource->GetUsedSlatePostBuffers();
+		}
+	};
+
+	RenderBatchesSDR.Reserve(RenderBatchesSDR.Num() + CachedBatches.Num());
+	RenderBatchesHDR.Reserve(RenderBatchesHDR.Num() + CachedBatches.Num());
+
+	for (const FSlateRenderBatch& CachedBatch : CachedBatches)
+	{
+		AddBatch(CachedBatch);
 	}
 
-	// Add the existing and new cached batches.
-	FSlateBatchData::AddCachedBatchesToBatchData(BatchData, BatchDataHDR, CachedElementData.GetCachedBatches());
 	CachedElementData.CleanupUnusedClipStates();
 
 #if SLATE_CSV_TRACKER
diff --git a/Engine/Source/Runtime/SlateCore/Public/Rendering/ElementBatcher.h b/Engine/Source/Runtime/SlateCore/Public/Rendering/ElementBatcher.h
--- a/Engine/Source/Runtime/SlateCore/Public/Rendering/ElementBatcher.h
+++ b/Engine/Source/Runtime/SlateCore/Public/Rendering/ElementBatcher.h
@@ -200,7 +200,13 @@ public:
 
 	/** Adds a cached batch, used in retained rendering */
 	void AddCachedBatches(const TSparseArray<FSlateRenderBatch>& InCachedBatches);
+
+	UE_DEPRECATED(5.5, "AddCachedBatchesToBatchData is no longer needed since this class is a friend FSlateElementBatcher. See: FSlateElementBatcher::AddCachedElements.")
 	static void AddCachedBatchesToBatchData(FSlateBatchData* BatchDataSDR, FSlateBatchData* BatchDataHDR, const TSparseArray<FSlateRenderBatch>& InCachedBatches);
+
+public:
+	friend FSlateElementBatcher;
+
 private:
 	void FillBuffersFromNewBatch(FSlateRenderBatch& Batch, FSlateVertexArray& FinalVertices, FSlateIndexArray& FinalIndices);
 	void CombineBatches(FSlateRenderBatch& FirstBatch, FSlateRenderBatch& SecondBatch, FSlateVertexArray& FinalVertices, FSlateIndexArray& FinalIndices);
