Fix ASan issue detected at exit in with Slate attribute invalidation system (non-local objects using delegates with FMRSWRecursiveAccessDetector enabled).
Issues occured (in dev build only) because a thread_local object from FMRSWRecursiveAccessDetector was used when destructing a delegate at exit time, but the thread_local object was already destructed.
Disabling the MT access detector for now for the problematic objects (the safest solution), but will investigate other solutions for 5.5.
#tests Used with engine test and Lyra project. Ran Lyra in PIE, no regression noticed. Validated using breakpoint that delegates were called as expected.

diff --git a/Engine/Source/Runtime/SlateCore/Private/Types/SlateAttributeDescriptor.cpp b/Engine/Source/Runtime/SlateCore/Private/Types/SlateAttributeDescriptor.cpp
--- a/Engine/Source/Runtime/SlateCore/Private/Types/SlateAttributeDescriptor.cpp
+++ b/Engine/Source/Runtime/SlateCore/Private/Types/SlateAttributeDescriptor.cpp
@@ -498,7 +498,7 @@ FSlateAttributeDescriptor::FContainerInitializer FSlateAttributeDescriptor::AddC
 	check(!ContainerName.IsNone());
 
 	const FContainer* FoundAttribute = FindContainer(ContainerName);
-	if (ensureAlwaysMsgf(FoundAttribute == nullptr, TEXT("The container '%s' already exist. (Do you have the correct parrent class in SLATE_DECLARE_WIDGET)"), *ContainerName.ToString()))
+	if (ensureAlwaysMsgf(FoundAttribute == nullptr, TEXT("The container '%s' already exist. (Do you have the correct parent class in SLATE_DECLARE_WIDGET)"), *ContainerName.ToString()))
 	{
 		Containers.Emplace(ContainerName, Offset);
 	}
@@ -512,7 +512,7 @@ FSlateAttributeDescriptor::FInitializer::FAttributeEntry FSlateAttributeDescript
 
 	int32 NewIndex = INDEX_NONE;
 	FAttribute const* FoundAttribute = FindAttribute(AttributeName);
-	if (ensureAlwaysMsgf(FoundAttribute == nullptr, TEXT("The attribute '%s' already exist. (Do you have the correct parrent class in SLATE_DECLARE_WIDGET)"), *AttributeName.ToString()))
+	if (ensureAlwaysMsgf(FoundAttribute == nullptr, TEXT("The attribute '%s' already exist. (Do you have the correct parent class in SLATE_DECLARE_WIDGET)"), *AttributeName.ToString()))
 	{
 		NewIndex = Attributes.Emplace(AttributeName, Offset, MoveTemp(Reason));
 	}
@@ -527,7 +527,7 @@ FSlateAttributeDescriptor::FContainerInitializer::FAttributeEntry FSlateAttribut
 
 	int32 NewIndex = INDEX_NONE;
 	FAttribute const* FoundAttribute = FindAttribute(AttributeName);
-	if (ensureAlwaysMsgf(FoundAttribute == nullptr, TEXT("The attribute '%s' already exist. (Do you have the correct parrent class in SLATE_DECLARE_WIDGET)"), *AttributeName.ToString()))
+	if (ensureAlwaysMsgf(FoundAttribute == nullptr, TEXT("The attribute '%s' already exist. (Do you have the correct parent class in SLATE_DECLARE_WIDGET)"), *AttributeName.ToString()))
 	{
 		NewIndex = Attributes.Emplace(ContainerName, AttributeName, Offset, MoveTemp(Reason));
 	}
diff --git a/Engine/Source/Runtime/SlateCore/Public/Types/SlateAttributeDescriptor.h b/Engine/Source/Runtime/SlateCore/Public/Types/SlateAttributeDescriptor.h
--- a/Engine/Source/Runtime/SlateCore/Public/Types/SlateAttributeDescriptor.h
+++ b/Engine/Source/Runtime/SlateCore/Public/Types/SlateAttributeDescriptor.h
@@ -41,7 +41,9 @@ public:
 		friend FSlateAttributeDescriptor;
 
 		using Arg1Type = const class SWidget&;
-		DECLARE_DELEGATE_RetVal_OneParam(EInvalidateWidgetReason, FGetter, Arg1Type);
+		// using "not checked" delegate to disable race detection. This delegate is used in a static (FSlateWidgetClassData) and will be destructed after the tls array used in the race detection code (which will result in a use after free).
+		using FGetter = TDelegate<EInvalidateWidgetReason(Arg1Type), FNotThreadSafeNotCheckedDelegateUserPolicy>;
+
 
 		FInvalidateWidgetReasonAttribute(const FInvalidateWidgetReasonAttribute&) = default;
 		FInvalidateWidgetReasonAttribute(FInvalidateWidgetReasonAttribute&&) = default;
@@ -84,8 +86,8 @@ public:
 		FGetter Getter;
 	};
 
-	/** */
-	DECLARE_DELEGATE_OneParam(FAttributeValueChangedDelegate, SWidget& /*Widget*/);
+	// using "not checked" delegate to disable race detection. This delegate is used in a static (FSlateWidgetClassData) and will be destructed after the tls array used in the race detection code (which will result in a use after free).
+	using FAttributeValueChangedDelegate = TDelegate<void(SWidget&), FNotThreadSafeNotCheckedDelegateUserPolicy>;
 
 	/** */
 	enum class ECallbackOverrideType
