Slate: Invalidation for visibility is now processed in the update lool. An active timer or a tick function can change the visibility. That would clear the cached elements but would not prevent the cached to be replaced during the OnPaint. The widget visibility is collapsed but the cache is valid and batched. Explored another solution to prevent the OnPaint when it changed. But the layout would be wrong and we would have empty spaces for one frame. Better draw it and fix it on the next frame.

diff --git a/Engine/Source/Runtime/SlateCore/Private/FastUpdate/SlateInvalidationRoot.cpp b/Engine/Source/Runtime/SlateCore/Private/FastUpdate/SlateInvalidationRoot.cpp
--- a/Engine/Source/Runtime/SlateCore/Private/FastUpdate/SlateInvalidationRoot.cpp
+++ b/Engine/Source/Runtime/SlateCore/Private/FastUpdate/SlateInvalidationRoot.cpp
@@ -166,11 +166,21 @@ FAutoConsoleVariableRef CVarSlateInvalidationEnableReindexLayerId(
 /**
  *
  */
- namespace Slate
+ namespace UE::Slate::Private
  {
+	 constexpr EInvalidateWidgetReason PreInvalidationReason = EInvalidateWidgetReason::Layout
+		 | EInvalidateWidgetReason::AttributeRegistration
+		 | EInvalidateWidgetReason::Visibility
+		 | EInvalidateWidgetReason::ChildOrder;
+	 constexpr EInvalidateWidgetReason PostInvalidationReason = EInvalidateWidgetReason::Layout
+		 | EInvalidateWidgetReason::Paint
+		 | EInvalidateWidgetReason::Volatility
+		 | EInvalidateWidgetReason::RenderTransform
+		 | EInvalidateWidgetReason::Prepass;
+
 	bool EInvalidateWidgetReason_HasPreUpdateFlag(EInvalidateWidgetReason InvalidateReason)
 	{
-		return EnumHasAnyFlags(InvalidateReason, EInvalidateWidgetReason::AttributeRegistration | EInvalidateWidgetReason::ChildOrder);
+		return EnumHasAnyFlags(InvalidateReason, PreInvalidationReason);
 	}
 
 	bool EInvalidateWidgetReason_HasPostUpdateFlag(EInvalidateWidgetReason InvalidateReason)
@@ -179,13 +189,6 @@ FAutoConsoleVariableRef CVarSlateInvalidationEnableReindexLayerId(
 		const uint8 AnyPostUpdate = (0xFF & ~(uint8)EInvalidateWidgetReason::AttributeRegistration);
 		return (((uint8)InvalidateReason & AnyPostUpdate) != 0);
 	}
-
-	constexpr EInvalidateWidgetReason PostInvalidationReason = EInvalidateWidgetReason::Layout
-		| EInvalidateWidgetReason::Paint
-		| EInvalidateWidgetReason::Volatility
-		| EInvalidateWidgetReason::RenderTransform
-		| EInvalidateWidgetReason::Visibility
-		| EInvalidateWidgetReason::Prepass;
  }
 
 
@@ -297,9 +300,9 @@ void FSlateInvalidationRoot::InvalidateWidget(FWidgetProxy& Proxy, EInvalidateWi
 
 	if (bProcessingAttributeUpdate)
 	{
-		if (ensureMsgf(FSlateAttributeBase::IsInvalidateWidgetReasonSupported(InvalidateReason)
-			, TEXT("An invalid invalidation occurred while processing the widget attributes. That may result in an infinit loop.")))
+		if (!FSlateAttributeBase::IsInvalidateWidgetReasonSupported(InvalidateReason))
 		{
+			ensureMsgf(false, TEXT("An invalid invalidation occurred while processing the widget attributes. That may result in an infinite loop."));
 			return;
 		}
 	}
@@ -307,7 +310,7 @@ void FSlateInvalidationRoot::InvalidateWidget(FWidgetProxy& Proxy, EInvalidateWi
 	if (!bNeedsSlowPath)
 	{
 		Proxy.CurrentInvalidateReason |= InvalidateReason;
-		if (Slate::EInvalidateWidgetReason_HasPreUpdateFlag(InvalidateReason))
+		if (UE::Slate::Private::EInvalidateWidgetReason_HasPreUpdateFlag(InvalidateReason))
 		{
 			WidgetsNeedingPreUpdate->HeapPushUnique(Proxy);
 		}
@@ -319,7 +322,7 @@ void FSlateInvalidationRoot::InvalidateWidget(FWidgetProxy& Proxy, EInvalidateWi
 			WidgetsNeedingPrepassUpdate->PushBackUnique(Proxy);
 		}
 
-		if (Slate::EInvalidateWidgetReason_HasPostUpdateFlag(InvalidateReason))
+		if (UE::Slate::Private::EInvalidateWidgetReason_HasPostUpdateFlag(InvalidateReason))
 		{
 			WidgetsNeedingPostUpdate->PushBackOrHeapUnique(Proxy);
 		}
@@ -997,9 +1000,24 @@ void FSlateInvalidationRoot::ProcessPreUpdate()
 					if (bIsInvalidationWidgetValid && EnumHasAnyFlags(CurrentInvalidateReason, EInvalidateWidgetReason::AttributeRegistration))
 					{
 						FSlateInvalidationWidgetList::InvalidationWidgetType& InvalidationWidget = (*FastWidgetPathList)[WidgetIndex];
+						// Process the attribute invalidation even if it could be process in a ChildOrder invalidation. The attribute is added with "Unique".
 						FastWidgetPathList->ProcessAttributeRegistrationInvalidation(InvalidationWidget);
 						EnumRemoveFlags(InvalidationWidget.CurrentInvalidateReason, EInvalidateWidgetReason::AttributeRegistration);
 					}
+
+					if (bIsInvalidationWidgetValid && EnumHasAnyFlags(CurrentInvalidateReason, EInvalidateWidgetReason::Visibility))
+					{
+						FSlateInvalidationWidgetList::InvalidationWidgetType& InvalidationWidget = (*FastWidgetPathList)[WidgetIndex];
+						WidgetPtr->UpdateFastPathVisibility(InvalidationWidget.Visibility.MimicAsParent(), GetHittestGrid());
+						// We should Keep the flag for the layout process (visibility changes affect the layout and paint) and it will be removed then.
+						//This phase only updated the internal flags and internal states.
+						//But if we are in the layout or paint process, and we add the Visibility flag, then the visibility flag will be remove and the UpdateFastPathVisibility won't execute.
+						//The EInvalidationWidgetReason is already full (8bits). We can add a new bit in the InvalidationWidgetType but that would mean more flag management.
+						//Easy solution is to use the "RenderTransform" flag because it does the same job as the "Visibility" in the layout and paint phase.
+						EnumRemoveFlags(InvalidationWidget.CurrentInvalidateReason, EInvalidateWidgetReason::Visibility);
+						EnumAddFlags(InvalidationWidget.CurrentInvalidateReason, EInvalidateWidgetReason::RenderTransform);
+						WidgetsNeedingPostUpdate->PushBackOrHeapUnique(InvalidationWidget);
+					}
 				}
 			}
 		}
@@ -1056,18 +1074,22 @@ void FSlateInvalidationRoot::ProcessAttributeUpdate()
 					AttributeItt.Advance();
 
 #if WITH_SLATE_DEBUGGING
+					// visibility is updated with FSlateAttributeMetaData::UpdateOnlyVisibilityAttributes
 					ensureMsgf(PreviousVisibility == InvalidationWidget.Visibility, TEXT("The visibility of widget '%s' doesn't match the previous visibility after the attribute update."), *FReflectionMetaData::GetWidgetDebugInfo(WidgetPtr));
+					// childorder is not support as an invalidation with attribute
 					ensureMsgf(PreviousLeafMostChildIndex == InvalidationWidget.LeafMostChildIndex, TEXT("The number of child of widget '%s' doesn't match the previous count after the attribute update."), *FReflectionMetaData::GetWidgetDebugInfo(WidgetPtr));
 #endif
 				}
 				else
 				{
-					AttributeItt.AdvanceToNextSibling();
+					FSlateAttributeMetaData::ApplyDelayedInvalidation(*WidgetPtr);
+					AttributeItt.Advance();
 				}
 			}
 			else
 			{
-				AttributeItt.AdvanceToNextParent();
+				FSlateAttributeMetaData::ApplyDelayedInvalidation(*WidgetPtr);
+				AttributeItt.Advance();
 			}
 		}
 		else
@@ -1138,7 +1160,7 @@ void FSlateInvalidationRoot::ProcessPrepassUpdate()
 			//If there is no new invalidation, then the Widget's CurrentInvalidation should be none and nothing will be execute.
 			//Then, if needed, ProcessPostUpdate will add the widget to the FinalUpdateList in the update correct order.
 			WidgetProxy.ProcessPostInvalidation(*WidgetsNeedingPostUpdate, *FastWidgetPathList, *this);
-			EnumRemoveFlags(WidgetProxy.CurrentInvalidateReason, Slate::PostInvalidationReason);
+			EnumRemoveFlags(WidgetProxy.CurrentInvalidateReason, UE::Slate::Private::PostInvalidationReason);
 		}
 	}
 	WidgetsNeedingPrepassUpdate->Reset(true);
@@ -1181,7 +1203,7 @@ bool FSlateInvalidationRoot::ProcessPostUpdate()
 #endif
 
 			bWidgetsNeedRepaint |= WidgetProxy.ProcessPostInvalidation(*WidgetsNeedingPostUpdate, *FastWidgetPathList, *this);
-			EnumRemoveFlags(WidgetProxy.CurrentInvalidateReason, Slate::PostInvalidationReason);
+			EnumRemoveFlags(WidgetProxy.CurrentInvalidateReason, UE::Slate::Private::PostInvalidationReason);
 
 			if (WidgetPtr->HasAnyUpdateFlags(EWidgetUpdateFlags::AnyUpdate) && WidgetProxy.Visibility.IsVisible())
 			{
@@ -1213,6 +1235,7 @@ bool FSlateInvalidationRoot::ProcessInvalidation()
 
 	bool bWidgetsNeedRepaint = false;
 
+	// Process ChildOrder | Visibility | AttributeRegistration invalidations.
 	if (!bNeedsSlowPath)
 	{
 		check(WidgetsNeedingPreUpdate);
@@ -1227,6 +1250,7 @@ bool FSlateInvalidationRoot::ProcessInvalidation()
 		ProcessPreUpdate();
 	}
 
+	// Execute the attributes delegate.
 	if (!bNeedsSlowPath)
 	{
 		SCOPED_NAMED_EVENT(Slate_InvalidationProcessing_AttributeUpdate, FColor::Blue);
@@ -1237,7 +1261,7 @@ bool FSlateInvalidationRoot::ProcessInvalidation()
 		ProcessAttributeUpdate();
 	}
 
-	/** Re-run any ChildOrder invalidation. Attributes may have added new ChildOrder. */
+	// Re-process new ChildOrder | Visibility | AttributeRegistration invalidations. Attributes may have added new invalidation.
 	if (!bNeedsSlowPath && WidgetsNeedingPreUpdate->Num() > 0)
 	{
 		SCOPED_NAMED_EVENT(Slate_InvalidationProcessing_PreUpdate, FColor::Blue);
@@ -1245,9 +1269,16 @@ bool FSlateInvalidationRoot::ProcessInvalidation()
 		ProcessPreUpdate();
 	}
 
+#if UE_SLATE_WITH_INVALIDATIONWIDGETLIST_DEBUGGING
+	if (GSlateInvalidationRootVerifyWidgetVisibility && !bNeedsSlowPath)
+	{
+		VerifyWidgetVisibility(GetFastPathWidgetList());
+	}
+#endif
+
 	if (!bNeedsSlowPath)
 	{
-		// Put all widgets in the VolatileUpdate list in the WidgetsNeedingPostUpdate
+		// Put all widgets in the VolatileUpdate (Paint, Tick, ActiveTimer) list in the WidgetsNeedingPostUpdate
 		WidgetsNeedingPrepassUpdate->Heapify();
 		WidgetsNeedingPostUpdate->Heapify();
 		{
@@ -1266,6 +1297,9 @@ bool FSlateInvalidationRoot::ProcessInvalidation()
 		}
 	}
 
+	// Process widget that needs Prepass and PrepassVolatile.
+	//They will execute forward and then executed again in ProcessPostUpdate in backward.
+	//That save time when the big widget needs Prepass because we can skip all the child widgets.
 	if (!bNeedsSlowPath)
 	{
 		SCOPED_NAMED_EVENT(Slate_InvalidationProcessing_PrepassUpdate, FColor::Blue);
@@ -1276,6 +1310,7 @@ bool FSlateInvalidationRoot::ProcessInvalidation()
 		ProcessPrepassUpdate();
 	}
 
+	// Process all the layout widgets (RenderTransform | Layout | Prepass)
 	if (!bNeedsSlowPath)
 	{
 		FinalUpdateList.Reset(WidgetsNeedingPostUpdate->Num());
@@ -1299,10 +1334,6 @@ bool FSlateInvalidationRoot::ProcessInvalidation()
 	}
 
 #if UE_SLATE_WITH_INVALIDATIONWIDGETLIST_DEBUGGING
-	if (GSlateInvalidationRootVerifyWidgetVisibility && !bNeedsSlowPath)
-	{
-		VerifyWidgetVisibility(GetFastPathWidgetList());
-	}
 	if (GSlateInvalidationRootVerifyWidgetVolatile && !bNeedsSlowPath)
 	{
 		VerifyWidgetVolatile(GetFastPathWidgetList(), FinalUpdateList);
@@ -1490,7 +1521,7 @@ void VerifyWidgetList(TSharedRef<SWidget> RootWidget, FSlateInvalidationRootHand
 {
 	FSlateInvalidationWidgetList List(InvalidationRootHandle, FSlateInvalidationWidgetList::FArguments{ 128, 128, 1000, false });
 	List.BuildWidgetList(RootWidget);
-	bool bIsIdentical = (List.DeapCompare(WidgetList));
+	bool bIsIdentical = (List.DeepCompare(WidgetList));
 	if (!bIsIdentical)
 	{
 		UE_LOG(LogSlate, Log, TEXT("**-- New Build List --**"));
diff --git a/Engine/Source/Runtime/SlateCore/Private/FastUpdate/SlateInvalidationWidgetList.cpp b/Engine/Source/Runtime/SlateCore/Private/FastUpdate/SlateInvalidationWidgetList.cpp
--- a/Engine/Source/Runtime/SlateCore/Private/FastUpdate/SlateInvalidationWidgetList.cpp
+++ b/Engine/Source/Runtime/SlateCore/Private/FastUpdate/SlateInvalidationWidgetList.cpp
@@ -451,7 +451,8 @@ FSlateInvalidationWidgetIndex FSlateInvalidationWidgetList::Internal_BuildWidget
 	}
 
 	if (ShouldBeAddedToVolatileUpdateList(Widget))
-	{		// The list is already sorted at this point. Add to the end.
+	{
+		// The list is already sorted at this point. Add to the end.
 		Data[LastestIndex].ElementIndexList_VolatileUpdateWidget.AddUnsorted(NewIndex.ElementIndex);
 	}
 
@@ -1458,6 +1459,7 @@ FSlateInvalidationWidgetList::FCutResult FSlateInvalidationWidgetList::Internal_
 				if (ShouldBeAddedToAttributeList(*Widget))
 				{
 					Data[NewArrayIndex].ElementIndexList_WidgetWithRegisteredSlateAttribute.AddUnsorted(NewElementIndex);
+					EnumRemoveFlags(NewInvalidationWidget.CurrentInvalidateReason, EInvalidateWidgetReason::AttributeRegistration);
 				}
 				if (ShouldBeAddedToVolatileUpdateList(*Widget))
 				{
@@ -1559,7 +1561,7 @@ TArray<TSharedPtr<SWidget>> FSlateInvalidationWidgetList::FindChildren(const SWi
 }
 
 
-bool FSlateInvalidationWidgetList::DeapCompare(const FSlateInvalidationWidgetList& Other) const
+bool FSlateInvalidationWidgetList::DeepCompare(const FSlateInvalidationWidgetList& Other) const
 {
 	if (Root.Pin() != Other.Root.Pin())
 	{
diff --git a/Engine/Source/Runtime/SlateCore/Private/FastUpdate/SlateInvalidationWidgetList.h b/Engine/Source/Runtime/SlateCore/Private/FastUpdate/SlateInvalidationWidgetList.h
--- a/Engine/Source/Runtime/SlateCore/Private/FastUpdate/SlateInvalidationWidgetList.h
+++ b/Engine/Source/Runtime/SlateCore/Private/FastUpdate/SlateInvalidationWidgetList.h
@@ -232,8 +232,12 @@ public:
 		check(IsValidIndex(BellowWidgetIndex));
 		const InvalidationWidgetType& BeginInvalidationWidget = (*this)[BellowWidgetIndex];
 		const FSlateInvalidationWidgetIndex EndWidgetIndex = BeginInvalidationWidget.LeafMostChildIndex;
-		const FSlateInvalidationWidgetIndex BeginWidgetIndex = IncrementIndex(BellowWidgetIndex);
+		if (BellowWidgetIndex == EndWidgetIndex)
+		{
+			return;
+		}
 
+		const FSlateInvalidationWidgetIndex BeginWidgetIndex = IncrementIndex(BellowWidgetIndex);
 		if (BeginWidgetIndex != FSlateInvalidationWidgetIndex::Invalid)
 		{
 			const bool bSameWidgetArrayIndex = BeginWidgetIndex.ArrayIndex == EndWidgetIndex.ArrayIndex;
@@ -454,7 +458,7 @@ public:
 	 * (1) be in the same order
 	 * (2) point the same SWidget (itself, parent, leaf)
 	 */
-	bool DeapCompare(const FSlateInvalidationWidgetList& Other) const;
+	bool DeepCompare(const FSlateInvalidationWidgetList& Other) const;
 
 	/** For testing purposes. Log the tree. */
 	void LogWidgetsList(bool bOnlyVisible) const;
diff --git a/Engine/Source/Runtime/SlateCore/Private/FastUpdate/WidgetProxy.cpp b/Engine/Source/Runtime/SlateCore/Private/FastUpdate/WidgetProxy.cpp
--- a/Engine/Source/Runtime/SlateCore/Private/FastUpdate/WidgetProxy.cpp
+++ b/Engine/Source/Runtime/SlateCore/Private/FastUpdate/WidgetProxy.cpp
@@ -159,7 +159,7 @@ void FWidgetProxy::ProcessLayoutInvalidation(FSlateInvalidationWidgetPostHeap& U
 	}
 
 	// If the desired size changed, invalidate the parent if it is visible
-	if (NewDesiredSize != CurrentDesiredSize || EnumHasAnyFlags(CurrentInvalidateReason, EInvalidateWidgetReason::Visibility | EInvalidateWidgetReason::RenderTransform))
+	if (NewDesiredSize != CurrentDesiredSize || EnumHasAnyFlags(CurrentInvalidateReason, EInvalidateWidgetReason::RenderTransform)) // NB Also, EInvalidateWidgetReason::Visibility. It use the same RenderTransform flag
 	{
 		if (ParentIndex != FSlateInvalidationWidgetIndex::Invalid)
 		{
@@ -208,7 +208,7 @@ bool FWidgetProxy::ProcessPostInvalidation(FSlateInvalidationWidgetPostHeap& Upd
 		}
 		bWidgetNeedsRepaint = true;
 	}
-	else if (EnumHasAnyFlags(CurrentInvalidateReason, EInvalidateWidgetReason::RenderTransform | EInvalidateWidgetReason::Layout | EInvalidateWidgetReason::Visibility))
+	else if (EnumHasAnyFlags(CurrentInvalidateReason, EInvalidateWidgetReason::RenderTransform | EInvalidateWidgetReason::Layout)) // NB Also, EInvalidateWidgetReason::Visibility. It use the same RenderTransform flag
 	{
 		ProcessLayoutInvalidation(UpdateList, FastWidgetPathList, Root);
 		
diff --git a/Engine/Source/Runtime/SlateCore/Private/Test/SlateInvalidationWidgetListTest.cpp b/Engine/Source/Runtime/SlateCore/Private/Test/SlateInvalidationWidgetListTest.cpp
--- a/Engine/Source/Runtime/SlateCore/Private/Test/SlateInvalidationWidgetListTest.cpp
+++ b/Engine/Source/Runtime/SlateCore/Private/Test/SlateInvalidationWidgetListTest.cpp
@@ -17,11 +17,7 @@
 //IMPLEMENT_SIMPLE_AUTOMATION_TEST(FSlateInvalidationWidgetListTest, "Slate.InvalidationWidgetList.AddBuildRemove", EAutomationTestFlags::ApplicationContextMask | EAutomationTestFlags::SmokeFilter)
 IMPLEMENT_SIMPLE_AUTOMATION_TEST(FSlateInvalidationWidgetListTest, "Slate.InvalidationWidgetList.AddBuildRemove", EAutomationTestFlags::ApplicationContextMask | EAutomationTestFlags::EngineFilter)
 
-namespace UE
-{
-namespace Slate
-{
-namespace Private
+namespace UE::Slate::Private
 {
 
 static const bool bUpdateOnlyWhatIsNeeded = false;
@@ -175,9 +171,8 @@ TSharedRef<SVerticalBox> BuildTestUI_Child(TSharedPtr<SVerticalBox>& WidgetA, TA
 	ChildOfWidgetA.Add(AddVerticalBox(WidgetA, TEXT('I')));
 	return WidgetA.ToSharedRef();
 }
-}
-}
-}
+
+} // namespace
 
 
 bool FSlateInvalidationWidgetListTest::RunTest(const FString& Parameters)
@@ -248,7 +243,7 @@ bool FSlateInvalidationWidgetListTest::RunTest(const FString& Parameters)
 					ArgsWithoutAssign.bAssignedWidgetIndex = false;
 					FSlateInvalidationWidgetList TempList = { FSlateInvalidationRootHandle(), ArgsWithoutAssign };
 					TempList.BuildWidgetList(RootChildOrder);
-					if (!TempList.DeapCompare(List))
+					if (!TempList.DeepCompare(List))
 					{
 						AddError(TEXT("Was not able to remove a child of F."));
 					}
@@ -280,7 +275,7 @@ bool FSlateInvalidationWidgetListTest::RunTest(const FString& Parameters)
 					ArgsWithoutAssign.bAssignedWidgetIndex = false;
 					FSlateInvalidationWidgetList TempList = { FSlateInvalidationRootHandle(), ArgsWithoutAssign };
 					TempList.BuildWidgetList(RootChildOrder);
-					if (!TempList.DeapCompare(List))
+					if (!TempList.DeepCompare(List))
 					{
 						AddError(TEXT("Was not able to remove F and C."));
 					}
@@ -298,7 +293,7 @@ bool FSlateInvalidationWidgetListTest::RunTest(const FString& Parameters)
 					ArgsWithoutAssign.bAssignedWidgetIndex = false;
 					FSlateInvalidationWidgetList TempList = { FSlateInvalidationRootHandle(), ArgsWithoutAssign };
 					TempList.BuildWidgetList(RootChildOrder);
-					if (!TempList.DeapCompare(List))
+					if (!TempList.DeepCompare(List))
 					{
 						AddError(TEXT("Was not able to remove the last item of A."));
 					}
@@ -343,7 +338,7 @@ bool FSlateInvalidationWidgetListTest::RunTest(const FString& Parameters)
 						ArgsWithoutAssign.bAssignedWidgetIndex = false;
 						FSlateInvalidationWidgetList TempList = { FSlateInvalidationRootHandle(), ArgsWithoutAssign };
 						TempList.BuildWidgetList(RootChildOrder);
-						if (!TempList.DeapCompare(List))
+						if (!TempList.DeepCompare(List))
 						{
 							AddError(FString::Printf(TEXT("Was not able to process invalidation of widget '%s'."), Message));
 						}
diff --git a/Engine/Source/Runtime/SlateCore/Private/Widgets/SWidget.cpp b/Engine/Source/Runtime/SlateCore/Private/Widgets/SWidget.cpp
--- a/Engine/Source/Runtime/SlateCore/Private/Widgets/SWidget.cpp
+++ b/Engine/Source/Runtime/SlateCore/Private/Widgets/SWidget.cpp
@@ -201,7 +201,6 @@ SWidget::SWidget()
 	, bNeedsPrepass(true)
 	, bHasRegisteredSlateAttribute(false)
 	, bEnabledAttributesUpdate(true)
-	, bHasPendingAttributesInvalidation(false)
 	, bIsDeclarativeSyntaxConstructionCompleted(false)
 	, bIsHoveredAttributeSet(false)
 	, bHasCustomPrepass(false)
@@ -816,7 +815,7 @@ void SWidget::UpdateWidgetProxy(int32 NewLayerId, FSlateCachedElementsHandle& Ca
 #if UE_SLATE_WITH_INVALIDATIONWIDGETLIST_DEBUGGING
 		MyProxy.bDebug_Updated = true;
 #endif
-		ensureMsgf(MyProxy.Visibility.IsVisibleDirectly() == GetVisibility().IsVisible()
+		ensureMsgf(MyProxy.Visibility.IsVisibleDirectly() == GetVisibility().IsVisible() || EnumHasAnyFlags(MyProxy.CurrentInvalidateReason, EInvalidateWidgetReason::Visibility)
 			, TEXT("The visibility of the widget '%s' changed during Paint")
 			, *FReflectionMetaData::GetWidgetPath(this));
 		if (IsVolatile() && !IsVolatileIndirectly())
@@ -860,12 +859,16 @@ void SWidget::SetFastPathProxyHandle(const FWidgetProxyHandle& Handle, FSlateInv
 
 	bInheritedVolatility = bParentVolatile;
 
-	if (!InvalidationVisibility.IsVisible() && PersistentState.CachedElementHandle.IsValid())
+	if (!InvalidationVisibility.IsVisible())
 	{
+		PersistentState.CachedElementHandle.RemoveFromCache();
+
 #if WITH_SLATE_DEBUGGING
-		check(PersistentState.CachedElementHandle.IsOwnedByWidget(this));
+		if (PersistentState.CachedElementHandle.IsValid())
+		{
+			check(PersistentState.CachedElementHandle.IsOwnedByWidget(this));
+		}
 #endif
-		PersistentState.CachedElementHandle.RemoveFromCache();
 	}
 
 	if (IsVolatile() && !IsVolatileIndirectly())
@@ -890,7 +893,7 @@ void SWidget::UpdateFastPathVisibility(FSlateInvalidationWidgetVisibility Parent
 	FHittestGrid* HittestGridToRemoveFrom = ParentHittestGrid;
 	if (FastPathProxyHandle.IsValid(this))
 	{	
-		// Try and remove this from the current handles hit test grid.  If we are in a nested invalidation situation the hittest grid may have changed
+		// Try and remove this from the current handles hit test grid. If we are in a nested invalidation situation the hittest grid may have changed
 		HittestGridToRemoveFrom = FastPathProxyHandle.GetInvalidationRoot_NoCheck()->GetHittestGrid();
 		FWidgetProxy& Proxy = FastPathProxyHandle.GetProxy();
 		Proxy.Visibility = NewVisibility;
@@ -1301,17 +1304,9 @@ void SWidget::Invalidate(EInvalidateWidgetReason InvalidateReason)
 		InvalidateReason |= EInvalidateWidgetReason::Layout;
 	}
 
-	const bool bVolatilityChanged = EnumHasAnyFlags(InvalidateReason, EInvalidateWidgetReason::Volatility) ? Advanced_InvalidateVolatility() : false;
-
 	if(FastPathProxyHandle.IsValid(this))
 	{
-		// Current thinking is that visibility and volatility should be updated right away, not during fast path invalidation processing next frame
-		if (EnumHasAnyFlags(InvalidateReason, EInvalidateWidgetReason::Visibility))
-		{
-			SCOPED_NAMED_EVENT(SWidget_UpdateFastPathVisibility, FColor::Red);
-			UpdateFastPathVisibility(FastPathProxyHandle.GetProxy().Visibility.MimicAsParent(), FastPathProxyHandle.GetInvalidationRoot_NoCheck()->GetHittestGrid());
-		}
-
+		const bool bVolatilityChanged = EnumHasAnyFlags(InvalidateReason, EInvalidateWidgetReason::Volatility) ? Advanced_InvalidateVolatility() : false;
 		if (bVolatilityChanged)
 		{
 			SCOPED_NAMED_EVENT(SWidget_UpdateFastPathVolatility, FColor::Red);
@@ -1460,10 +1455,6 @@ int32 SWidget::Paint(const FPaintArgs& Args, const FGeometry& AllottedGeometry,
 		UE_TRACE_SCOPED_SLATE_WIDGET_UPDATE(this);
 		if (HasAnyUpdateFlags(EWidgetUpdateFlags::NeedsActiveTimerUpdate))
 		{
-			if (bHasPendingAttributesInvalidation)
-			{
-				FSlateAttributeMetaData::ApplyDelayedInvalidation(*MutableThis);
-			}
 
 			SCOPE_CYCLE_COUNTER(STAT_SlateExecuteActiveTimers);
 			MutableThis->ExecuteActiveTimers(Args.GetCurrentTime(), Args.GetDeltaTime());
@@ -1471,11 +1462,6 @@ int32 SWidget::Paint(const FPaintArgs& Args, const FGeometry& AllottedGeometry,
 
 		if (HasAnyUpdateFlags(EWidgetUpdateFlags::NeedsTick))
 		{
-			if (bHasPendingAttributesInvalidation)
-			{
-				FSlateAttributeMetaData::ApplyDelayedInvalidation(*MutableThis);
-			}
-
 			INC_DWORD_STAT(STAT_SlateNumTickedWidgets);
 
 			SCOPE_CYCLE_COUNTER(STAT_SlateTickWidgets);
@@ -1484,11 +1470,6 @@ int32 SWidget::Paint(const FPaintArgs& Args, const FGeometry& AllottedGeometry,
 		}
 	}
 
-	if (bHasPendingAttributesInvalidation)
-	{
-		FSlateAttributeMetaData::ApplyDelayedInvalidation(*MutableThis);
-	}
-
 	// the rule our parent has set for us
 	const bool bInheritedHittestability = Args.GetInheritedHittestability();
 	const bool bOutgoingHittestability = bInheritedHittestability && GetVisibility().AreChildrenHitTestVisible();
diff --git a/Engine/Source/Runtime/SlateCore/Public/Widgets/SWidget.h b/Engine/Source/Runtime/SlateCore/Public/Widgets/SWidget.h
--- a/Engine/Source/Runtime/SlateCore/Public/Widgets/SWidget.h
+++ b/Engine/Source/Runtime/SlateCore/Public/Widgets/SWidget.h
@@ -1781,9 +1781,6 @@ private:
 	/** Are bound Slate Attributes will be updated once per frame. */
 	uint8 bEnabledAttributesUpdate : 1;
 
-	/** At least one SlateAttributes was updated but the invalidation was delayed. */
-	uint8 bHasPendingAttributesInvalidation : 1;
-
 	/** The SNew or SAssignedNew construction is completed. */
 	uint8 bIsDeclarativeSyntaxConstructionCompleted : 1;
 
