Slate: Added support for FillContentSize for scrollbox
- Also removed  shrink size specialization from ArrangeChildrenInStack(), which did not infact work with SBoxPanel
- Fixed available space consumption and distribution for FillContentSize layout
- Added FillContentSize() example in SLayoutExample

diff --git a/Engine/Source/Runtime/AppFramework/Private/Framework/Testing/SLayoutExample.cpp b/Engine/Source/Runtime/AppFramework/Private/Framework/Testing/SLayoutExample.cpp
--- a/Engine/Source/Runtime/AppFramework/Private/Framework/Testing/SLayoutExample.cpp
+++ b/Engine/Source/Runtime/AppFramework/Private/Framework/Testing/SLayoutExample.cpp
@@ -43,6 +43,9 @@ public:
 	{
 		const FVector2D HeadingShadowOffset(2,2);
 
+		TSharedPtr<int32> SharedSizeValue = MakeShared<int32>(1);
+		TAttribute<float> SizeAttrib = TAttribute<float>::Create(TAttribute<float>::FGetter::CreateLambda([SharedSizeValue] { return (float)*SharedSizeValue.Get(); })); 
+		
 		FSlateFontInfo LargeLayoutFont = FCoreStyle::GetDefaultFontStyle("Regular", 16);
 		FSlateFontInfo SmallLayoutFont = FCoreStyle::GetDefaultFontStyle("Regular", 10);
 		this->ChildSlot
@@ -97,6 +100,41 @@ public:
 					SNew(SButton) .Text( LOCTEXT("ExampleLayout-TextLabel07", ".FillWidth(3)") )
 				]
 			]
+			// Fill Contents example
+			+SScrollBox::Slot() .Padding(5.f)
+			[
+				SNew(STextBlock) .ShadowOffset(HeadingShadowOffset) .Font( LargeLayoutFont ) .Text( LOCTEXT("ExampleLayout-FillContentsLabel", "Fill Content Size:") )
+			]
+			+SScrollBox::Slot() .Padding(10.f,5.f)
+			[
+				SNew(STextBlock) .Font(SmallLayoutFont) .Text( LOCTEXT("ExampleLayout-FillContent0", "Calculated size tarts with content size (like auto), but allows to control shrinking and growing.") )
+			]
+			+SScrollBox::Slot() .Padding(10.f,5.f)
+			[
+				SNew(SHorizontalBox)
+				+SHorizontalBox::Slot() .FillContentWidth(0.f, 1.f)
+				[
+					SNew(SButton) .Text( LOCTEXT("ExampleLayout-FillContent1", "No grow, allow shrink.\n.FillContentWidth(0, 1)") )
+				]
+				+SHorizontalBox::Slot() .FillContentWidth(SizeAttrib)
+				[
+					SNew(SButton)
+					.Text_Lambda([SharedSizeValue]()
+					{
+						return FText::Format(LOCTEXT("ExampleLayout-FillContent2", "Allow grow and Shrink\n.FillContentWidth({0})"), *SharedSizeValue.Get());
+					})
+					.OnClicked_Lambda([SharedSizeValue]()
+					{
+						int32& Value = *SharedSizeValue;
+						Value = (Value + 1) % 4;
+						return FReply::Handled();						
+					})
+				]
+				+SHorizontalBox::Slot() .FillContentWidth(1.f, 0.f)
+				[
+					SNew(SButton) .Text( LOCTEXT("ExampleLayout-FillContent3", "Allow grow, no shrink\n.FillContentWidth(1, 0)") )
+				]
+			]
 			// Aspect Ratio example
 			+SScrollBox::Slot() .Padding(5)
 			[
diff --git a/Engine/Source/Runtime/Slate/Private/Widgets/Layout/SScrollBox.cpp b/Engine/Source/Runtime/Slate/Private/Widgets/Layout/SScrollBox.cpp
--- a/Engine/Source/Runtime/Slate/Private/Widgets/Layout/SScrollBox.cpp
+++ b/Engine/Source/Runtime/Slate/Private/Widgets/Layout/SScrollBox.cpp
@@ -28,6 +28,8 @@ void SScrollBox::FSlot::RegisterAttributes(FSlateWidgetSlotAttributeInitializer&
 	SLATE_ADD_SLOT_ATTRIBUTE_DEFINITION_WITH_NAME(FSlot, AttributeInitializer, "Slot.MaxSize", MaxSize, EInvalidateWidgetReason::Layout);
 	SLATE_ADD_SLOT_ATTRIBUTE_DEFINITION_WITH_NAME(FSlot, AttributeInitializer, "Slot.SizeValue", SizeValue, EInvalidateWidgetReason::Layout)
 		.UpdatePrerequisite("Slot.MaxSize");
+	SLATE_ADD_SLOT_ATTRIBUTE_DEFINITION_WITH_NAME(FSlot, AttributeInitializer, "Slot.ShrinkSizeValue", ShrinkSizeValue, EInvalidateWidgetReason::Layout)
+		.UpdatePrerequisite("Slot.MaxSize");
 }
 
 SScrollBox::FSlot::FSlotArguments SScrollBox::Slot()
diff --git a/Engine/Source/Runtime/Slate/Public/Widgets/Layout/SScrollBox.h b/Engine/Source/Runtime/Slate/Public/Widgets/Layout/SScrollBox.h
--- a/Engine/Source/Runtime/Slate/Public/Widgets/Layout/SScrollBox.h
+++ b/Engine/Source/Runtime/Slate/Public/Widgets/Layout/SScrollBox.h
@@ -91,6 +91,19 @@ public:
 				_SizeParam = FStretch(MoveTemp(InStretchCoefficient));
 				return Me();
 			}
+			
+			/**
+			 * The widget's content size is adjusted proportionally to fit the available space.
+			 * The slots size starts at DesiredSize, and a slot with coefficient of 2 will get adjusted twice as much as slot with coefficient 1 to fit the available space.
+			 * @param InStretchCoefficient Stretch coefficient for this slot.
+			 * @param InShrinkStretchCoefficient If specified, this stretch coefficient is used when the slots is shrinking below desired size. Otherwise InStretchCoefficient is used for both shrink and grow.  
+			 */
+			FSlotArguments& FillContentSize(TAttribute<float> InStretchCoefficient, TAttribute<float> InShrinkStretchCoefficient = TAttribute<float>())
+			{
+				_SizeParam = FStretchContent(MoveTemp(InStretchCoefficient), MoveTemp(InShrinkStretchCoefficient));
+				return Me();
+			}
+
 			/** Set the max size in SlateUnit this slot can be. */
 			FSlotArguments& MaxSize(TAttribute<float> InMaxHeight)
 			{
@@ -104,6 +117,7 @@ public:
 			: TBasicLayoutWidgetSlot<FSlot>(HAlign_Fill, VAlign_Fill)
 			, SizeRule(FSizeParam::SizeRule_Auto)
 			, SizeValue(*this, 1.f)
+			, ShrinkSizeValue(*this, 1.f)
 			, MaxSize(*this, 0.0f)
 		{ }
 
@@ -122,16 +136,44 @@ public:
 			return SizeValue.Get();
 		}
 
+		/**
+		 * Get the size parameter for the space rule, used when the slot size is shrinking below desired size.
+		 * Used for size rule SizeRule_StretchContent.
+		 */
+		float GetShrinkSizeValue() const
+		{
+			return ShrinkSizeValue.Get();
+		}
+		
 		/** Get the max size the slot can be.*/
 		float GetMaxSize() const
 		{
 			return MaxSize.Get();
 		}
 
-		/** Set the size Param of the slot, It could be a FStretch or a FAuto. */
+		/** Set the size Param of the slot, It could be a FStretch, FStretchContent, or a FAuto. */
 		void SetSizeParam(FSizeParam InSizeParam)
 		{
 			SizeRule = InSizeParam.SizeRule;
+
+			// ShrinkSizeValue is only used for StretchContent.
+			// If ShrinkValue is not set, make it equal to the Value. 
+			if (SizeRule == FSizeParam::SizeRule_StretchContent)
+			{
+				if (InSizeParam.ShrinkValue.IsSet())
+				{
+					ShrinkSizeValue.Assign(*this, MoveTemp(InSizeParam.ShrinkValue));
+				}
+				else
+				{
+					ShrinkSizeValue.Assign(*this, InSizeParam.Value); // Make copy, let SizeValue use move.
+				}
+			}
+			else
+			{
+				ShrinkSizeValue.Set(*this, 1.f); // Reset
+			}
+
 			SizeValue.Assign(*this, MoveTemp(InSizeParam.Value));
 		}
 
@@ -147,6 +189,17 @@ public:
 			SetSizeParam(FStretch(MoveTemp(StretchCoefficient)));
 		}
 
+		/**
+		 * The widget's content size is adjusted proportionally to fit the available space.
+		 * The slots size starts at DesiredSize, and a slot with coefficient of 2 will get adjusted twice as much as slot with coefficient 1 to fit the available space.
+		 * @param InStretchCoefficient Stretch coefficient for this slot.
+		 * @param InShrinkStretchCoefficient If specified, this stretch coefficient is used when the slots is shrinking below desired size. Otherwise InStretchCoefficient is used for both shrink and grow.  
+		 */
+		void SetSizeToStretchContent(TAttribute<float> InStretchCoefficient, TAttribute<float> InShrinkStretchCoefficient = TAttribute<float>())
+		{
+			SetSizeParam(FStretchContent(MoveTemp(InStretchCoefficient), MoveTemp(InShrinkStretchCoefficient)));
+		}
+
 		/** Set the max size in SlateUnit this slot can be. */
 		void SetMaxSize(TAttribute<float> InMaxSize)
 		{
@@ -165,9 +218,15 @@ public:
 		 /** The sizing rule to use. */
 		FSizeParam::ESizeRule SizeRule;
 
+		/** Flag indicating if the ShrinkSizeValue value is set. */
+		bool bIsShrinkSizeValueSet;
+		
 		/** The actual value this size parameter stores. */
 		typename TBasicLayoutWidgetSlot<FSlot>::template TSlateSlotAttribute<float> SizeValue;
 
+		/** The actual value this size parameter stores, used for shrinking (negative if not defined, use SizeValue). */
+		typename TBasicLayoutWidgetSlot<FSlot>::template TSlateSlotAttribute<float> ShrinkSizeValue;
+
 		/** The max size that this slot can be (0 if no max) */
 		typename TBasicLayoutWidgetSlot<FSlot>::template TSlateSlotAttribute<float> MaxSize;
 
diff --git a/Engine/Source/Runtime/SlateCore/Public/Layout/LayoutUtils.h b/Engine/Source/Runtime/SlateCore/Public/Layout/LayoutUtils.h
--- a/Engine/Source/Runtime/SlateCore/Public/Layout/LayoutUtils.h
+++ b/Engine/Source/Runtime/SlateCore/Public/Layout/LayoutUtils.h
@@ -247,9 +247,6 @@ static void ArrangeSingleChild(EFlowDirection InFlowDirection, const FGeometry&
 	}
 }
 
-// Defines THasMemberFunction_GetShrinkSizeValue<T>
-GENERATE_MEMBER_FUNCTION_CHECK(GetShrinkSizeValue, float, const);
-
 template<EOrientation Orientation, typename SlotType>
 static void ArrangeChildrenInStack(EFlowDirection InLayoutFlow, const TPanelChildren<SlotType>& Children, const FGeometry& AllottedGeometry, FArrangedChildren& ArrangedChildren, float InOffset, bool bInAllowShrink)
 {
@@ -345,16 +342,8 @@ static void ArrangeChildrenInStack(EFlowDirection InLayoutFlow, const TPanelChil
 			else if (CurChild.GetSizeRule() == FSizeParam::SizeRule_StretchContent)
 			{
 				// Allow separate values from grow and shrink, as the adjustment is relative to the child size. 
-				Item.GrowStretchValue = CurChild.GetSizeValue();
-
-				if constexpr (THasMemberFunction_GetShrinkSizeValue<SlotType>::Value)
-				{
-					Item.ShrinkStretchValue = CurChild.GetShrinkSizeValue();
-				}
-				else
-				{
-					Item.ShrinkStretchValue = Item.GrowStretchValue;
-				}
+				Item.GrowStretchValue = FMath::Max(0.f, CurChild.GetSizeValue());
+				Item.ShrinkStretchValue = FMath::Max(0.f,CurChild.GetShrinkSizeValue());
 				Item.Size = ChildSize;
 				Item.BasisSize = ChildSize;
 
@@ -426,12 +415,22 @@ static void ArrangeChildrenInStack(EFlowDirection InLayoutFlow, const TPanelChil
 		// First, consume each items desired size from the available space.
 		// The remainder is corrected by growing ot shrinking the items.
 		int32 NumStretchContentItems = 0;
-		for (const FStretchItem& Item : StretchItems)
+		for (FStretchItem& Item : StretchItems)
 		{
 			if (Item.SizeRule == FSizeParam::SizeRule_StretchContent)
 			{
 				AvailableSpace -= Item.Size;
 				NumStretchContentItems++;
+
+				// If the item cannot shrink or grow, mark it already frozen.
+				if (bIsGrowing)
+				{
+					Item.bFrozen |= FMath::IsNearlyZero(Item.GrowStretchValue);
+				}
+				else
+				{
+					Item.bFrozen |= FMath::IsNearlyZero(Item.ShrinkStretchValue);
+				}
 			}
 		}
 
@@ -444,6 +443,12 @@ static void ArrangeChildrenInStack(EFlowDirection InLayoutFlow, const TPanelChil
 		const int32 MaxPasses = FMath::Min(NumStretchContentItems, 5);
 		for (int32 Pass = 0; Pass < MaxPasses; Pass++)
 		{
+			// If no available space, stop.
+			if (FMath::IsNearlyZero(AvailableSpace))
+			{
+				break;
+			}
+
 			// On each pass calculate the total coefficients for valid items.
 			GrowStretchCoefficientTotal = 0.0f;
 			ShrinkStretchCoefficientTotal = 0.0f;
@@ -464,13 +469,14 @@ static void ArrangeChildrenInStack(EFlowDirection InLayoutFlow, const TPanelChil
 				? GrowStretchCoefficientTotal
 				: ShrinkStretchCoefficientTotal;
 			
-			// If all items are frozen, or no space to distribute, stop.
-			if (StretchCoefficientTotal < UE_KINDA_SMALL_NUMBER
-				|| FMath::Abs(AvailableSpace) < UE_KINDA_SMALL_NUMBER)
+			// If none of the items can stretch, stop.
+			if (StretchCoefficientTotal < UE_KINDA_SMALL_NUMBER)
 			{
 				break;
 			}
 
+			float ConsumedSpace = 0.0f;
+			
 			for (FStretchItem& Item : StretchItems)
 			{
 				if (Item.SizeRule == FSizeParam::SizeRule_StretchContent
@@ -479,6 +485,13 @@ static void ArrangeChildrenInStack(EFlowDirection InLayoutFlow, const TPanelChil
 					const float SizeAdjust = bIsGrowing
 						? (AvailableSpace * (Item.GrowStretchValue / GrowStretchCoefficientTotal))
 						: (AvailableSpace * (Item.ShrinkStretchValue * Item.BasisSize / ShrinkStretchCoefficientTotal));
+
+					// If the item cannot be adjusted anymore, mark it frozen.
+					if (FMath::IsNearlyZero(SizeAdjust))
+					{
+						Item.bFrozen = true;
+						continue;
+					}
 					
 					constexpr float MinSize = 0.0f;
 					const float MaxSize = Item.MaxSize;
@@ -487,7 +500,7 @@ static void ArrangeChildrenInStack(EFlowDirection InLayoutFlow, const TPanelChil
 					if ((Item.Size + SizeAdjust) <= MinSize)
 					{
 						// Adjustment goes past min constraint, apply what we can and freeze since the item cannot change anymore.
-						AvailableSpace -= MinSize - Item.Size;
+						ConsumedSpace += MinSize - Item.Size;
 						Item.Size = MinSize;
 						Item.bFrozen = true;
 					}
@@ -495,18 +508,20 @@ static void ArrangeChildrenInStack(EFlowDirection InLayoutFlow, const TPanelChil
 						&& (Item.Size + SizeAdjust) >= MaxSize)
 					{
 						// Adjustment goes past max constraint, apply what we can and freeze since the item cannot change anymore.
-						AvailableSpace -= MaxSize - Item.Size;
+						ConsumedSpace += MaxSize - Item.Size;
 						Item.Size = MaxSize;
 						Item.bFrozen = true;
 					}
 					else
 					{
 						// Within constraints, adjust.
-						AvailableSpace -= SizeAdjust;
+						ConsumedSpace += SizeAdjust;
 						Item.Size += SizeAdjust;
 					}
 				}
 			}
+			
+			AvailableSpace -= ConsumedSpace;
 		}
 	}
 	
diff --git a/Engine/Source/Runtime/SlateCore/Public/Widgets/SBoxPanel.h b/Engine/Source/Runtime/SlateCore/Public/Widgets/SBoxPanel.h
--- a/Engine/Source/Runtime/SlateCore/Public/Widgets/SBoxPanel.h
+++ b/Engine/Source/Runtime/SlateCore/Public/Widgets/SBoxPanel.h
@@ -47,7 +47,7 @@ protected:
 			: TBasicLayoutWidgetSlot<SlotType>(HAlign_Fill, VAlign_Fill)
 			, SizeRule(FSizeParam::SizeRule_Stretch)
 			, SizeValue(*this, 1.f)
-			, ShrinkSizeValue(*this, -1.f) // Negative means unset, SizeValue is used instead.
+			, ShrinkSizeValue(*this, 1.f)
 			, MaxSize(*this, 0.0f)
 		{ }
 
@@ -96,8 +96,7 @@ protected:
 		 */
 		float GetShrinkSizeValue() const
 		{
-			const float Value = ShrinkSizeValue.Get();
-			return Value < 0.0f ? GetSizeValue() : Value;
+			return ShrinkSizeValue.Get();
 		}
 
 		/** Get the max size the slot can be.*/
@@ -111,16 +110,26 @@ protected:
 		void SetSizeParam(FSizeParam InSizeParam)
 		{
 			SizeRule = InSizeParam.SizeRule;
-			SizeValue.Assign(*this, MoveTemp(InSizeParam.Value));
-			if (InSizeParam.ShrinkValue.IsSet())
+
+			// ShrinkSizeValue is only used for StretchContent.
+			// If ShrinkValue is not set, make it equal to the Value. 
+			if (SizeRule == FSizeParam::SizeRule_StretchContent)
 			{
-				ShrinkSizeValue.Assign(*this, MoveTemp(InSizeParam.ShrinkValue));
+				if (InSizeParam.ShrinkValue.IsSet())
+				{
+					ShrinkSizeValue.Assign(*this, MoveTemp(InSizeParam.ShrinkValue));
+				}
+				else
+				{
+					ShrinkSizeValue.Assign(*this, InSizeParam.Value); // Make copy, let SizeValue use move.
+				}
 			}
 			else
 			{
-				// Negative value means to use size value for shrink and grow.
-				ShrinkSizeValue.Set(*this, -1.f);
+				ShrinkSizeValue.Set(*this, 1.f); // Reset
 			}
+			
+			SizeValue.Assign(*this, MoveTemp(InSizeParam.Value));
 		}
 
 		/**
