Slate: Remove deprecated code previous to 5.0

diff --git a/Engine/Source/Runtime/Slate/Private/Widgets/Layout/SScrollBox.cpp b/Engine/Source/Runtime/Slate/Private/Widgets/Layout/SScrollBox.cpp
--- a/Engine/Source/Runtime/Slate/Private/Widgets/Layout/SScrollBox.cpp
+++ b/Engine/Source/Runtime/Slate/Private/Widgets/Layout/SScrollBox.cpp
@@ -37,22 +37,6 @@ SScrollBox::FSlot::FSlotArguments SScrollBox::Slot()
 	return FSlot::FSlotArguments(MakeUnique<FSlot>());
 }
 
-PRAGMA_DISABLE_DEPRECATION_WARNINGS
-void SScrollPanel::Construct(const FArguments& InArgs, const TArray<SScrollBox::FSlot*>& InSlots)
-{
-	PhysicalOffset = 0;
-	Children.Reserve(InSlots.Num());
-	for (int32 SlotIndex = 0; SlotIndex < InSlots.Num(); ++SlotIndex)
-	{
-		Children.Add(InSlots[SlotIndex]);
-	}
-	Orientation = InArgs._Orientation;
-	BackPadScrolling = InArgs._BackPadScrolling;
-	FrontPadScrolling = InArgs._FrontPadScrolling;
-}
-PRAGMA_ENABLE_DEPRECATION_WARNINGS
-
-
 void SScrollPanel::Construct(const FArguments& InArgs, TArray<SScrollBox::FSlot::FSlotArguments> InSlots)
 {
 	PhysicalOffset = 0;
diff --git a/Engine/Source/Runtime/Slate/Public/Widgets/Layout/SScrollBox.h b/Engine/Source/Runtime/Slate/Public/Widgets/Layout/SScrollBox.h
--- a/Engine/Source/Runtime/Slate/Public/Widgets/Layout/SScrollBox.h
+++ b/Engine/Source/Runtime/Slate/Public/Widgets/Layout/SScrollBox.h
@@ -528,12 +528,6 @@ protected:
 	/** Whether to permit overscroll on this scroll box */
 	EAllowOverscroll AllowOverscroll;
 
-#if WITH_EDITORONLY_DATA
-	/** Padding to the scrollbox */
-	UE_DEPRECATED(5.0, "ScrollBarPadding is deprecated, Use SetScrollBarPadding")
-	FMargin ScrollBarPadding;
-#endif
-
 	/** Whether to back pad this scroll box, allowing user to scroll backward until child contents are no longer visible */
 	bool BackPadScrolling;
 
@@ -632,9 +626,6 @@ public:
 	{
 	}
 
-	UE_DEPRECATED(5.0, "Direct construction of FSlot is deprecated")
-	void Construct(const FArguments& InArgs, const TArray<SScrollBox::FSlot*>& InSlots);
-
 	void Construct(const FArguments& InArgs, TArray<SScrollBox::FSlot::FSlotArguments> InSlots);
 
 public:
diff --git a/Engine/Source/Runtime/SlateCore/Private/Application/SlateApplicationBase.cpp b/Engine/Source/Runtime/SlateCore/Private/Application/SlateApplicationBase.cpp
--- a/Engine/Source/Runtime/SlateCore/Private/Application/SlateApplicationBase.cpp
+++ b/Engine/Source/Runtime/SlateCore/Private/Application/SlateApplicationBase.cpp
@@ -105,16 +105,6 @@ const FHitTesting& FSlateApplicationBase::GetHitTesting() const
 	return HitTesting;
 }
 
-TSharedRef<SWidget> FSlateApplicationBase::MakeWindowTitleBar(const TSharedRef<SWindow>& Window, const TSharedPtr<SWidget>& CenterContent, EHorizontalAlignment CenterContentAlignment, TSharedPtr<IWindowTitleBar>& OutTitleBar) const
-{
-	FWindowTitleBarArgs Args(Window);
-
-	Args.CenterContent = CenterContent;
-	Args.CenterContentAlignment = CenterContentAlignment;
-	
-	return MakeWindowTitleBar(Args, OutTitleBar);
-}
-
 void FSlateApplicationBase::RegisterActiveTimer( const TSharedRef<FActiveTimerHandle>& ActiveTimerHandle )
 {
 	FScopeLock ActiveTimerLock(&ActiveTimerCS);
diff --git a/Engine/Source/Runtime/SlateCore/Private/Children.cpp b/Engine/Source/Runtime/SlateCore/Private/Children.cpp
deleted file mode 100644
diff --git a/Engine/Source/Runtime/SlateCore/Private/FastUpdate/SlateInvalidationRoot.cpp b/Engine/Source/Runtime/SlateCore/Private/FastUpdate/SlateInvalidationRoot.cpp
--- a/Engine/Source/Runtime/SlateCore/Private/FastUpdate/SlateInvalidationRoot.cpp
+++ b/Engine/Source/Runtime/SlateCore/Private/FastUpdate/SlateInvalidationRoot.cpp
@@ -262,11 +262,6 @@ FString FSlateInvalidationRoot::GetReferencerName() const
 	return TEXT("FSlateInvalidationRoot");
 }
 
-void FSlateInvalidationRoot::InvalidateRoot(const SWidget* Investigator)
-{
-	InvalidateRootChildOrder(Investigator);
-}
-
 void FSlateInvalidationRoot::InvalidateRootChildOrder(const SWidget* Investigator)
 {
 	// Invalidate all proxy handles
diff --git a/Engine/Source/Runtime/SlateCore/Private/Input/HittestGrid.cpp b/Engine/Source/Runtime/SlateCore/Private/Input/HittestGrid.cpp
--- a/Engine/Source/Runtime/SlateCore/Private/Input/HittestGrid.cpp
+++ b/Engine/Source/Runtime/SlateCore/Private/Input/HittestGrid.cpp
@@ -725,16 +725,6 @@ bool FHittestGrid::SameSize(const FHittestGrid* OtherGrid) const
 	return GridOrigin == OtherGrid->GridOrigin && GridWindowOrigin == OtherGrid->GridWindowOrigin && GridSize == OtherGrid->GridSize;
 }
 
-void FHittestGrid::AddWidget(const TSharedRef<SWidget>& InWidget, int32 InBatchPriorityGroup, int32 InLayerId, int32 InSecondarySort)
-{
-	AddWidget(&(InWidget.Get()), InBatchPriorityGroup, InLayerId, FSlateInvalidationWidgetSortOrder());
-}
-
-void FHittestGrid::AddWidget(const TSharedRef<SWidget>& InWidget, int32 InBatchPriorityGroup, int32 InLayerId, FSlateInvalidationWidgetSortOrder InSecondarySort)
-{
-	AddWidget(&(InWidget.Get()), InBatchPriorityGroup, InLayerId, InSecondarySort);
-}
-
 void FHittestGrid::AddWidget(const SWidget* InWidget, int32 InBatchPriorityGroup, int32 InLayerId, FSlateInvalidationWidgetSortOrder InSecondarySort)
 {
 	check(InWidget);
@@ -796,11 +786,6 @@ void FHittestGrid::AddWidget(const SWidget* InWidget, int32 InBatchPriorityGroup
 	}
 }
 
-void FHittestGrid::RemoveWidget(const TSharedRef<SWidget>& InWidget)
-{
-	RemoveWidget(&*InWidget);
-}
-
 void FHittestGrid::RemoveWidget(const SWidget* InWidget)
 {
 #if UE_SLATE_ENABLE_HITTEST_STATS
@@ -831,11 +816,6 @@ void FHittestGrid::RemoveWidget(const SWidget* InWidget)
 	RemoveGrid(InWidget);
 }
 
-void FHittestGrid::UpdateWidget(const TSharedRef<SWidget>& InWidget, FSlateInvalidationWidgetSortOrder InSecondarySort)
-{
-	UpdateWidget(&(InWidget.Get()), InSecondarySort);
-}
-
 void FHittestGrid::UpdateWidget(const SWidget* InWidget, FSlateInvalidationWidgetSortOrder InSecondarySort)
 {
 	check(InWidget);
@@ -850,11 +830,6 @@ bool FHittestGrid::ContainsWidget(const SWidget* InWidget) const
 	return WidgetMap.Contains(InWidget);
 }
 
-void FHittestGrid::InsertCustomHitTestPath(const TSharedRef<SWidget> InWidget, TSharedRef<ICustomHitTestPath> CustomHitTestPath)
-{
-	InsertCustomHitTestPath(&InWidget.Get(), CustomHitTestPath);
-}
-
 void FHittestGrid::InsertCustomHitTestPath(const SWidget* InWidget, const TSharedRef<ICustomHitTestPath>& CustomHitTestPath)
 {
 	int32* WidgetIndex = WidgetMap.Find(InWidget);
diff --git a/Engine/Source/Runtime/SlateCore/Private/Layout/ArrangedWidget.cpp b/Engine/Source/Runtime/SlateCore/Private/Layout/ArrangedWidget.cpp
--- a/Engine/Source/Runtime/SlateCore/Private/Layout/ArrangedWidget.cpp
+++ b/Engine/Source/Runtime/SlateCore/Private/Layout/ArrangedWidget.cpp
@@ -22,15 +22,6 @@ FWidgetAndPointer::FWidgetAndPointer()
 	: FArrangedWidget(FArrangedWidget::GetNullWidget())
 {}
 
-FWidgetAndPointer::FWidgetAndPointer( const FArrangedWidget& InWidget, const TSharedPtr<const FVirtualPointerPosition>& InPosition )
-	: FArrangedWidget(InWidget)
-{
-	if (InPosition)
-	{
-		OptionalPointerPosition = *InPosition.Get();
-	}
-}
-
 FWidgetAndPointer::FWidgetAndPointer(const FArrangedWidget& InWidget)
 	: FArrangedWidget(InWidget)
 	, OptionalPointerPosition()
diff --git a/Engine/Source/Runtime/SlateCore/Private/Layout/WidgetPath.cpp b/Engine/Source/Runtime/SlateCore/Private/Layout/WidgetPath.cpp
--- a/Engine/Source/Runtime/SlateCore/Private/Layout/WidgetPath.cpp
+++ b/Engine/Source/Runtime/SlateCore/Private/Layout/WidgetPath.cpp
@@ -7,7 +7,6 @@
 DECLARE_CYCLE_STAT(TEXT("Weak-To-Strong WidgetPath"), STAT_WeakToStrong_WidgetPath, STATGROUP_Slate);
 
 
-PRAGMA_DISABLE_DEPRECATION_WARNINGS
 FWidgetPath::FWidgetPath()
 : Widgets( EVisibility::Visible )
 , TopLevelWindow()
@@ -33,7 +32,6 @@ FWidgetPath::FWidgetPath( TArrayView<FWidgetAndPointer> InWidgetsAndPointers )
 		VirtualPointerPositions.Add(WidgetAndPointer.GetPointerPosition());
 	}
 }
-PRAGMA_ENABLE_DEPRECATION_WARNINGS
 
 FWidgetPath FWidgetPath::GetPathDownTo( TSharedRef<const SWidget> MarkerWidget ) const
 {
@@ -57,18 +55,6 @@ FWidgetPath FWidgetPath::GetPathDownTo( TSharedRef<const SWidget> MarkerWidget )
 	}	
 }
 
-const TSharedPtr<const FVirtualPointerPosition>& FWidgetPath::GetCursorAt( int32 Index ) const
-{
-	static TSharedPtr<const FVirtualPointerPosition> CursorAt;
-	return CursorAt;
-}
-
-
-bool FWidgetPath::ContainsWidget(TSharedRef<const SWidget> WidgetToFind) const
-{
-	return ContainsWidget(&WidgetToFind.Get());
-}
-
 
 bool FWidgetPath::ContainsWidget( const SWidget* WidgetToFind ) const
 {
@@ -428,11 +414,6 @@ FWeakWidgetPath::EPathResolutionResult::Result FWeakWidgetPath::ToWidgetPath( FW
 	}
 }
 
-bool FWeakWidgetPath::ContainsWidget( const TSharedRef< const SWidget >& SomeWidget ) const
-{
-	return ContainsWidget(&SomeWidget.Get());
-}
-
 bool FWeakWidgetPath::ContainsWidget(const SWidget* SomeWidget) const
 {
 	for (int32 WidgetIndex = 0; WidgetIndex < Widgets.Num(); ++WidgetIndex)
diff --git a/Engine/Source/Runtime/SlateCore/Private/SlotBase.cpp b/Engine/Source/Runtime/SlateCore/Private/SlotBase.cpp
--- a/Engine/Source/Runtime/SlateCore/Private/SlotBase.cpp
+++ b/Engine/Source/Runtime/SlateCore/Private/SlotBase.cpp
@@ -5,13 +5,9 @@
 #include "Widgets/SWidget.h"
 #include "Widgets/SNullWidget.h"
 
-PRAGMA_DISABLE_DEPRECATION_WARNINGS
 FSlotBase::FSlotBase()
 	: Owner(nullptr)
 	, Widget(SNullWidget::NullWidget)
-#if WITH_EDITORONLY_DATA
-	, RawParentPtr(nullptr)
-#endif
 {
 
 }
@@ -19,9 +15,6 @@ FSlotBase::FSlotBase()
 FSlotBase::FSlotBase(const FChildren& Children)
 	: Owner(&Children)
 	, Widget(SNullWidget::NullWidget)
-#if WITH_EDITORONLY_DATA
-	, RawParentPtr(nullptr)
-#endif
 {
 
 }
@@ -29,13 +22,9 @@ FSlotBase::FSlotBase(const FChildren& Children)
 FSlotBase::FSlotBase( const TSharedRef<SWidget>& InWidget )
 	: Owner(nullptr)
 	, Widget(InWidget)
-#if WITH_EDITORONLY_DATA
-	, RawParentPtr(nullptr)
-#endif
 {
 
 }
-PRAGMA_ENABLE_DEPRECATION_WARNINGS
 
 SWidget* FSlotBase::GetOwnerWidget() const
 {
diff --git a/Engine/Source/Runtime/SlateCore/Private/Widgets/Images/SLayeredImage.cpp b/Engine/Source/Runtime/SlateCore/Private/Widgets/Images/SLayeredImage.cpp
--- a/Engine/Source/Runtime/SlateCore/Private/Widgets/Images/SLayeredImage.cpp
+++ b/Engine/Source/Runtime/SlateCore/Private/Widgets/Images/SLayeredImage.cpp
@@ -121,21 +121,6 @@ bool SLayeredImage::IsValidIndex(int32 Index) const
 	// Index 0 is our local SImage
 	return Index == 0 || Layers.IsValidIndex(Index - 1);
 }
-const FSlateBrush* SLayeredImage::GetLayerBrush(int32 Index) const
-{
-	if (Index == 0)
-	{
-		return GetImageAttribute().Get();
-	}
-	else if (Layers.IsValidIndex(Index - 1))
-	{
-		return Layers[Index - 1].Key.Get();
-	}
-	else
-	{
-		return nullptr;
-	}
-}
 
 void SLayeredImage::SetLayerBrush(int32 Index, TAttribute<const FSlateBrush*> Brush)
 {
@@ -153,22 +138,6 @@ void SLayeredImage::SetLayerBrush(int32 Index, TAttribute<const FSlateBrush*> Br
 	}
 }
 
-FSlateColor SLayeredImage::GetLayerColor(int32 Index) const
-{
-	if (Index == 0)
-	{
-		return GetColorAndOpacityAttribute().Get();
-	}
-	else if (Layers.IsValidIndex(Index - 1))
-	{
-		return Layers[Index - 1].Value.Get();
-	}
-	else
-	{
-		return FSlateColor();
-	}
-}
-
 void SLayeredImage::SetLayerColor(int32 Index, TAttribute<FSlateColor> Color)
 {
 	if (Index == 0)
diff --git a/Engine/Source/Runtime/SlateCore/Private/Widgets/SWidget.cpp b/Engine/Source/Runtime/SlateCore/Private/Widgets/SWidget.cpp
--- a/Engine/Source/Runtime/SlateCore/Private/Widgets/SWidget.cpp
+++ b/Engine/Source/Runtime/SlateCore/Private/Widgets/SWidget.cpp
@@ -186,14 +186,10 @@ void SWidget::PrivateRegisterAttributes(FSlateAttributeInitializer& AttributeIni
 	SLATE_ADD_MEMBER_ATTRIBUTE_DEFINITION_WITH_NAME(AttributeInitializer, "RenderTransformPivot", RenderTransformPivotAttribute, EInvalidateWidgetReason::Layout | EInvalidateWidgetReason::RenderTransform);
 }
 
-PRAGMA_DISABLE_DEPRECATION_WARNINGS
 SWidget::SWidget()
 	: bCanSupportFocus(true)
 	, bCanHaveChildren(true)
 	, bClippingProxy(false)
-#if WITH_EDITORONLY_DATA
-	, bIsHovered(false)
-#endif
 	, bToolTipForceFieldEnabled(false)
 	, bForceVolatile(false)
 	, bCachedVolatile(false)
@@ -243,9 +239,7 @@ SWidget::SWidget()
 	UE_SLATE_DEBUG_WIDGETLIST_ADD_WIDGET(this);
 	UE_TRACE_SLATE_WIDGET_ADDED(this);
 }
-PRAGMA_ENABLE_DEPRECATION_WARNINGS
 
-PRAGMA_DISABLE_DEPRECATION_WARNINGS
 SWidget::~SWidget()
 {
 #if WITH_SLATE_DEBUGGING
@@ -300,55 +294,6 @@ SWidget::~SWidget()
 	DEC_DWORD_STAT(STAT_SlateTotalWidgets);
 	DEC_MEMORY_STAT_BY(STAT_SlateSWidgetAllocSize, AllocSize);
 }
-PRAGMA_ENABLE_DEPRECATION_WARNINGS
-
-PRAGMA_DISABLE_DEPRECATION_WARNINGS
-void SWidget::Construct(
-	const TAttribute<FText>& InToolTipText,
-	const TSharedPtr<IToolTip>& InToolTip,
-	const TAttribute< TOptional<EMouseCursor::Type> >& InCursor,
-	const TAttribute<bool>& InEnabledState,
-	const TAttribute<EVisibility>& InVisibility,
-	const float InRenderOpacity,
-	const TAttribute<TOptional<FSlateRenderTransform>>& InTransform,
-	const TAttribute<FVector2D>& InTransformPivot,
-	const FName& InTag,
-	const bool InForceVolatile,
-	const EWidgetClipping InClipping,
-	const EFlowDirectionPreference InFlowPreference,
-	const TOptional<FAccessibleWidgetData>& InAccessibleData,
-	const TArray<TSharedRef<ISlateMetaData>>& InMetaData
-)
-{
-	FSlateBaseNamedArgs Args;
-	Args._ToolTipText = InToolTipText;
-	Args._ToolTip = InToolTip;
-	Args._Cursor = InCursor;
-	Args._IsEnabled = InEnabledState;
-	Args._Visibility = InVisibility;
-	Args._RenderOpacity = InRenderOpacity;
-	Args._ForceVolatile = InForceVolatile;
-	Args._Clipping = InClipping;
-	Args._PixelSnappingMethod = EWidgetPixelSnapping::Inherit;
-	Args._FlowDirectionPreference = InFlowPreference;
-	Args._RenderTransform = InTransform;
-	Args._RenderTransformPivot = InTransformPivot;
-	Args._Tag = InTag;
-	Args._AccessibleParams = InAccessibleData;
-	Args.MetaData = InMetaData;
-	SWidgetConstruct(Args);
-}
-PRAGMA_ENABLE_DEPRECATION_WARNINGS
-
-PRAGMA_DISABLE_DEPRECATION_WARNINGS
-void SWidget::SWidgetConstruct(const TAttribute<FText>& InToolTipText, const TSharedPtr<IToolTip>& InToolTip, const TAttribute< TOptional<EMouseCursor::Type> >& InCursor, const TAttribute<bool>& InEnabledState,
-							   const TAttribute<EVisibility>& InVisibility, const float InRenderOpacity, const TAttribute<TOptional<FSlateRenderTransform>>& InTransform, const TAttribute<FVector2D>& InTransformPivot,
-							   const FName& InTag, const bool InForceVolatile, const EWidgetClipping InClipping, const EFlowDirectionPreference InFlowPreference, const TOptional<FAccessibleWidgetData>& InAccessibleData,
-							   const TArray<TSharedRef<ISlateMetaData>>& InMetaData)
-{
-	Construct(InToolTipText, InToolTip, InCursor, InEnabledState, InVisibility, InRenderOpacity, InTransform, InTransformPivot, InTag, InForceVolatile, InClipping, InFlowPreference, InAccessibleData, InMetaData);
-}
-PRAGMA_ENABLE_DEPRECATION_WARNINGS
 
 void SWidget::SWidgetConstruct(const FSlateBaseNamedArgs& Args)
 {
@@ -719,11 +664,6 @@ void SWidget::SlatePrepass(float InLayoutScaleMultiplier)
 	}
 }
 
-void SWidget::InvalidatePrepass()
-{
-	MarkPrepassAsDirty();
-}
-
 void SWidget::InvalidateChildRemovedFromTree(SWidget& Child)
 {
 	// If the root is invalidated, we need to clear out its PersistentState regardless.
@@ -1269,11 +1209,6 @@ void SWidget::SetPixelSnapping(EWidgetPixelSnapping InPixelSnappingMethod)
 	}
 }
 
-bool SWidget::IsFastPathVisible() const
-{
-	return FastPathProxyHandle.GetWidgetVisibility(this).IsVisible();
-}
-
 void SWidget::Invalidate(EInvalidateWidgetReason InvalidateReason)
 {
 	SLATE_CROSS_THREAD_CHECK();
diff --git a/Engine/Source/Runtime/SlateCore/Public/Application/SlateApplicationBase.h b/Engine/Source/Runtime/SlateCore/Public/Application/SlateApplicationBase.h
--- a/Engine/Source/Runtime/SlateCore/Public/Application/SlateApplicationBase.h
+++ b/Engine/Source/Runtime/SlateCore/Public/Application/SlateApplicationBase.h
@@ -404,18 +404,6 @@ public:
 	 */
 	virtual TSharedRef<IToolTip> MakeToolTip( const FText& ToolTipText ) = 0;
 
-	/**
-	 * Creates a title bar for the specified window.
-	 *
-	 * @param Window The window to create the title bar for.
-	 * @param CenterContent Optional content for the title bar's center (will override window title).
-	 * @param CenterContentAlignment The horizontal alignment of the center content.
-	 * @param OutTitleBar Will hold a pointer to the title bar's interface.
-	 * @return The new title bar widget.
-	 */
-	UE_DEPRECATED(4.26, "This version of MakeWindowTitleBar has been deprecated. Use the version that takes in an FWindowTitleBarArgs parameter instead.")
-	SLATECORE_API virtual TSharedRef<SWidget> MakeWindowTitleBar(const TSharedRef<SWindow>& Window, const TSharedPtr<SWidget>& CenterContent, EHorizontalAlignment CenterContentAlignment, TSharedPtr<IWindowTitleBar>& OutTitleBar) const;
-
 	/**
 	 * Creates a title bar for the specified window.
 	 *
diff --git a/Engine/Source/Runtime/SlateCore/Public/FastUpdate/SlateInvalidationRoot.h b/Engine/Source/Runtime/SlateCore/Public/FastUpdate/SlateInvalidationRoot.h
--- a/Engine/Source/Runtime/SlateCore/Public/FastUpdate/SlateInvalidationRoot.h
+++ b/Engine/Source/Runtime/SlateCore/Public/FastUpdate/SlateInvalidationRoot.h
@@ -88,9 +88,6 @@ public:
 	SLATECORE_API virtual FString GetReferencerName() const override;
 	//~ End FGCObject interface
 
-	/** Rebuild the list and request a SlowPath. */
-	UE_DEPRECATED(4.27, "InvalidateRoot is deprecated, use InvalidateRootChildOrder or InvalidateRootChildOrder")
-	SLATECORE_API void InvalidateRoot(const SWidget* Investigator = nullptr);
 	/** Rebuild the list and request a SlowPath. */
 	SLATECORE_API void InvalidateRootChildOrder(const SWidget* Investigator = nullptr);
 	/** Invalidate the layout, forcing the parent of the InvalidationRoot to be repainted. */
diff --git a/Engine/Source/Runtime/SlateCore/Public/Fonts/SlateFontInfo.h b/Engine/Source/Runtime/SlateCore/Public/Fonts/SlateFontInfo.h
--- a/Engine/Source/Runtime/SlateCore/Public/Fonts/SlateFontInfo.h
+++ b/Engine/Source/Runtime/SlateCore/Public/Fonts/SlateFontInfo.h
@@ -23,8 +23,6 @@ enum class EFontFallback : uint8
 {
 	/** No fallback font */
 	FF_NoFallback,
-	/** Fallback to localized font set */
-	FF_LocalizedFallback UE_DEPRECATED(4.24, "Legacy localized fallback fonts have been removed. FF_LocalizedFallback no longer has any meaning, so use FF_NoFallback instead."),
 	/** Fallback to last resort font set */
 	FF_LastResortFallback,
 	/** Tries all fallbacks */
diff --git a/Engine/Source/Runtime/SlateCore/Public/Input/HittestGrid.h b/Engine/Source/Runtime/SlateCore/Public/Input/HittestGrid.h
--- a/Engine/Source/Runtime/SlateCore/Public/Input/HittestGrid.h
+++ b/Engine/Source/Runtime/SlateCore/Public/Input/HittestGrid.h
@@ -46,10 +46,6 @@ public:
 	 */
 	SLATECORE_API bool SetHittestArea(const UE::Slate::FDeprecateVector2DParameter& HittestPositionInDesktop, const UE::Slate::FDeprecateVector2DParameter& HittestDimensions, const UE::Slate::FDeprecateVector2DParameter& HitestOffsetInWindow = FVector2f::ZeroVector);
 
-	/** Insert custom hit test data for a widget already in the grid */
-	UE_DEPRECATED(5.0, "Deprecated. Use the InsertCustomHitTestPath with a pointer.")
-	SLATECORE_API void InsertCustomHitTestPath(const TSharedRef<SWidget> InWidget, TSharedRef<ICustomHitTestPath> CustomHitTestPath);
-
 	/** Insert custom hit test data for a widget already in the grid */
 	SLATECORE_API void InsertCustomHitTestPath(const SWidget* InWidget, const TSharedRef<ICustomHitTestPath>& CustomHitTestPath);
 
@@ -82,27 +78,11 @@ public:
 	/** Clear the grid */
 	SLATECORE_API void Clear();
 
-	/** Add SWidget from the HitTest Grid */
-	UE_DEPRECATED(4.27, "Deprecated. Use the AddWidget with the FSlateInvalidationWidgetSortOrder type parameters instead. Passing FSlateInvalidationWidgetSortOrder()")
-	SLATECORE_API void AddWidget(const TSharedRef<SWidget>& InWidget, int32 InBatchPriorityGroup, int32 InLayerId, int32 InSecondarySort);
-
-	/** Add SWidget from the HitTest Grid */
-	UE_DEPRECATED(5.0, "Deprecated. Use the AddWidget with a pointer.")
-	SLATECORE_API void AddWidget(const TSharedRef<SWidget>& InWidget, int32 InBatchPriorityGroup, int32 InLayerId, FSlateInvalidationWidgetSortOrder InSecondarySort);
-
 	/** Add SWidget from the HitTest Grid */
 	SLATECORE_API void AddWidget(const SWidget* InWidget, int32 InBatchPriorityGroup, int32 InLayerId, FSlateInvalidationWidgetSortOrder InSecondarySort);
 
-	/** Remove SWidget from the HitTest Grid */
-	UE_DEPRECATED(5.0, "Deprecated. Use the RemoveWidget with a pointer.")
-	SLATECORE_API void RemoveWidget(const TSharedRef<SWidget>& InWidget);
-
 	/** Remove SWidget from the HitTest Grid */
 	SLATECORE_API void RemoveWidget(const SWidget* InWidget);
-
-	/** Update the widget SecondarySort without removing it and readding it again. */
-	UE_DEPRECATED(5.0, "Deprecated. Use the UpdateWidget with a pointer.")
-	SLATECORE_API void UpdateWidget(const TSharedRef<SWidget>& InWidget, FSlateInvalidationWidgetSortOrder InSecondarySort);
 	
 	/** Update the widget SecondarySort without removing it and readding it again. */
 	SLATECORE_API void UpdateWidget(const SWidget* InWidget, FSlateInvalidationWidgetSortOrder InSecondarySort);
@@ -110,10 +90,6 @@ public:
 	/** Check if SWidget is contained within the HitTest Grid */
 	SLATECORE_API bool ContainsWidget(const SWidget* InWidget) const;
 
-	/** Append an already existing grid that occupy the same space. */
-	UE_DEPRECATED(4.26, "Deprecated. Use the FHittestGrid::AddGrid method instead")
-	void AppendGrid(FHittestGrid& OtherGrid) {}
-
 	/**
 	 * Add an already existing grid that occupy the same space.
 	 * The grid needs to have an owner, not be this grid and occupy the same space as this grid.
diff --git a/Engine/Source/Runtime/SlateCore/Public/Layout/ArrangedWidget.h b/Engine/Source/Runtime/SlateCore/Public/Layout/ArrangedWidget.h
--- a/Engine/Source/Runtime/SlateCore/Public/Layout/ArrangedWidget.h
+++ b/Engine/Source/Runtime/SlateCore/Public/Layout/ArrangedWidget.h
@@ -60,7 +60,6 @@ public:
 	TSharedRef<SWidget> Widget;
 };
 
-PRAGMA_DISABLE_DEPRECATION_WARNINGS
 struct SLATECORE_API FWidgetAndPointer : public FArrangedWidget
 {
 public:
@@ -68,9 +67,6 @@ public:
 	FWidgetAndPointer(const FArrangedWidget& InWidget);
 	FWidgetAndPointer(const FArrangedWidget& InWidget, TOptional<FVirtualPointerPosition> InPosition);
 
-	UE_DEPRECATED(5.0, "FWidgetAndPointer constructor that takes a shared ptr is deprecated.")
-	FWidgetAndPointer( const FArrangedWidget& InWidget, const TSharedPtr<const FVirtualPointerPosition>& InPosition );
-
 	TOptional<FVirtualPointerPosition> GetPointerPosition() const
 	{
 		return OptionalPointerPosition;
@@ -81,13 +77,6 @@ public:
 		OptionalPointerPosition = InPosition;
 	}
 
-public:
-#if WITH_EDITOR
-	UE_DEPRECATED(5.0, "Direct access to PointerPosition is now deprecated. Use the getter or setter.")
-	TSharedPtr<const FVirtualPointerPosition> PointerPosition;
-#endif
-
 private:
 	TOptional<FVirtualPointerPosition> OptionalPointerPosition;
 };
-PRAGMA_ENABLE_DEPRECATION_WARNINGS
diff --git a/Engine/Source/Runtime/SlateCore/Public/Layout/BasicLayoutWidgetSlot.h b/Engine/Source/Runtime/SlateCore/Public/Layout/BasicLayoutWidgetSlot.h
--- a/Engine/Source/Runtime/SlateCore/Public/Layout/BasicLayoutWidgetSlot.h
+++ b/Engine/Source/Runtime/SlateCore/Public/Layout/BasicLayoutWidgetSlot.h
@@ -11,7 +11,7 @@
 
 #include <type_traits>
 
-PRAGMA_DISABLE_DEPRECATION_WARNINGS
+
 /** Mixin to add the alignment functionality to a base slot. */
 template <typename MixedIntoType>
 class TAlignmentWidgetSlotMixin
@@ -58,21 +58,6 @@ protected:
 		VAlignment = InArgs._VAlignment.Get(VAlignment);
 	}
 
-public:
-	UE_DEPRECATED(5.0, "HAlign is now deprecated. Use the FSlotArgument or the SetHorizontalAlignment function.")
-	MixedIntoType& HAlign(EHorizontalAlignment InHAlignment)
-	{
-		HAlignment = InHAlignment;
-		return *(static_cast<MixedIntoType*>(this));
-	}
-
-	UE_DEPRECATED(5.0, "VAlign is now deprecated. Use the FSlotArgument or the SetVerticalAlignment function.")
-	MixedIntoType& VAlign(EVerticalAlignment InVAlignment)
-	{
-		VAlignment = InVAlignment;
-		return *(static_cast<MixedIntoType*>(this));
-	}
-
 public:
 	void SetHorizontalAlignment(EHorizontalAlignment Alignment)
 	{
@@ -102,18 +87,14 @@ public:
 		return VAlignment;
 	}
 
-public:
+private:
 	/** Horizontal positioning of child within the allocated slot */
-	UE_DEPRECATED(5.0, "Direct access to HAlignment is now deprecated. Use the getter.")
 	EHorizontalAlignment HAlignment;
 	/** Vertical positioning of child within the allocated slot */
-	UE_DEPRECATED(5.0, "Direct access to VAlignment is now deprecated. Use the getter.")
 	EVerticalAlignment VAlignment;
 };
-PRAGMA_ENABLE_DEPRECATION_WARNINGS
 
 
-PRAGMA_DISABLE_DEPRECATION_WARNINGS
 /** Mixin to add the alignment functionality to a base slot that is also a single children. */
 template <typename MixedIntoType>
 class TAlignmentSingleWidgetSlotMixin
@@ -206,18 +187,14 @@ public:
 		return VAlignment;
 	}
 
-public:
+private:
 	/** Horizontal positioning of child within the allocated slot */
-	UE_DEPRECATED(5.0, "Direct access to HAlignment is now deprecated. Use the getter.")
 	EHorizontalAlignment HAlignment;
 	/** Vertical positioning of child within the allocated slot */
-	UE_DEPRECATED(5.0, "Direct access to VAlignment is now deprecated. Use the getter.")
 	EVerticalAlignment VAlignment;
 };
-PRAGMA_ENABLE_DEPRECATION_WARNINGS
 
 
-PRAGMA_DISABLE_DEPRECATION_WARNINGS
 /** Mixin to add the padding functionality to a base slot. */
 template <typename MixedIntoType>
 class TPaddingWidgetSlotMixin
@@ -277,35 +254,6 @@ protected:
 		SLATE_ADD_SLOT_ATTRIBUTE_DEFINITION_WITH_NAME(MixedIntoType, AttributeInitializer, "Slot.Padding", SlotPaddingAttribute, EInvalidateWidgetReason::Layout);
 	}
 
-public:
-	UE_DEPRECATED(5.0, "Padding is now deprecated. Use the FSlotArgument or the SetPadding function.")
-	MixedIntoType& Padding(TAttribute<FMargin> InPadding)
-	{
-		SetPadding(MoveTemp(InPadding));
-		return *(static_cast<MixedIntoType*>(this));
-	}
-
-	UE_DEPRECATED(5.0, "Padding is now deprecated. Use the FSlotArgument or the SetPadding function.")
-	MixedIntoType& Padding(float Uniform)
-	{
-		SetPadding(FMargin(Uniform));
-		return *(static_cast<MixedIntoType*>(this));
-	}
-
-	UE_DEPRECATED(5.0, "Padding is now deprecated. Use the FSlotArgument or the SetPadding function.")
-	MixedIntoType& Padding(float Horizontal, float Vertical)
-	{
-		SetPadding(FMargin(Horizontal, Vertical));
-		return *(static_cast<MixedIntoType*>(this));
-	}
-
-	UE_DEPRECATED(5.0, "Padding is now deprecated. Use the FSlotArgument or the SetPadding function.")
-	MixedIntoType& Padding(float Left, float Top, float Right, float Bottom)
-	{
-		SetPadding(FMargin(Left, Top, Right, Bottom));
-		return *(static_cast<MixedIntoType*>(this));
-	}
-
 public:
 	void SetPadding(TAttribute<FMargin> InPadding)
 	{
@@ -317,20 +265,13 @@ public:
 		return SlotPaddingAttribute.Get();
 	}
 
-public:
-	UE_DEPRECATED(5.0, "Direct access to SlotPadding is now deprecated. Use the setter or getter.")
-	TSlateDeprecatedTAttribute<FMargin> SlotPadding;
-
 private:
 	using SlotPaddingCompareType = TSlateAttributeComparePredicate<>;
 	using SlotPaddingType = ::SlateAttributePrivate::TSlateContainedAttribute<FMargin, ::SlateAttributePrivate::FSlateAttributeNoInvalidationReason, SlotPaddingCompareType>;
 	SlotPaddingType SlotPaddingAttribute;
 };
-PRAGMA_ENABLE_DEPRECATION_WARNINGS
 
 
-
-PRAGMA_DISABLE_DEPRECATION_WARNINGS
 /** Mixin to add the padding functionality to a base slot that is also a single children. */
 template <typename MixedIntoType, EInvalidateWidgetReason InPaddingInvalidationReason = EInvalidateWidgetReason::Layout>
 class TPaddingSingleWidgetSlotMixin
@@ -419,11 +360,6 @@ public:
 		return *(static_cast<MixedIntoType*>(this));
 	}
 
-#if WITH_EDITORONLY_DATA
-	UE_DEPRECATED(5.0, "Direct access to SlotPadding is now deprecated. Use the setter or getter.")
-	TSlateDeprecatedTAttribute<FMargin> SlotPadding;
-#endif
-
 public:
 	void SetPadding(TAttribute<FMargin> InPadding)
 	{
@@ -453,7 +389,6 @@ public:
 protected:
 	SlotPaddingAttributeType SlotPaddingAttribute;
 };
-PRAGMA_ENABLE_DEPRECATION_WARNINGS
 
 
 /** A templated basic slot that can be used by layout. */
diff --git a/Engine/Source/Runtime/SlateCore/Public/Layout/Children.h b/Engine/Source/Runtime/SlateCore/Public/Layout/Children.h
--- a/Engine/Source/Runtime/SlateCore/Public/Layout/Children.h
+++ b/Engine/Source/Runtime/SlateCore/Public/Layout/Children.h
@@ -151,9 +151,6 @@ public:
 	static SLATECORE_API FNoChildren NoChildrenInstance;
 
 public:
-	UE_DEPRECATED(5.0, "FNoChildren take a valid reference to a SWidget")
-	SLATECORE_API FNoChildren();
-
 	FNoChildren(SWidget* InOwner)
 		: FChildren(InOwner)
 	{
@@ -205,56 +202,6 @@ private:
 };
 
 
-/**
- * Widgets that will only have one child.
- */
-template <typename MixedIntoType>
-class UE_DEPRECATED(5.0, "TSupportsOneChildMixin is deprecated because it got confused between FSlot and FChildren. Use FSingleWidgetChildren.")
-TSupportsOneChildMixin : public FChildren, public TSlotBase<MixedIntoType>
-{
-public:
-	TSupportsOneChildMixin(SWidget* InOwner)
-		: FChildren(InOwner)
-		, TSlotBase<MixedIntoType>(static_cast<const FChildren&>(*this))
-	{
-	}
-
-	TSupportsOneChildMixin(std::nullptr_t) = delete;
-
-	virtual int32 Num() const override { return 1; }
-
-	virtual TSharedRef<SWidget> GetChildAt( int32 ChildIndex ) override
-	{
-		check(ChildIndex == 0);
-		return FSlotBase::GetWidget();
-	}
-
-	virtual TSharedRef<const SWidget> GetChildAt( int32 ChildIndex ) const override
-	{
-		check(ChildIndex == 0);
-		return FSlotBase::GetWidget();
-	}
-
-private:
-	virtual const FSlotBase& GetSlotAt(int32 ChildIndex) const override
-	{
-		check(ChildIndex == 0);
-		return *this;
-	}
-
-	virtual FWidgetRef GetChildRefAt(int32 ChildIndex) override
-	{
-		check(ChildIndex == 0);
-		return FWidgetRef(ReferenceConstruct, FSlotBase::GetWidget().Get());
-	}
-	virtual FConstWidgetRef GetChildRefAt(int32 ChildIndex) const override
-	{
-		check(ChildIndex == 0);
-		return FConstWidgetRef(ReferenceConstruct, FSlotBase::GetWidget().Get());
-	}
-};
-
-
 /**
  * For widgets that do not own their content, but are responsible for presenting someone else's content.
  * e.g. Tooltips are just presented by the owner window; not actually owned by it. They can go away at any time
@@ -349,21 +296,6 @@ private:
 };
 
 
-template <typename MixedIntoType>
-class UE_DEPRECATED(5.0, "Renamed TSupportsContentPaddingMixin to TAlignmentWidgetSlotMixin to differenciate from FSlot and FChildren.")
-TSupportsContentAlignmentMixin : public TAlignmentWidgetSlotMixin<MixedIntoType>
-{
-	using TAlignmentWidgetSlotMixin<MixedIntoType>::TAlignmentWidgetSlotMixin;
-};
-
-template <typename MixedIntoType>
-class UE_DEPRECATED(5.0, "Renamed TSupportsContentPaddingMixin to TPaddingWidgetSlotMixin to differenciate from FSlot and FChildren.")
-TSupportsContentPaddingMixin : public TPaddingWidgetSlotMixin<MixedIntoType>
-{
-	using TPaddingWidgetSlotMixin<MixedIntoType>::TPaddingWidgetSlotMixin;
-};
-
-
 /** A FChildren that has only one child and can take a templated slot. */
 template<typename SlotType>
 class TSingleWidgetChildrenWithSlot : public FChildren, protected TSlotBase<SlotType>
@@ -519,15 +451,6 @@ class FSingleWidgetChildrenWithBasicLayoutSlot : public TSingleWidgetChildrenWit
 };
 
 
-/** A slot that support alignment of content and padding */
-class UE_DEPRECATED(5.0, "FSimpleSlot is deprecated because it got confused from FChildren with FSlot. Use FSingleWidgetChildrenWithSimpleSlot.")
-FSimpleSlot : public FSingleWidgetChildrenWithBasicLayoutSlot
-{
-public:
-	using FSingleWidgetChildrenWithBasicLayoutSlot::FSingleWidgetChildrenWithBasicLayoutSlot;
-};
-
-
 /**
  * A generic FChildren that stores children along with layout-related information.
  * The type containing Widget* and layout info is specified by ChildType.
@@ -580,16 +503,6 @@ public:
 	}
 
 public:
-	UE_DEPRECATED(5.0, "Add a slot directly has been deprecated. use the FSlotArgument to create a new slot")
-	int32 Add( SlotType* Slot )
-	{
-		int32 Index = Children.Add(TUniquePtr<SlotType>(Slot));
-		check(Slot);
-		Slot->SetOwner(*this);
-
-		return Index;
-	}
-
 	int32 AddSlot(typename SlotType::FSlotArguments&& SlotArgument)
 	{
 		TUniquePtr<SlotType> NewSlot = SlotArgument.StealSlot();
@@ -653,14 +566,6 @@ public:
 		Children = MoveTemp(ChildrenCopy);
 	}
 
-	UE_DEPRECATED(5.0, "Insert a slot directly has been deprecated. use the FSlotArgument to create a new slot")
-	void Insert(SlotType* Slot, int32 Index)
-	{
-		check(Slot);
-		Children.Insert(TUniquePtr<SlotType>(Slot), Index);
-		Slot->SetOwner(*this);
-	}
-
 	void InsertSlot(typename SlotType::FSlotArguments&& SlotArgument, int32 Index)
 	{
 		TUniquePtr<SlotType> NewSlot = SlotArgument.StealSlot();
diff --git a/Engine/Source/Runtime/SlateCore/Public/Layout/ChildrenBase.h b/Engine/Source/Runtime/SlateCore/Public/Layout/ChildrenBase.h
--- a/Engine/Source/Runtime/SlateCore/Public/Layout/ChildrenBase.h
+++ b/Engine/Source/Runtime/SlateCore/Public/Layout/ChildrenBase.h
@@ -23,7 +23,6 @@ class FSlotBase;
  * 
  * FChildren is intended to be returned by the GetChildren() method.
  */
-PRAGMA_DISABLE_DEPRECATION_WARNINGS
 class FChildren
 {
 public:
@@ -220,11 +219,8 @@ protected:
 
 #endif // !WITH_SLATE_DEBUGGING
 
-protected:
-	UE_DEPRECATED(5.0, "Direct access to Owner is now deprecated. Use the getter.")
-	SWidget* Owner;
-
 private:
+	SWidget* Owner;
 	FName Name;
 
 #if UE_WITH_SLATE_CHILDREN_DEBUGGING
@@ -232,4 +228,3 @@ private:
 	mutable uint8 Debug_DestroyedTag = 0xDC;
 #endif // WITH_SLATE_DEBUGGING
 };
-PRAGMA_ENABLE_DEPRECATION_WARNINGS
diff --git a/Engine/Source/Runtime/SlateCore/Public/Layout/WidgetPath.h b/Engine/Source/Runtime/SlateCore/Public/Layout/WidgetPath.h
--- a/Engine/Source/Runtime/SlateCore/Public/Layout/WidgetPath.h
+++ b/Engine/Source/Runtime/SlateCore/Public/Layout/WidgetPath.h
@@ -31,7 +31,6 @@ struct FWidgetMatcher
  * The window is needed for its ability to determine its own geometry, from which the geometries of the rest
  * of the widget can be determined.
  */
-PRAGMA_DISABLE_DEPRECATION_WARNINGS
 class SLATECORE_API FWidgetPath
 {
 public:
@@ -47,20 +46,12 @@ public:
 	 * @return a copy of the widget path down to and including the MarkerWidget. If the MarkerWidget is not found in the path, return an invalid path.
 	 */
 	FWidgetPath GetPathDownTo( TSharedRef<const SWidget> MarkerWidget ) const;
-	
-	/** Get the widget and associated virtual cursor at a given index. */
-	UE_DEPRECATED(5.0, "GetCursorAt is deprecated. Use GetVirtualPointerPosition instead.")
-	const TSharedPtr<const FVirtualPointerPosition>& GetCursorAt( int32 Index ) const;
 
 	/** Get the virtual representation of the mouse at each level in the widget path. */
 	TOptional<FVirtualPointerPosition> GetVirtualPointerPosition(int32 Index) const
 	{
 		return VirtualPointerPositions[Index];
 	}
-
-	/** @return true if the WidgetToFind is in this WidgetPath, false otherwise. */
-	UE_DEPRECATED(5.0, "ContainsWidget that takes a shared ptr is deprecated.")
-	bool ContainsWidget( TSharedRef<const SWidget> WidgetToFind ) const;
 	
 	/** @return true if the WidgetToFind is in this WidgetPath, false otherwise. */
 	bool ContainsWidget( const SWidget* WidgetToFind ) const;
@@ -170,11 +161,11 @@ public:
 	FArrangedChildren Widgets;
 
 	/** The top level window of this widget path. */
-	TSharedPtr< SWindow > TopLevelWindow;
+	TSharedPtr<SWindow> TopLevelWindow;
 
+private:
 	/** The virtual representation of the mouse at each level in the widget path.  Due to 3D widgets, the space you transition to can be completely arbitrary as you traverse the tree. */
-	UE_DEPRECATED(5.0, "Direct access to VirtualPointerPositions is now deprecated. Use the getter.")
-	TArray< TOptional<FVirtualPointerPosition> > VirtualPointerPositions;
+	TArray<TOptional<FVirtualPointerPosition>> VirtualPointerPositions;
 
 private:
 
@@ -195,7 +186,6 @@ private:
 	template<typename MatchRuleType>
 	static bool SearchForWidgetRecursively_Reverse( const MatchRuleType& MatchRule, const FArrangedWidget& InCandidate, FArrangedChildren& OutReversedPath, EVisibility VisibilityFilter = EVisibility::Visible );
 };
-PRAGMA_ENABLE_DEPRECATION_WARNINGS
 
 /**
  * Just like a WidgetPath, but uses weak pointers and does not store geometry.
@@ -239,11 +229,6 @@ public:
 		};
 	};
 
-
-	/** @return true if the WidgetToFind is in this WidgetPath, false otherwise. */
-	UE_DEPRECATED(5.0, "ContainsWidget that takes a shared ptr is deprecated.")
-	bool ContainsWidget(const TSharedRef<const SWidget>& WidgetToFind) const;
-
 	/** @return true if the WidgetToFind is in this WidgetPath, false otherwise. */
 	bool ContainsWidget(const SWidget* WidgetToFind) const;
 
diff --git a/Engine/Source/Runtime/SlateCore/Public/Rendering/DrawElementTypes.h b/Engine/Source/Runtime/SlateCore/Public/Rendering/DrawElementTypes.h
--- a/Engine/Source/Runtime/SlateCore/Public/Rendering/DrawElementTypes.h
+++ b/Engine/Source/Runtime/SlateCore/Public/Rendering/DrawElementTypes.h
@@ -315,12 +315,6 @@ public:
 	 */
 	void ApplyPositionOffset(UE::Slate::FDeprecateVector2DParameter InOffset);
 
-	UE_DEPRECATED(4.23, "GetClippingIndex has been deprecated.  If you were using this please use GetPrecachedClippingIndex instead.")
-	FORCEINLINE const int32 GetClippingIndex() const { return GetPrecachedClippingIndex(); }
-
-	UE_DEPRECATED(4.23, "SetClippingIndex has been deprecated.  If you were using this please use SetPrecachedClippingIndex instead.")
-	FORCEINLINE void SetClippingIndex(const int32 InClippingIndex) { SetPrecachedClippingIndex(InClippingIndex); }
-
 private:
 	void Init(FSlateWindowElementList& ElementList, EElementType InElementType, uint32 InLayer, const FPaintGeometry& PaintGeometry, ESlateDrawEffect InDrawEffects);
 
diff --git a/Engine/Source/Runtime/SlateCore/Public/Rendering/DrawElements.h b/Engine/Source/Runtime/SlateCore/Public/Rendering/DrawElements.h
--- a/Engine/Source/Runtime/SlateCore/Public/Rendering/DrawElements.h
+++ b/Engine/Source/Runtime/SlateCore/Public/Rendering/DrawElements.h
@@ -235,15 +235,6 @@ public:
 
 	SLATECORE_API ~FSlateWindowElementList();
 
-	/** @return Get the window that we will be painting */
-	UE_DEPRECATED(4.21, "FSlateWindowElementList::GetWindow is not thread safe but window element lists are accessed on multiple threads.  Please call GetPaintWindow instead")
-	FORCEINLINE TSharedPtr<SWindow> GetWindow() const
-	{
-		// check that we are in game thread or are in slate/movie loading thread
-		check(IsInGameThread() || IsInSlateThread());
-		return WeakPaintWindow.Pin();
-	}
-
 	/** @return Get the window that we will be painting */
 	SWindow* GetPaintWindow() const
 	{
diff --git a/Engine/Source/Runtime/SlateCore/Public/SlotBase.h b/Engine/Source/Runtime/SlateCore/Public/SlotBase.h
--- a/Engine/Source/Runtime/SlateCore/Public/SlotBase.h
+++ b/Engine/Source/Runtime/SlateCore/Public/SlotBase.h
@@ -25,9 +25,6 @@ public:
 	struct FSlotArguments {};
 
 public:
-	UE_DEPRECATED(5.0, "AttachWidgetParent is not used anymore. Use get SetOwner.")
-	void AttachWidgetParent(SWidget* InParent) { }
-
 	/**
 	 * Access the FChildren that own the slot.
 	 * The owner can be invalid when the slot is not attached.
@@ -116,13 +113,6 @@ private:
 	const FChildren* Owner;
 	/** The content widget of the slot. */
 	TSharedRef<SWidget> Widget;
-
-#if WITH_EDITORONLY_DATA
-protected:
-	/** The parent and owner of the slot. */
-	UE_DEPRECATED(5.0, "RawParentPtr is not used anymore. Use GetOwnerWidget.")
-	SWidget* RawParentPtr;
-#endif
 };
 
 
diff --git a/Engine/Source/Runtime/SlateCore/Public/Styling/SlateTypes.h b/Engine/Source/Runtime/SlateCore/Public/Styling/SlateTypes.h
--- a/Engine/Source/Runtime/SlateCore/Public/Styling/SlateTypes.h
+++ b/Engine/Source/Runtime/SlateCore/Public/Styling/SlateTypes.h
@@ -995,9 +995,8 @@ struct FEditableTextBoxStyle : public FSlateWidgetStyle
 	FEditableTextBoxStyle& operator=(const FEditableTextBoxStyle&) = default;
 
 	virtual ~FEditableTextBoxStyle() = default;
-	SLATECORE_API PRAGMA_ENABLE_DEPRECATION_WARNINGS
 
-	virtual void GetResources( TArray< const FSlateBrush* >& OutBrushes ) const override;
+	SLATECORE_API virtual void GetResources( TArray< const FSlateBrush* >& OutBrushes ) const override;
 
 	static SLATECORE_API const FName TypeName;
 	virtual const FName GetTypeName() const override { return TypeName; };
diff --git a/Engine/Source/Runtime/SlateCore/Public/Widgets/Images/SImage.h b/Engine/Source/Runtime/SlateCore/Public/Widgets/Images/SImage.h
--- a/Engine/Source/Runtime/SlateCore/Public/Widgets/Images/SImage.h
+++ b/Engine/Source/Runtime/SlateCore/Public/Widgets/Images/SImage.h
@@ -97,15 +97,6 @@ protected:
 	/** @return an attribute reference of DesiredSizeOverride */
 	TSlateAttributeRef<TOptional<FVector2D>> GetDesiredSizeOverrideAttribute() const { return TSlateAttributeRef<TOptional<FVector2D>>(SharedThis(this), DesiredSizeOverrideAttribute); }
 
-#if WITH_EDITORONLY_DATA
-	UE_DEPRECATED(5.0, "Direct access to Image is now deprecated. Use the setter or getter.")
-	FInvalidatableBrushAttribute Image;
-	UE_DEPRECATED(5.0, "Direct access to ColorAndOpacity is now deprecated. Use the setter or getter.")
-	TSlateDeprecatedTAttribute<FSlateColor> ColorAndOpacity;
-	UE_DEPRECATED(5.0, "Direct access to DesiredSizeOverride is now deprecated. Use the setter or getter.")
-	TSlateDeprecatedTAttribute<TOptional<FVector2D>> DesiredSizeOverride;
-#endif
-
 private:
 	/** The slate brush to draw for the ImageAttribute that we can invalidate. */
 	TSlateAttribute<const FSlateBrush*> ImageAttribute;
diff --git a/Engine/Source/Runtime/SlateCore/Public/Widgets/Images/SLayeredImage.h b/Engine/Source/Runtime/SlateCore/Public/Widgets/Images/SLayeredImage.h
--- a/Engine/Source/Runtime/SlateCore/Public/Widgets/Images/SLayeredImage.h
+++ b/Engine/Source/Runtime/SlateCore/Public/Widgets/Images/SLayeredImage.h
@@ -50,25 +50,9 @@ public:
 	/** Determines whether a Layer exists at the given index */
 	SLATECORE_API bool IsValidIndex(int32 Index) const;
 
-	/**
-	 * Gets the brush for a given layer
-	 * 
-	 * @return Null if the layer index is invalid. Otherwise, the brush for the given layer
-	 */
-	UE_DEPRECATED(5.0, "GetLayerBrush is not accessible anymore since it's attribute value may not have been updated yet.")
-	SLATECORE_API const FSlateBrush* GetLayerBrush(int32 Index) const;
-
 	/** Sets the brush for a given layer, if it exists */
 	SLATECORE_API void SetLayerBrush(int32 Index, TAttribute<const FSlateBrush*> Brush);
 
-	/** 
-	 * Gets the Color for a given layer
-	 *
-	 * @return Uninitialized fuschia color if the layer index is invalid. Otherwise the color for the given layer
-	 */
-	UE_DEPRECATED(5.0, "GetLayerColor is not accessible anymore since it's attribute value may not have been updated yet.")
-	SLATECORE_API FSlateColor GetLayerColor(int32 Index) const;
-
 	/** Sets the color for a given layer, if it exists. */
 	SLATECORE_API void SetLayerColor(int32 Index, TAttribute<FSlateColor> Color);
 
diff --git a/Engine/Source/Runtime/SlateCore/Public/Widgets/InvalidateWidgetReason.h b/Engine/Source/Runtime/SlateCore/Public/Widgets/InvalidateWidgetReason.h
--- a/Engine/Source/Runtime/SlateCore/Public/Widgets/InvalidateWidgetReason.h
+++ b/Engine/Source/Runtime/SlateCore/Public/Widgets/InvalidateWidgetReason.h
@@ -66,12 +66,6 @@ enum class EInvalidateWidgetReason : uint8
 	 * that you invalidate volatility so that it can be recalculated and cached.
 	 */
 	LayoutAndVolatility = Layout | Volatility,
-
-
-	/**
-	 * Do not use this ever unless you know what you are doing
-	 */
-	All UE_DEPRECATED(4.22, "EInvalidateWidget::All has been deprecated.  You probably wanted EInvalidateWidget::Layout but if you need more than that then use bitwise or to combine them") = 0xff
 };
 
 ENUM_CLASS_FLAGS(EInvalidateWidgetReason)
diff --git a/Engine/Source/Runtime/SlateCore/Public/Widgets/SBoxPanel.h b/Engine/Source/Runtime/SlateCore/Public/Widgets/SBoxPanel.h
--- a/Engine/Source/Runtime/SlateCore/Public/Widgets/SBoxPanel.h
+++ b/Engine/Source/Runtime/SlateCore/Public/Widgets/SBoxPanel.h
@@ -360,15 +360,6 @@ public:
 		{
 			SBoxPanel::TSlot<FSlot>::Construct(SlotOwner, MoveTemp(InArgs));
 		}
-
-		UE_DEPRECATED(5.0, "Chained AutoWidth is deprecated. Use the FSlotArgument or SetAutoWidth")
-		FSlot& AutoWidth() { SetAutoWidth(); return *this; }
-
-		UE_DEPRECATED(5.0, "Chained FillWidth is deprecated. Use the FSlotArgument or SetFillWidth")
-		FSlot& FillWidth(TAttribute<float> InStretchCoefficient) { SetFillWidth(InStretchCoefficient); return *this; }
-
-		UE_DEPRECATED(5.0, "Chained MaxWidth is deprecated. Use the FSlotArgument or SetMaxWidth")
-		FSlot& MaxWidth(TAttribute<float> InMaxWidth) { SetMaxWidth(InMaxWidth); return *this; }
 	};
 
 	static FSlot::FSlotArguments Slot()
@@ -501,15 +492,6 @@ public:
 		{
 			SBoxPanel::TSlot<FSlot>::Construct(SlotOwner, MoveTemp(InArgs));
 		}
-
-		UE_DEPRECATED(5.0, "Chained AutoHeight is deprecated. Use the FSlotArgument or SetAutoHeight")
-		FSlot& AutoHeight() { SetAutoHeight(); return *this; }
-
-		UE_DEPRECATED(5.0, "Chained FillWidth is deprecated. Use the FSlotArgument or SetFillWidth")
-		FSlot& FillHeight(TAttribute<float> InStretchCoefficient) { SetFillHeight(InStretchCoefficient); return *this; }
-
-		UE_DEPRECATED(5.0, "Chained MaxWidth is deprecated. Use the FSlotArgument or SetMaxHeight")
-		FSlot& MaxHeight(TAttribute<float> InMaxHeight) { SetMaxHeight(InMaxHeight); return *this; }
 	};
 
 	static FSlot::FSlotArguments Slot()
diff --git a/Engine/Source/Runtime/SlateCore/Public/Widgets/SCompoundWidget.h b/Engine/Source/Runtime/SlateCore/Public/Widgets/SCompoundWidget.h
--- a/Engine/Source/Runtime/SlateCore/Public/Widgets/SCompoundWidget.h
+++ b/Engine/Source/Runtime/SlateCore/Public/Widgets/SCompoundWidget.h
@@ -111,15 +111,6 @@ protected:
 	/** The slot that contains this widget's descendants.*/
 	FCompoundWidgetOneChildSlot ChildSlot;
 
-#if WITH_EDITORONLY_DATA
-	UE_DEPRECATED(5.0, "Direct access to ContentScale is now deprecated. Use the setter or getter.")
-	TSlateDeprecatedTAttribute<FVector2D> ContentScale;
-	UE_DEPRECATED(5.0, "Direct access to ColorAndOpacity is now deprecated. Use the setter or getter.")
-	TSlateDeprecatedTAttribute<FLinearColor> ColorAndOpacity;
-	UE_DEPRECATED(5.0, "Direct access to ForegroundColor is now deprecated. Use the setter or getter.")
-	TSlateDeprecatedTAttribute<FSlateColor> ForegroundColor;
-#endif
-
 private:
 	/** The layout scale to apply to this widget's contents; useful for animation. */
 	TSlateAttribute<FVector2D> ContentScaleAttribute;
diff --git a/Engine/Source/Runtime/SlateCore/Public/Widgets/SOverlay.h b/Engine/Source/Runtime/SlateCore/Public/Widgets/SOverlay.h
--- a/Engine/Source/Runtime/SlateCore/Public/Widgets/SOverlay.h
+++ b/Engine/Source/Runtime/SlateCore/Public/Widgets/SOverlay.h
@@ -37,7 +37,6 @@ class SOverlay : public SPanel
 public:	
 
 	/** A slot that support alignment of content and padding and z-order */
-	PRAGMA_DISABLE_DEPRECATION_WARNINGS
 	class SLATECORE_API FOverlaySlot : public TBasicLayoutWidgetSlot<FOverlaySlot>
 	{
 	public:
@@ -59,15 +58,13 @@ public:
 
 		void SetZOrder(int32 InOrder);
 
-	public:
+	private:
 		/**
 		 * Slots with larger ZOrder values will draw above slots with smaller ZOrder values. Slots
 		 * with the same ZOrder will simply draw in the order they were added.
 		 */
-		UE_DEPRECATED(5.0, "Direct access to ZOrder is now deprecated. Use the getter or setter.")
 		int32 ZOrder;
 	};
-	PRAGMA_ENABLE_DEPRECATION_WARNINGS
 
 	SLATE_BEGIN_ARGS( SOverlay )
 	{
diff --git a/Engine/Source/Runtime/SlateCore/Public/Widgets/SWidget.h b/Engine/Source/Runtime/SlateCore/Public/Widgets/SWidget.h
--- a/Engine/Source/Runtime/SlateCore/Public/Widgets/SWidget.h
+++ b/Engine/Source/Runtime/SlateCore/Public/Widgets/SWidget.h
@@ -262,41 +262,6 @@ protected:
 	};
 
 public:
-
-	/** Construct a SWidget based on initial parameters. */
-	 UE_DEPRECATED(4.27, "SWidget::Construct should not be called directly. Use SNew or SAssignNew to create a SWidget")
-	SLATECORE_API void Construct(
-		const TAttribute<FText>& InToolTipText,
-		const TSharedPtr<IToolTip>& InToolTip,
-		const TAttribute< TOptional<EMouseCursor::Type> >& InCursor,
-		const TAttribute<bool>& InEnabledState,
-		const TAttribute<EVisibility>& InVisibility,
-		const float InRenderOpacity,
-		const TAttribute<TOptional<FSlateRenderTransform>>& InTransform,
-		const TAttribute<FVector2D>& InTransformPivot,
-		const FName& InTag,
-		const bool InForceVolatile,
-		const EWidgetClipping InClipping,
-		const EFlowDirectionPreference InFlowPreference,
-		const TOptional<FAccessibleWidgetData>& InAccessibleData,
-		const TArray<TSharedRef<ISlateMetaData>>& InMetaData);
-
-	UE_DEPRECATED(4.27, "SWidget::SWidgetConstruct should not be called directly. Use SNew or SAssignNew to create a SWidget")
-	SLATECORE_API void SWidgetConstruct(const TAttribute<FText>& InToolTipText,
-		const TSharedPtr<IToolTip>& InToolTip,
-		const TAttribute< TOptional<EMouseCursor::Type> >& InCursor,
-		const TAttribute<bool>& InEnabledState,
-		const TAttribute<EVisibility>& InVisibility,
-		const float InRenderOpacity,
-		const TAttribute<TOptional<FSlateRenderTransform>>& InTransform,
-		const TAttribute<FVector2D>& InTransformPivot,
-		const FName& InTag,
-		const bool InForceVolatile,
-		const EWidgetClipping InClipping,
-		const EFlowDirectionPreference InFlowPreference,
-		const TOptional<FAccessibleWidgetData>& InAccessibleData,
-		const TArray<TSharedRef<ISlateMetaData>>& InMetaData);
-
 	//
 	// GENERAL EVENTS
 	//
@@ -846,9 +811,6 @@ public:
 #endif
 	}
 
-	UE_DEPRECATED(4.24, "GetRelativeLayoutScale(int32 ChildIndex, float LayoutScaleMultiplier), your widget will also need to set bHasRelativeLayoutScale in their Construct/ctor.")
-	virtual float GetRelativeLayoutScale(const FSlotBase& Child, float LayoutScaleMultiplier) const { return 1.0f; }
-
 	/** What is the Child's scale relative to this widget. */
 	SLATECORE_API virtual float GetRelativeLayoutScale(const int32 ChildIndex, float LayoutScaleMultiplier) const;
 
@@ -950,11 +912,6 @@ public:
 	 */
 	SLATECORE_API bool HasMouseCaptureByUser(int32 UserIndex, TOptional<int32> PointerIndex = TOptional<int32>()) const;
 
-protected:
-	/** Called when this widget had captured the mouse, but that capture has been revoked for some reason. */
-	UE_DEPRECATED(4.20, "Please use OnMouseCaptureLost(const FCaptureLostEvent& CaptureLostEvent)")
-	void OnMouseCaptureLost() { }
-
 public:
 	/** Called when this widget had captured the mouse, but that capture has been revoked for some reason. */
 	SLATECORE_API virtual void OnMouseCaptureLost(const FCaptureLostEvent& CaptureLostEvent);
@@ -1036,13 +993,6 @@ public:
 	/** @param InVisibility  should this widget be */
 	SLATECORE_API virtual void SetVisibility(TAttribute<EVisibility> InVisibility);
 
-	/**
-	 * @return is the widget visible and its parents also visible.
-	 * @note only valid if the widget is contained by an InvalidationRoot (the proxy is valid).
-	 */
-	UE_DEPRECATED(5.0, "IsFastPathVisible is deprecated and should not be used.")
-	SLATECORE_API bool IsFastPathVisible() const;
-
 #if WITH_ACCESSIBILITY
 	/**
 	 * Get the text that should be reported to the user when attempting to access this widget.
@@ -1147,9 +1097,6 @@ public:
 		bCachedVolatile = bForceVolatile || ComputeVolatility();
 	}
 
-	UE_DEPRECATED(5.0, "InvalidatePrepass is deprecated. Use the Invalidate(EInvalidateWidgetReason::Prepass) or use MarkPrepassAsDirty()")
-	SLATECORE_API void InvalidatePrepass();
-
 	/**
 	 * In fast path, if the widget is mark, do a full Prepass on its next update to calculate it's desired size.
 	 * This does not invalidate the widget.
@@ -1750,12 +1697,6 @@ protected:
 	  */
 	uint8 bClippingProxy : 1;
 
-#if WITH_EDITORONLY_DATA
-	/** Is this widget hovered? */
-	UE_DEPRECATED(5.0, "Direct access to bIsHovered is now deprecated. Use the IsHovered getter.")
-	uint8 bIsHovered : 1;
-#endif
-
 private:
 	/**
 	 * Whether this widget is a "tool tip force field".  That is, tool-tips should never spawn over the area
@@ -1873,21 +1814,6 @@ protected:
 	*/
 	FMargin CullingBoundsExtension;
 
-#if WITH_EDITORONLY_DATA
-	/** Whether or not this widget is enabled */
-	UE_DEPRECATED(5.0, "Direct access to EnabledState is now deprecated. Use the setter or getter.")
-	TSlateDeprecatedTAttribute<bool> EnabledState;
-	/** Is this widget visible, hidden or collapsed */
-	UE_DEPRECATED(5.0, "Direct access to Visibility is now deprecated. Use the setter or getter.")
-	TSlateDeprecatedTAttribute<EVisibility> Visibility;
-	/** Render transform of this widget. TOptional<> to allow code to skip expensive overhead if there is no render transform applied. */
-	UE_DEPRECATED(5.0, "Direct access to RenderTransform is now deprecated. Use the setter or getter.")
-	TSlateDeprecatedTAttribute< TOptional<FSlateRenderTransform> > RenderTransform;
-	/** Render transform pivot of this widget (in normalized local space) */
-	UE_DEPRECATED(5.0, "Direct access to RenderTransformPivot is now deprecated. Use the setter or getter.")
-	TAttribute<FVector2D> RenderTransformPivot;
-#endif
-
 	/** The opacity of the widget. Automatically applied during rendering. */
 	float RenderOpacity;
 
@@ -1927,12 +1853,6 @@ private:
 	uint8 Debug_DestroyedTag = 0xDC;
 #endif // WITH_SLATE_DEBUGGING
 
-#if !(UE_BUILD_SHIPPING || UE_BUILD_TEST)
-	UE_DEPRECATED(4.27, "Access to SWidget::Cursor is deprecated and will not function. Call SetCursor/GetCursor instead")
-	/** The cursor to show when the mouse is hovering over this widget. */
-	TAttribute<TOptional<EMouseCursor::Type>> Cursor;
-#endif
-
 #if UE_SLATE_WITH_WIDGET_UNIQUE_IDENTIFIER
 	/** The widget's id */
 	uint64 UniqueIdentifier;
