Fix build issue with macro balanced introduced with CL 32885267

diff --git a/Engine/Source/Runtime/SlateCore/Public/Styling/SlateTypes.h b/Engine/Source/Runtime/SlateCore/Public/Styling/SlateTypes.h
--- a/Engine/Source/Runtime/SlateCore/Public/Styling/SlateTypes.h
+++ b/Engine/Source/Runtime/SlateCore/Public/Styling/SlateTypes.h
@@ -988,13 +988,14 @@ struct FEditableTextBoxStyle : public FSlateWidgetStyle
 {
 	GENERATED_USTRUCT_BODY()
 
-	SLATECORE_API PRAGMA_DISABLE_DEPRECATION_WARNINGS
-	FEditableTextBoxStyle();
+	PRAGMA_DISABLE_DEPRECATION_WARNINGS
+	SLATECORE_API FEditableTextBoxStyle();
 	FEditableTextBoxStyle(const FEditableTextBoxStyle&) = default;
 
 	FEditableTextBoxStyle& operator=(const FEditableTextBoxStyle&) = default;
 
 	virtual ~FEditableTextBoxStyle() = default;
+	PRAGMA_ENABLE_DEPRECATION_WARNINGS
 
 	SLATECORE_API virtual void GetResources( TArray< const FSlateBrush* >& OutBrushes ) const override;
 
