Cache post buffer usage into renderbatches to avoid a virtual method call

#tests Console RR, PIE

diff --git a/Engine/Source/Runtime/SlateCore/Private/Rendering/ElementBatcher.cpp b/Engine/Source/Runtime/SlateCore/Private/Rendering/ElementBatcher.cpp
--- a/Engine/Source/Runtime/SlateCore/Private/Rendering/ElementBatcher.cpp
+++ b/Engine/Source/Runtime/SlateCore/Private/Rendering/ElementBatcher.cpp
@@ -572,6 +572,7 @@ void FSlateElementBatcher::AddCachedElements(FSlateCachedElementData& CachedElem
 	TArray<FSlateRenderBatch>& RenderBatchesSDR = BatchData->RenderBatches;
 	TArray<FSlateRenderBatch>& RenderBatchesHDR = BatchDataHDR->RenderBatches;
 
+	ESlatePostRT CachedPostBuffers = ESlatePostRT::None;
 	auto AddBatch = [&](const FSlateRenderBatch& CachedBatch)
 	{
 		if (EnumHasAnyFlags(CachedBatch.GetDrawFlags(), ESlateBatchDrawFlag::HDR))
@@ -582,11 +583,8 @@ void FSlateElementBatcher::AddCachedElements(FSlateCachedElementData& CachedElem
 		{
 			RenderBatchesSDR.Add(CachedBatch);
 		}
-
-		if (const FSlateShaderResource* ShaderResource = CachedBatch.GetShaderResource())
-		{
-			UsedSlatePostBuffers |= ShaderResource->GetUsedSlatePostBuffers();
-		}
+		
+		CachedPostBuffers |= CachedBatch.CachedUsedPostBuffers;
 	};
 
 	RenderBatchesSDR.Reserve(RenderBatchesSDR.Num() + CachedBatches.Num());
@@ -597,6 +595,8 @@ void FSlateElementBatcher::AddCachedElements(FSlateCachedElementData& CachedElem
 		AddBatch(CachedBatch);
 	}
 
+	UsedSlatePostBuffers |= CachedPostBuffers;
+
 	CachedElementData.CleanupUnusedClipStates();
 
 #if SLATE_CSV_TRACKER
@@ -3146,7 +3146,8 @@ FSlateRenderBatch& FSlateElementBatcher::CreateRenderBatch(
 
 	if (InResource)
 	{
-		UsedSlatePostBuffers |= InResource->GetUsedSlatePostBuffers();
+		NewBatch.CachedUsedPostBuffers = InResource->GetUsedSlatePostBuffers();
+		UsedSlatePostBuffers |= NewBatch.CachedUsedPostBuffers;
 	}
 
 	return NewBatch;
@@ -3171,7 +3172,8 @@ FSlateRenderBatch& FSlateElementBatcher::CreateRenderBatch(
 
 	if (InResource)
 	{
-		UsedSlatePostBuffers |= InResource->GetUsedSlatePostBuffers();
+		NewBatch.CachedUsedPostBuffers = InResource->GetUsedSlatePostBuffers();
+		UsedSlatePostBuffers |= NewBatch.CachedUsedPostBuffers;
 	}
 
 	return NewBatch;
diff --git a/Engine/Source/Runtime/SlateCore/Private/Rendering/SlateRenderBatch.cpp b/Engine/Source/Runtime/SlateCore/Private/Rendering/SlateRenderBatch.cpp
--- a/Engine/Source/Runtime/SlateCore/Private/Rendering/SlateRenderBatch.cpp
+++ b/Engine/Source/Runtime/SlateCore/Private/Rendering/SlateRenderBatch.cpp
@@ -40,6 +40,7 @@ FSlateRenderBatch::FSlateRenderBatch(
 	, ShaderType(InShaderType)
 	, DrawPrimitiveType(InPrimitiveType)
 	, DrawEffects(InDrawEffects)
+	, CachedUsedPostBuffers(ESlatePostRT::None)
 	, bIsMergable(true)
 	, bIsMerged(false)
 {
diff --git a/Engine/Source/Runtime/SlateCore/Public/Rendering/SlateRenderBatch.h b/Engine/Source/Runtime/SlateCore/Public/Rendering/SlateRenderBatch.h
--- a/Engine/Source/Runtime/SlateCore/Public/Rendering/SlateRenderBatch.h
+++ b/Engine/Source/Runtime/SlateCore/Public/Rendering/SlateRenderBatch.h
@@ -230,6 +230,8 @@ public:
 
 	ESlateDrawEffect DrawEffects;
 
+	ESlatePostRT CachedUsedPostBuffers;
+
 	/** Whether or not the batch can be merged with others in the same layer */
 	uint8 bIsMergable : 1;
 
