Mark certain slate classes as final, in attempt to support devirtualization

#tests Console RR, PIE

diff --git a/Engine/Source/Runtime/SlateCore/Public/Layout/Children.h b/Engine/Source/Runtime/SlateCore/Public/Layout/Children.h
--- a/Engine/Source/Runtime/SlateCore/Public/Layout/Children.h
+++ b/Engine/Source/Runtime/SlateCore/Public/Layout/Children.h
@@ -15,7 +15,7 @@
  * This data structure can be used to link multiple FChildren under a single accessor so you can always return
  * all children from GetChildren, but internally manage them in their own child lists.
  */
-class FCombinedChildren : public FChildren
+class FCombinedChildren final : public FChildren
 {
 public:
 	using FChildren::FChildren;
@@ -145,7 +145,7 @@ protected:
  * Widgets with no Children can return an instance of FNoChildren.
  * For convenience a shared instance FNoChildren::NoChildrenInstance can be used.
  */
-class FNoChildren : public FChildren
+class FNoChildren final : public FChildren
 {
 public:
 	static SLATECORE_API FNoChildren NoChildrenInstance;
@@ -208,7 +208,7 @@ private:
  *      and then they'll just stop being shown.
  */
 template <typename ChildType>
-class TWeakChild : public FChildren
+class TWeakChild final : public FChildren
 {
 public:
 	using FChildren::FChildren;
@@ -394,7 +394,7 @@ private:
 
 
 /** A FChildren that has only one child. */
-class FSingleWidgetChildrenWithSlot : public TSingleWidgetChildrenWithSlot<FSingleWidgetChildrenWithSlot>
+class FSingleWidgetChildrenWithSlot final : public TSingleWidgetChildrenWithSlot<FSingleWidgetChildrenWithSlot>
 {
 public:
 	using TSingleWidgetChildrenWithSlot<FSingleWidgetChildrenWithSlot>::TSingleWidgetChildrenWithSlot;
@@ -457,7 +457,7 @@ class FSingleWidgetChildrenWithBasicLayoutSlot : public TSingleWidgetChildrenWit
  * ChildType must have a public member SWidget* Widget;
  */
 template<typename SlotType>
-class TPanelChildren : public FChildren
+class TPanelChildren final : public FChildren
 {
 private:
 	TArray<TUniquePtr<SlotType>> Children;
@@ -662,7 +662,7 @@ public:
 
 public:
 	/** At the end of the scope a slot will be constructed and added to the FChildren. */
-	struct FScopedWidgetSlotArguments : public SlotType::FSlotArguments
+	struct FScopedWidgetSlotArguments final : public SlotType::FSlotArguments
 	{
 	public:
 		FScopedWidgetSlotArguments(TUniquePtr<SlotType> InSlot, TPanelChildren<SlotType>& InChildren, int32 InIndex)
@@ -832,7 +832,7 @@ private:
  * TSlotlessChildren should not be used for general-purpose widgets.
  */
 template<typename ChildType>
-class TSlotlessChildren : public FChildren
+class TSlotlessChildren final : public FChildren
 {
 private:
 	TArray<TSharedRef<ChildType>> Children;
@@ -1024,7 +1024,7 @@ private:
 
 /** Required to implement GetChildren() in a way that can dynamically return the currently active child. */
 template<typename SlotType>
-class TOneDynamicChild : public FChildren
+class TOneDynamicChild final : public FChildren
 {
 public:
 	TOneDynamicChild(SWidget* InOwner, TPanelChildren<SlotType>* InAllChildren, const TAttribute<int32>* InWidgetIndex)
diff --git a/Engine/Source/Runtime/SlateCore/Public/Types/Attributes/SlateAttributeBase.inl b/Engine/Source/Runtime/SlateCore/Public/Types/Attributes/SlateAttributeBase.inl
--- a/Engine/Source/Runtime/SlateCore/Public/Types/Attributes/SlateAttributeBase.inl
+++ b/Engine/Source/Runtime/SlateCore/Public/Types/Attributes/SlateAttributeBase.inl
@@ -554,7 +554,7 @@ namespace SlateAttributePrivate
 
 	private:
 		template<typename SlateAttributeType>
-		class FSlateAttributeGetterWrapper : public ISlateAttributeGetter
+		class FSlateAttributeGetterWrapper final : public ISlateAttributeGetter
 		{
 		public:
 			using ObjectType = typename SlateAttributeType::ObjectType;
