Remove unused virtual destructor causing static analysis warnings.

diff --git a/Engine/Source/Runtime/SlateCore/Public/Layout/Children.h b/Engine/Source/Runtime/SlateCore/Public/Layout/Children.h
--- a/Engine/Source/Runtime/SlateCore/Public/Layout/Children.h
+++ b/Engine/Source/Runtime/SlateCore/Public/Layout/Children.h
@@ -684,7 +684,7 @@ public:
 		FScopedWidgetSlotArguments(FScopedWidgetSlotArguments&&) = default;
 		FScopedWidgetSlotArguments& operator=(FScopedWidgetSlotArguments&&) = default;
 	
-		virtual ~FScopedWidgetSlotArguments()
+		~FScopedWidgetSlotArguments()
 		{
 			if (const SlotType* SlotPtr = this->GetSlot())	// Is nullptr when the FScopedWidgetSlotArguments was moved-constructed.
 			{
