Prevent crashes if custom drawer is GC'd before element is added

#tests PIE

diff --git a/Engine/Source/Runtime/SlateCore/Private/Rendering/ElementBatcher.cpp b/Engine/Source/Runtime/SlateCore/Private/Rendering/ElementBatcher.cpp
--- a/Engine/Source/Runtime/SlateCore/Private/Rendering/ElementBatcher.cpp
+++ b/Engine/Source/Runtime/SlateCore/Private/Rendering/ElementBatcher.cpp
@@ -3010,10 +3010,13 @@ void FSlateElementBatcher::AddCustomElement( const FSlateCustomDrawerElement& Dr
 {
 	const int32 Layer = DrawElement.GetLayer();
 
-	FSlateRenderBatch& RenderBatch = CreateRenderBatch(Layer, FShaderParams(), nullptr, ESlateDrawPrimitive::None, ESlateShader::Default, ESlateDrawEffect::None, ESlateBatchDrawFlag::None, DrawElement);
-	RenderBatch.CustomDrawer = DrawElement.CustomDrawer.Pin().Get();
-	RenderBatch.bIsMergable = false;
-	RenderBatch.CustomDrawer->PostCustomElementAdded(*this);
+	if (TSharedPtr<ICustomSlateElement> CustomDrawerPinned = DrawElement.CustomDrawer.Pin())
+	{
+		FSlateRenderBatch& RenderBatch = CreateRenderBatch(Layer, FShaderParams(), nullptr, ESlateDrawPrimitive::None, ESlateShader::Default, ESlateDrawEffect::None, ESlateBatchDrawFlag::None, DrawElement);
+		RenderBatch.CustomDrawer = CustomDrawerPinned.Get();
+		RenderBatch.bIsMergable = false;
+		RenderBatch.CustomDrawer->PostCustomElementAdded(*this);
+	}
 }
 
 void FSlateElementBatcher::AddCustomVerts(const FSlateCustomVertsElement& DrawElement)
