Slate: Fix cached volatility that was not updated when the Widget doesn't have a FastPathProxyHandle yet.

diff --git a/Engine/Source/Runtime/SlateCore/Private/Widgets/SWidget.cpp b/Engine/Source/Runtime/SlateCore/Private/Widgets/SWidget.cpp
--- a/Engine/Source/Runtime/SlateCore/Private/Widgets/SWidget.cpp
+++ b/Engine/Source/Runtime/SlateCore/Private/Widgets/SWidget.cpp
@@ -1239,9 +1239,10 @@ void SWidget::Invalidate(EInvalidateWidgetReason InvalidateReason)
 		InvalidateReason |= EInvalidateWidgetReason::Layout;
 	}
 
+	// NB Advanced_InvalidateVolatility needs to be called to update the bCachedVolatility
+	const bool bVolatilityChanged = EnumHasAnyFlags(InvalidateReason, EInvalidateWidgetReason::Volatility) ? Advanced_InvalidateVolatility() : false;
 	if(FastPathProxyHandle.IsValid(this))
 	{
-		const bool bVolatilityChanged = EnumHasAnyFlags(InvalidateReason, EInvalidateWidgetReason::Volatility) ? Advanced_InvalidateVolatility() : false;
 		if (bVolatilityChanged)
 		{
 			SCOPED_NAMED_EVENT(SWidget_UpdateFastPathVolatility, FColor::Red);
