Added default initializers to FGeometry's FVector2f members.

diff --git a/Engine/Source/Runtime/SlateCore/Public/Layout/Geometry.h b/Engine/Source/Runtime/SlateCore/Public/Layout/Geometry.h
--- a/Engine/Source/Runtime/SlateCore/Public/Layout/Geometry.h
+++ b/Engine/Source/Runtime/SlateCore/Public/Layout/Geometry.h
@@ -615,7 +615,7 @@ public:
 	 * 
 	 * Absolute coordinates could be either desktop or window space depending on what space the root of the widget hierarchy is in.
 	 */
-	const FVector2f AbsolutePosition;	
+	const FVector2f AbsolutePosition = FVector2f::ZeroVector;	
 
 	/** 
 	 * !!! DEPRECATED !!! 
@@ -624,7 +624,7 @@ public:
 	 * If you know your children have no additional scale applied to them, you can use this as the Local->Parent layout transform. If your children
 	 * DO have additional scale applied, there is no way to determine the actual Local->Parent layout transform, since the scale is accumulated.
 	 */
-	const FVector2f /*Local*/Position;
+	const FVector2f /*Local*/Position = FVector2f::ZeroVector;
 
 private:
 
