Slate: Enabled the invalidation option EnableReindexLayerId by default. When a widget painted in invalidation returns a bigger LayerId than previously, the other widget need to be adjusted to respect the z ordering.

diff --git a/Engine/Source/Runtime/SlateCore/Private/FastUpdate/SlateInvalidationRoot.cpp b/Engine/Source/Runtime/SlateCore/Private/FastUpdate/SlateInvalidationRoot.cpp
--- a/Engine/Source/Runtime/SlateCore/Private/FastUpdate/SlateInvalidationRoot.cpp
+++ b/Engine/Source/Runtime/SlateCore/Private/FastUpdate/SlateInvalidationRoot.cpp
@@ -157,7 +157,7 @@ FAutoConsoleVariableRef CVarSlateInvalidationWidgetListNumElementLeftBeforeSplit
 	GSlateInvalidationWidgetListNumberElementLeftBeforeSplitting,
 	TEXT("With the invalidation system, when splitting, only split the array when the number of element left is under X."));
 
-bool GSlateInvalidationEnableReindexLayerId = false;
+bool GSlateInvalidationEnableReindexLayerId = true;
 FAutoConsoleVariableRef CVarSlateInvalidationEnableReindexLayerId(
 	TEXT("Slate.InvalidationList.EnableReindexLayerId"),
 	GSlateInvalidationEnableReindexLayerId,
