Slate: Fix rendering priority order for ChildOrder invalidation.

diff --git a/Engine/Source/Runtime/SlateCore/Private/Debugging/ConsoleSlateDebuggerInvalidate.cpp b/Engine/Source/Runtime/SlateCore/Private/Debugging/ConsoleSlateDebuggerInvalidate.cpp
--- a/Engine/Source/Runtime/SlateCore/Private/Debugging/ConsoleSlateDebuggerInvalidate.cpp
+++ b/Engine/Source/Runtime/SlateCore/Private/Debugging/ConsoleSlateDebuggerInvalidate.cpp
@@ -289,14 +289,26 @@ int32 FConsoleSlateDebuggerInvalidate::GetInvalidationPriority(EInvalidateWidget
 		return 50;
 	}
 
-	if (EnumHasAnyFlags(InvalidationInfo, EInvalidateWidgetReason::Prepass | EInvalidateWidgetReason::Layout | EInvalidateWidgetReason::ChildOrder | EInvalidateWidgetReason::Visibility | EInvalidateWidgetReason::RenderTransform))
+	if (EnumHasAnyFlags(InvalidationInfo, EInvalidateWidgetReason::ChildOrder))
+	{
+		return 45;
+	}
+	if (EnumHasAnyFlags(InvalidationInfo, EInvalidateWidgetReason::Prepass))
 	{
 		return 40;
 	}
-	else if (EnumHasAnyFlags(InvalidationInfo, EInvalidateWidgetReason::Volatility))
+	if (EnumHasAnyFlags(InvalidationInfo, EInvalidateWidgetReason::Layout))
+	{
+		return 35;
+	}
+	if (EnumHasAnyFlags(InvalidationInfo, EInvalidateWidgetReason::Visibility | EInvalidateWidgetReason::RenderTransform))
 	{
 		return 30;
 	}
+	else if (EnumHasAnyFlags(InvalidationInfo, EInvalidateWidgetReason::Volatility))
+	{
+		return 25;
+	}
 	else if (EnumHasAnyFlags(InvalidationInfo, EInvalidateWidgetReason::Paint))
 	{
 		return 20;
@@ -319,7 +331,11 @@ const FLinearColor& FConsoleSlateDebuggerInvalidate::GetColor(const FInvalidatio
 		return DrawRootScreenPositionColor;
 	}
 
-	if (EnumHasAnyFlags(InvalidationInfo.WidgetReason, EInvalidateWidgetReason::Prepass))
+	if (EnumHasAnyFlags(InvalidationInfo.WidgetReason, EInvalidateWidgetReason::ChildOrder))
+	{
+		return DrawWidgetChildOrderColor;
+	}
+	else if (EnumHasAnyFlags(InvalidationInfo.WidgetReason, EInvalidateWidgetReason::Prepass))
 	{
 		return DrawWidgetPrepassColor;
 	}
@@ -327,18 +343,6 @@ const FLinearColor& FConsoleSlateDebuggerInvalidate::GetColor(const FInvalidatio
 	{
 		return DrawWidgetLayoutColor;
 	}
-	else if (EnumHasAnyFlags(InvalidationInfo.WidgetReason, EInvalidateWidgetReason::Paint))
-	{
-		return DrawWidgetPaintColor;
-	}
-	else if (EnumHasAnyFlags(InvalidationInfo.WidgetReason, EInvalidateWidgetReason::Volatility))
-	{
-		return DrawWidgetVolatilityColor;
-	}
-	else if (EnumHasAnyFlags(InvalidationInfo.WidgetReason, EInvalidateWidgetReason::ChildOrder))
-	{
-		return DrawWidgetChildOrderColor;
-	}
 	else if (EnumHasAnyFlags(InvalidationInfo.WidgetReason, EInvalidateWidgetReason::RenderTransform))
 	{
 		return DrawWidgetRenderTransformColor;
@@ -347,6 +351,14 @@ const FLinearColor& FConsoleSlateDebuggerInvalidate::GetColor(const FInvalidatio
 	{
 		return DrawWidgetVisibilityColor;
 	}
+	else if (EnumHasAnyFlags(InvalidationInfo.WidgetReason, EInvalidateWidgetReason::Volatility))
+	{
+		return DrawWidgetVolatilityColor;
+	}
+	else if (EnumHasAnyFlags(InvalidationInfo.WidgetReason, EInvalidateWidgetReason::Paint))
+	{
+		return DrawWidgetPaintColor;
+	}
 
 	check(false);
 	return FLinearColor::Black;
@@ -477,7 +489,10 @@ void FConsoleSlateDebuggerInvalidate::HandleWidgetInvalidated(const FSlateDebugg
 		// Is this invalidation is more important for the display?
 		//Z->C[RenderTransform] to B->C[Layout]
 		//NB we use < instead of <= so only the first incoming invalidation will be considered 
-		FoundInvalidated->UpdateInvalidationReason(Args, InvalidationPriority);
+		if (FoundInvalidated->InvalidationPriority < InvalidationPriority)
+		{
+			FoundInvalidated->UpdateInvalidationReason(Args, InvalidationPriority);
+		}
 	}
 	else
 	{
