Slate: Reset the cache handle of widgets that are removed from the WidgetList.

diff --git a/Engine/Source/Runtime/SlateCore/Private/FastUpdate/SlateInvalidationRoot.cpp b/Engine/Source/Runtime/SlateCore/Private/FastUpdate/SlateInvalidationRoot.cpp
--- a/Engine/Source/Runtime/SlateCore/Private/FastUpdate/SlateInvalidationRoot.cpp
+++ b/Engine/Source/Runtime/SlateCore/Private/FastUpdate/SlateInvalidationRoot.cpp
@@ -141,6 +141,14 @@ static FAutoConsoleVariableRef CVarSlateInvalidationRootVerifyWidgetsAreUpdatedO
 	TEXT("Verify that the widgets are painted only once per tick.")
 );
 
+bool GSlateInvalidationRootVerifyCachedElementDataList = false;
+static FAutoConsoleVariableRef CVarSlateInvalidationRootVerifyCachedElementDataList(
+	TEXT("Slate.InvalidationRoot.VerifyCachedElementDataList"),
+	GSlateInvalidationRootVerifyCachedElementDataList,
+	TEXT("Verify that the widgets cached are valid.")
+);
+void VerifyCachedElementDataList(const FSlateInvalidationRoot* InvalidationRoot);
+
 #endif //UE_SLATE_WITH_INVALIDATIONWIDGETLIST_DEBUGGING
 
 
@@ -415,6 +423,14 @@ FSlateInvalidationResult FSlateInvalidationRoot::PaintInvalidationRoot(const FSl
 		check(RootWidget == FastWidgetPathList->GetRoot().Pin());
 
 		Result.bRepaintedWidgets = PaintFastPath(Context);
+
+		// Once everything is painted and ticked, new widget (SListView, SRichBox, SScrollBox) might have created new widget.
+		//It might create a 1 frame hitch because the widget is in the cached and should not.
+		if (!bNeedsSlowPath && WidgetsNeedingPreUpdate->Num() > 0)
+		{
+			SCOPED_NAMED_EVENT(Slate_InvalidationProcessing_PreUpdate, FColor::Blue);
+			ProcessPreUpdate();
+		}
 	}
 
 	if (Context.bAllowFastPathUpdate)
@@ -425,6 +441,10 @@ FSlateInvalidationResult FSlateInvalidationRoot::PaintInvalidationRoot(const FSl
 	FinalUpdateList.Reset();
 
 #if UE_SLATE_WITH_INVALIDATIONWIDGETLIST_DEBUGGING
+	if (GSlateInvalidationRootVerifyCachedElementDataList && !bNeedsSlowPath)
+	{
+		VerifyCachedElementDataList(this);
+	}
 	if (GSlateInvalidationRootVerifyHittestGrid && Context.bAllowFastPathUpdate)
 	{
 		VerifyHittest(InvalidationRootWidget, GetFastPathWidgetList(), GetHittestGrid());
@@ -827,10 +847,6 @@ void FSlateInvalidationRoot::ProcessPreUpdate()
 	{
 		VerifyWidgetsUpdateList_BeforeProcessPreUpdate(GetRootWidget(), FastWidgetPathList.Get(), WidgetsNeedingPreUpdate.Get(), WidgetsNeedingPostUpdate.Get(), FinalUpdateList);
 	}
-	if (GSlateInvalidationRootVerifySlateAttribute)
-	{
-		VerifySlateAttribute_BeforeUpdate(*FastWidgetPathList);
-	}
 #endif
 
 	TGuardValue<bool> Tmp(bProcessingPreUpdate, true);
@@ -1041,6 +1057,12 @@ void FSlateInvalidationRoot::ProcessAttributeUpdate()
 {
 	TGuardValue<bool> Tmp(bProcessingAttributeUpdate, true);
 
+#if UE_SLATE_WITH_INVALIDATIONWIDGETLIST_DEBUGGING
+	if (GSlateInvalidationRootVerifySlateAttribute)
+	{
+		VerifySlateAttribute_BeforeUpdate(*FastWidgetPathList);
+	}
+#endif
 
 	FSlateInvalidationWidgetList::FWidgetAttributeIterator AttributeItt = FastWidgetPathList->CreateWidgetAttributeIterator();
 	while (AttributeItt.IsValid())
@@ -1092,6 +1114,7 @@ void FSlateInvalidationRoot::ProcessAttributeUpdate()
 			AttributeItt.Advance();
 		}
 	}
+
 #if UE_SLATE_WITH_INVALIDATIONWIDGETLIST_DEBUGGING
 	if (GSlateInvalidationRootVerifySlateAttribute)
 	{
@@ -1693,7 +1716,7 @@ void VerifyWidgetVisibility(FSlateInvalidationWidgetList& WidgetList)
 					{
 						UE_SLATE_LOG_ERROR_IF_FALSE(!Widget->GetPersistentState().CachedElementHandle.HasCachedElements()
 							, CVarSlateInvalidationRootVerifyWidgetVisibility
-							, TEXT("Widget '%s' has cached element and is not visibled.")
+							, TEXT("Widget '%s' has cached element and is not visible.")
 							, *FReflectionMetaData::GetWidgetDebugInfo(Widget));
 					}
 				}
@@ -1896,9 +1919,6 @@ void VerifyWidgetsUpdateList_AfterProcessPostUpdate(const TSharedRef<SWidget>& R
 		return;
 	}
 
-	UE_SLATE_LOG_ERROR_IF_FALSE(WidgetsNeedingPreUpdate->Num() == 0
-		, CVarSlateInvalidationRootVerifyWidgetsUpdateList
-		, TEXT("The list of Pre Update should already been processed."));
 	UE_SLATE_LOG_ERROR_IF_FALSE(WidgetsNeedingPostUpdate->Num() == 0
 		, CVarSlateInvalidationRootVerifyWidgetsUpdateList
 		, TEXT("The list of Post Update should already been processed."));
@@ -1940,6 +1960,33 @@ void VerifySlateAttribute_AfterUpdate(const FSlateInvalidationWidgetList& FastWi
 		, TEXT("The VerifySlateAttribute failed in post."));
 }
 
+void VerifyCachedElementDataList(const FSlateInvalidationRoot* InvalidationRoot)
+{
+	for (const TSharedPtr<FSlateCachedElementList>& CachedElementList : InvalidationRoot->GetCachedElements().GetCachedElementLists())
+	{
+		if (CachedElementList)
+		{
+			const SWidget* WidgetPtr = CachedElementList->OwningWidget;
+			UE_SLATE_LOG_ERROR_IF_FALSE(WidgetPtr
+				, CVarSlateInvalidationRootVerifyCachedElementDataList
+				, TEXT("A widget is invalid and should not be cached."));
+			if (WidgetPtr)
+			{
+				UE_SLATE_LOG_ERROR_IF_FALSE(WidgetPtr->GetProxyHandle().IsValid(WidgetPtr)
+					, CVarSlateInvalidationRootVerifyCachedElementDataList
+					, TEXT("The widget '%s' proxy is invalid and should not be cached.")
+					, *FReflectionMetaData::GetWidgetDebugInfo(WidgetPtr));
+				UE_SLATE_LOG_ERROR_IF_FALSE((WidgetPtr->GetProxyHandle().GetWidgetVisibility(WidgetPtr).IsVisible()
+					|| WidgetPtr->GetProxyHandle().HasAllInvalidationReason(WidgetPtr, EInvalidateWidgetReason::Visibility)
+					|| CachedElementList->IsEmpty())
+					, CVarSlateInvalidationRootVerifyCachedElementDataList
+					, TEXT("The widget '%s' proxy is not visible and should not be cached.")
+					, *FReflectionMetaData::GetWidgetDebugInfo(WidgetPtr));
+			}
+		}
+	}
+}
+
 #undef UE_SLATE_LOG_ERROR_IF_FALSE
 #endif //UE_SLATE_WITH_INVALIDATIONWIDGETLIST_DEBUGGING
 
diff --git a/Engine/Source/Runtime/SlateCore/Private/FastUpdate/SlateInvalidationWidgetList.cpp b/Engine/Source/Runtime/SlateCore/Private/FastUpdate/SlateInvalidationWidgetList.cpp
--- a/Engine/Source/Runtime/SlateCore/Private/FastUpdate/SlateInvalidationWidgetList.cpp
+++ b/Engine/Source/Runtime/SlateCore/Private/FastUpdate/SlateInvalidationWidgetList.cpp
@@ -21,11 +21,6 @@
 DECLARE_CYCLE_STAT(TEXT("WidgetList ChildOrder Invalidation"), STAT_WidgetList_ProcessChildOrderInvalidation, STATGROUP_Slate);
 
 #define UE_SLATE_WITH_WIDGETLIST_UPDATEONLYWHATISNEEDED 0
-#define UE_SLATE_WITH_WIDGETLIST_ASSIGNINVALIDPROXYWHENREMOVED 0
-
-#if UE_SLATE_WITH_WIDGETLIST_ASSIGNINVALIDPROXYWHENREMOVED
-uint16 GSlateInvalidationWidgetIndex_RemovedIndex = 0xffee;
-#endif
 
 // See FSlateInvalidationWidgetSortOrder::FSlateInvalidationWidgetSortOrder
 int32 FSlateInvalidationWidgetList::FArguments::MaxPreferedElementsNum = (1 << 10) - 1;
@@ -1221,19 +1216,16 @@ void FSlateInvalidationWidgetList::Internal_RemoveRangeFromSameParent(const FInd
 	// Destroy/Remove the data that is not needed anymore
 	{
 		FSlateInvalidationWidgetList* Self = this;
-		auto SetFakeInvalidatWidgetHandle = [Self](IndexType ArrayIndex, int32 StartIndex, int32 Num)
+		auto ResetInvalidationWidgetHandle = [Self](IndexType ArrayIndex, int32 StartIndex, int32 Num)
 		{
-#if UE_SLATE_WITH_WIDGETLIST_ASSIGNINVALIDPROXYWHENREMOVED
-			const FSlateInvalidationWidgetIndex SlateInvalidationWidgetIndexRemoved = { GSlateInvalidationWidgetIndex_RemovedIndex, GSlateInvalidationWidgetIndex_RemovedIndex };
 			for (int32 ElementIndex = StartIndex; ElementIndex < Num; ++ElementIndex)
 			{
 				InvalidationWidgetType& InvalidationWidget = Self->Data[ArrayIndex].ElementList[ElementIndex];
 				if (SWidget* Widget = InvalidationWidget.GetWidget())
 				{
-					Widget->SetFastPathProxyHandle(FWidgetProxyHandle{ Self->Owner, SlateInvalidationWidgetIndexRemoved, FSlateInvalidationWidgetSortOrder() });
+					Widget->SetFastPathProxyHandle(FWidgetProxyHandle());
 				}
 			}
-#endif //UE_SLATE_WITH_WIDGETLIST_ASSIGNINVALIDPROXYWHENREMOVED
 		};
 		auto ResetInvalidationWidget = [Self](IndexType ArrayIndex, IndexType StartIndex, int32 Num)
 		{
@@ -1243,7 +1235,12 @@ void FSlateInvalidationWidgetList::Internal_RemoveRangeFromSameParent(const FInd
 				ElementListType& ResetElementList = ArrayNode.ElementList;
 				for (int32 ElementIndex = StartIndex; ElementIndex < Num; ++ElementIndex)
 				{
-					ResetElementList[ElementIndex].ResetWidget();
+					InvalidationWidgetType& InvalidationWidget = ResetElementList[ElementIndex];
+					if (SWidget* Widget = InvalidationWidget.GetWidget())
+					{
+						Widget->SetFastPathProxyHandle(FWidgetProxyHandle());
+					}
+					InvalidationWidget.ResetWidget();
 				}
 				ArrayNode.RemoveElementIndexBetweenOrEqualThan(StartIndex, (IndexType)(StartIndex + Num - 1));
 			}
@@ -1272,7 +1269,7 @@ void FSlateInvalidationWidgetList::Internal_RemoveRangeFromSameParent(const FInd
 				IndexType CurrentArrayIndex = NextToRemove;
 				do
 				{
-					SetFakeInvalidatWidgetHandle(CurrentArrayIndex, Data[CurrentArrayIndex].StartIndex, Data[CurrentArrayIndex].ElementList.Num());
+					ResetInvalidationWidgetHandle(CurrentArrayIndex, Data[CurrentArrayIndex].StartIndex, Data[CurrentArrayIndex].ElementList.Num());
 					CurrentArrayIndex = NextToRemove;
 					NextToRemove = (IndexType)Data[CurrentArrayIndex].NextArrayIndex;
 					RemoveDataNode(CurrentArrayIndex);
@@ -1284,7 +1281,7 @@ void FSlateInvalidationWidgetList::Internal_RemoveRangeFromSameParent(const FInd
 		if (bShouldCutArray && !bRangeIsInSameElementArray)
 		{
 			// The valid data in the array was moved (ie. ii)
-			SetFakeInvalidatWidgetHandle(Range.GetInclusiveMaxWidgetIndex().ArrayIndex, Data[Range.GetInclusiveMaxWidgetIndex().ArrayIndex].StartIndex, Range.GetInclusiveMaxWidgetIndex().ElementIndex + 1);
+			ResetInvalidationWidgetHandle(Range.GetInclusiveMaxWidgetIndex().ArrayIndex, Data[Range.GetInclusiveMaxWidgetIndex().ArrayIndex].StartIndex, Range.GetInclusiveMaxWidgetIndex().ElementIndex + 1);
 			RemoveDataNode(Range.GetInclusiveMaxWidgetIndex().ArrayIndex);
 		}
 		else if (!bRangeIsInSameElementArray)
@@ -1292,7 +1289,6 @@ void FSlateInvalidationWidgetList::Internal_RemoveRangeFromSameParent(const FInd
 			// Set StartIndex (ie. i)
 			check(Range.GetInclusiveMinWidgetIndex().ArrayIndex != Range.GetInclusiveMaxWidgetIndex().ArrayIndex);
 
-			SetFakeInvalidatWidgetHandle(Range.GetInclusiveMaxWidgetIndex().ArrayIndex, Data[Range.GetInclusiveMaxWidgetIndex().ArrayIndex].StartIndex, Range.GetInclusiveMaxWidgetIndex().ElementIndex + 1);
 			ResetInvalidationWidget(Range.GetInclusiveMaxWidgetIndex().ArrayIndex, Data[Range.GetInclusiveMaxWidgetIndex().ArrayIndex].StartIndex, Range.GetInclusiveMaxWidgetIndex().ElementIndex + 1);
 			Data[Range.GetInclusiveMaxWidgetIndex().ArrayIndex].StartIndex = Range.GetInclusiveMaxWidgetIndex().ElementIndex + 1;
 			RemoveDataNodeIfNeeded(Range.GetInclusiveMaxWidgetIndex().ArrayIndex);
@@ -1306,11 +1302,11 @@ void FSlateInvalidationWidgetList::Internal_RemoveRangeFromSameParent(const FInd
 			ElementListType& RemoveElementList = ArrayNode.ElementList;
 			if (bRangeIsInSameElementArray)
 			{
-				SetFakeInvalidatWidgetHandle(Range.GetInclusiveMinWidgetIndex().ArrayIndex, Range.GetInclusiveMinWidgetIndex().ElementIndex, Range.GetInclusiveMaxWidgetIndex().ElementIndex);
+				ResetInvalidationWidgetHandle(Range.GetInclusiveMinWidgetIndex().ArrayIndex, Range.GetInclusiveMinWidgetIndex().ElementIndex, Range.GetInclusiveMaxWidgetIndex().ElementIndex);
 			}
 			else
 			{
-				SetFakeInvalidatWidgetHandle(Range.GetInclusiveMinWidgetIndex().ArrayIndex, Range.GetInclusiveMinWidgetIndex().ElementIndex, RemoveElementList.Num());
+				ResetInvalidationWidgetHandle(Range.GetInclusiveMinWidgetIndex().ArrayIndex, Range.GetInclusiveMinWidgetIndex().ElementIndex, RemoveElementList.Num());
 			}
 
 			const IndexType RemoveArrayAt = Range.GetInclusiveMinWidgetIndex().ElementIndex;
@@ -1326,7 +1322,6 @@ void FSlateInvalidationWidgetList::Internal_RemoveRangeFromSameParent(const FInd
 			check(Range.GetInclusiveMinWidgetIndex().ArrayIndex == Range.GetInclusiveMaxWidgetIndex().ArrayIndex);
 			check(Range.GetInclusiveMinWidgetIndex().ElementIndex == Data[Range.GetInclusiveMinWidgetIndex().ArrayIndex].StartIndex);
 
-			SetFakeInvalidatWidgetHandle(Range.GetInclusiveMinWidgetIndex().ArrayIndex, Range.GetInclusiveMinWidgetIndex().ElementIndex, Range.GetInclusiveMaxWidgetIndex().ElementIndex + 1);
 			ResetInvalidationWidget(Range.GetInclusiveMinWidgetIndex().ArrayIndex, Range.GetInclusiveMinWidgetIndex().ElementIndex, Range.GetInclusiveMaxWidgetIndex().ElementIndex + 1);
 			Data[Range.GetInclusiveMinWidgetIndex().ArrayIndex].StartIndex = Range.GetInclusiveMaxWidgetIndex().ElementIndex + 1;
 			RemoveDataNodeIfNeeded(Range.GetInclusiveMinWidgetIndex().ArrayIndex);
@@ -1854,5 +1849,4 @@ bool FSlateInvalidationWidgetList::VerifyElementIndexList() const
 }
 #endif //UE_SLATE_WITH_INVALIDATIONWIDGETLIST_DEBUGGING
 
-#undef UE_SLATE_WITH_WIDGETLIST_ASSIGNINVALIDPROXYWHENREMOVED
 #undef UE_SLATE_WITH_WIDGETLIST_UPDATEONLYWHATISNEEDED
diff --git a/Engine/Source/Runtime/SlateCore/Private/Rendering/DrawElements.cpp b/Engine/Source/Runtime/SlateCore/Private/Rendering/DrawElements.cpp
--- a/Engine/Source/Runtime/SlateCore/Private/Rendering/DrawElements.cpp
+++ b/Engine/Source/Runtime/SlateCore/Private/Rendering/DrawElements.cpp
@@ -560,13 +560,6 @@ void FSlateCachedElementList::ClearCachedElements()
 	DestroyCachedData();
 
 	CachedRenderingData = new FSlateCachedFastPathRenderingData;
-
-#if 0 // enable this if you want to know why a widget is invalidated after it has been drawn but before it has been batched (probably a child or parent invalidating a relation)
-	if (ensure(!bNewData))
-	{
-		UE_LOG(LogSlate, Log, TEXT("Cleared out data in cached ElementList for Widget: %s before it was batched"), *Widget->GetTag().ToString());
-	}
-#endif
 }
 
 FSlateRenderBatch& FSlateCachedElementList::AddRenderBatch(int32 InLayer, const FShaderParams& InShaderParams, const FSlateShaderResource* InResource, ESlateDrawPrimitive InPrimitiveType, ESlateShader InShaderType, ESlateDrawEffect InDrawEffects, ESlateBatchDrawFlag InDrawFlags, int8 SceneIndex)
@@ -622,7 +615,7 @@ bool FSlateCachedElementList::IsEmpty()
 
 	VisitTupleElements(CheckElementsExist, DrawElements);
 
-	return bElementsFound;
+	return !bElementsFound;
 }
 
 int32 FSlateCachedElementList::NumElements()
