Slate: Add an option to detection NaN for SWidget::RenderTransform and SWidget::RenderTransformPivot.

diff --git a/Engine/Source/Runtime/SlateCore/Private/Widgets/SWidget.cpp b/Engine/Source/Runtime/SlateCore/Private/Widgets/SWidget.cpp
--- a/Engine/Source/Runtime/SlateCore/Private/Widgets/SWidget.cpp
+++ b/Engine/Source/Runtime/SlateCore/Private/Widgets/SWidget.cpp
@@ -45,6 +45,10 @@
 	#define UE_WITH_SLATE_DEBUG_FIND_WIDGET_REFLECTION_METADATA 0
 #endif
 
+#ifndef UE_SLATE_WITH_WIDGET_RENDERING_TRANSFORM_NAN_DIAGNOSTIC
+	#define UE_SLATE_WITH_WIDGET_RENDERING_TRANSFORM_NAN_DIAGNOSTIC (!ENABLE_NAN_DIAGNOSTIC && 0) // NAN diagnostic is already activated for every engine systems
+#endif UE_SLATE_WITH_WIDGET_RENDERING_TRANSFORM_NAN_DIAGNOSTIC
+
 #if UE_WITH_SLATE_DEBUG_FIND_WIDGET_REFLECTION_METADATA
 namespace FindWidgetMetaData
 {
@@ -184,6 +188,27 @@ void SWidget::PrivateRegisterAttributes(FSlateAttributeInitializer& AttributeIni
 	SLATE_ADD_MEMBER_ATTRIBUTE_DEFINITION_WITH_NAME(AttributeInitializer, "Hovered", HoveredAttribute, EInvalidateWidgetReason::None);
 	SLATE_ADD_MEMBER_ATTRIBUTE_DEFINITION_WITH_NAME(AttributeInitializer, "RenderTransform", RenderTransformAttribute, EInvalidateWidgetReason::Layout | EInvalidateWidgetReason::RenderTransform);
 	SLATE_ADD_MEMBER_ATTRIBUTE_DEFINITION_WITH_NAME(AttributeInitializer, "RenderTransformPivot", RenderTransformPivotAttribute, EInvalidateWidgetReason::Layout | EInvalidateWidgetReason::RenderTransform);
+
+#if UE_SLATE_WITH_WIDGET_RENDERING_TRANSFORM_NAN_DIAGNOSTIC
+	AttributeInitializer.OverrideOnValueChanged("RenderTransform", FSlateAttributeDescriptor::ECallbackOverrideType::ReplacePrevious
+		, FSlateAttributeDescriptor::FAttributeValueChangedDelegate::CreateLambda([](SWidget& Widget)
+		{
+				if (Widget.RenderTransformAttribute.Get().IsSet() && Widget.RenderTransformAttribute.Get().GetValue().ContainsNaN())
+				{
+					logOrEnsureNanError(TEXT("RenderTransform contains NaN"));
+					Widget.RenderTransformAttribute.Set(Widget, TOptional<FSlateRenderTransform>());
+				}
+		}));
+	AttributeInitializer.OverrideOnValueChanged("RenderTransformPivot", FSlateAttributeDescriptor::ECallbackOverrideType::ReplacePrevious
+		, FSlateAttributeDescriptor::FAttributeValueChangedDelegate::CreateLambda([](SWidget& Widget)
+			{
+				if (Widget.RenderTransformPivotAttribute.Get().ContainsNaN())
+				{
+					logOrEnsureNanError(TEXT("RenderTransformPivot contains NaN"));
+					Widget.RenderTransformPivotAttribute.Set(Widget, FVector2D::ZeroVector);
+				}
+			}));
+#endif
 }
 
 SWidget::SWidget()
