Fix build issue 466614 introduced with CL 33491426
#rnx

diff --git a/Engine/Source/Runtime/SlateCore/Private/Widgets/SWidget.cpp b/Engine/Source/Runtime/SlateCore/Private/Widgets/SWidget.cpp
--- a/Engine/Source/Runtime/SlateCore/Private/Widgets/SWidget.cpp
+++ b/Engine/Source/Runtime/SlateCore/Private/Widgets/SWidget.cpp
@@ -47,7 +47,7 @@
 
 #ifndef UE_SLATE_WITH_WIDGET_RENDERING_TRANSFORM_NAN_DIAGNOSTIC
 	#define UE_SLATE_WITH_WIDGET_RENDERING_TRANSFORM_NAN_DIAGNOSTIC (!ENABLE_NAN_DIAGNOSTIC && 0) // NAN diagnostic is already activated for every engine systems
-#endif UE_SLATE_WITH_WIDGET_RENDERING_TRANSFORM_NAN_DIAGNOSTIC
+#endif
 
 #if UE_WITH_SLATE_DEBUG_FIND_WIDGET_REFLECTION_METADATA
 namespace FindWidgetMetaData
