Slate: With Invalidation, detect if the previous LayerId is smaller than the current LayerId. That occurs when the GlobalInvalidation is deactivated and the a invalidation panel is used.

diff --git a/Engine/Source/Runtime/SlateCore/Private/FastUpdate/SlateInvalidationRoot.cpp b/Engine/Source/Runtime/SlateCore/Private/FastUpdate/SlateInvalidationRoot.cpp
--- a/Engine/Source/Runtime/SlateCore/Private/FastUpdate/SlateInvalidationRoot.cpp
+++ b/Engine/Source/Runtime/SlateCore/Private/FastUpdate/SlateInvalidationRoot.cpp
@@ -741,10 +741,36 @@ bool FSlateInvalidationRoot::PaintFastPath(const FSlateInvalidationContext& Cont
 	{
 		TGuardValue<bool> OnFastPathGuard(GSlateIsOnFastUpdatePath, true);
 
-
-		// Widgets that needs reindexing while updating the FinalUpdateList.
+		// Widgets that needs re-indexing while updating the FinalUpdateList.
 		UE::Slate::Private::FSlateInvalidationPaintFastPathContext FastPaintContext;
 
+		// Is the new LayerId bigger than the previous we received.
+		{
+			const FSlateInvalidationWidgetIndex FirstWidgetIndex = FastWidgetPathList->FirstIndex();
+			if (FirstWidgetIndex != FSlateInvalidationWidgetIndex::Invalid)
+			{
+				const FSlateInvalidationWidgetList::InvalidationWidgetType& FirstWidgetInvalidationWidget = (*FastWidgetPathList)[FirstWidgetIndex];
+				// Are we going to already process it anyway
+				if (FinalUpdateList.Num() == 0 || FinalUpdateList.Last().GetWidgetIndex() != FirstWidgetIndex)
+				{
+					if (SWidget* FirstWidgetPtr = FirstWidgetInvalidationWidget.GetWidget())
+					{
+						// The new LayerId is bigger than previously. Because of some widget optimization, see PaintFastPath_FixupLayerId, we need to repaint the widgets.
+						if (Context.IncomingLayerId > FirstWidgetPtr->GetPersistentState().LayerId)
+						{
+							// +1 is enough but we add a buffer just in case.
+							const int32 LayerIdBuffer = 100;
+							const_cast<FSlateWidgetPersistentState&>(FirstWidgetPtr->GetPersistentState()).LayerId = Context.IncomingLayerId + LayerIdBuffer;
+							FirstWidgetPtr->UpdateFlags |= EWidgetUpdateFlags::NeedsRepaint;
+
+							const FSlateInvalidationWidgetSortOrder FirstWidgetSortIndex{ *FastWidgetPathList, FirstWidgetIndex };
+							FinalUpdateList.Emplace(FirstWidgetIndex, FirstWidgetSortIndex);
+						}
+					}
+				}
+			}
+		}
+
 		// The update list is put in reverse order by ProcessInvalidation
 		while (FinalUpdateList.Num() > 0 || !FastPaintContext.ReindexUpdateList.IsEmpty())
 		{
