Slate: Move ActiveTimers to a SWidget metadata to save some memory. In current testing projects, ~50 out of 14k widgets uses the ActiveTimers. Most of them are a ListView uses it only to keep the editor from sleeping. Add option to the WidgetList to track the number of ActiveTimers. Notice and fix potential bug when ticking the active timers and the timer add or remove an active timer inside the timer. Fix SetCursor that always tries to remove the cursor metadata on SWidget::Construct.

diff --git a/Engine/Source/Runtime/SlateCore/Private/Debugging/WidgetList.cpp b/Engine/Source/Runtime/SlateCore/Private/Debugging/WidgetList.cpp
--- a/Engine/Source/Runtime/SlateCore/Private/Debugging/WidgetList.cpp
+++ b/Engine/Source/Runtime/SlateCore/Private/Debugging/WidgetList.cpp
@@ -33,15 +33,16 @@ TArray<const SWidget*> FWidgetList::AllWidgets;
 
 struct FLogAllWidgetsDebugInfoFlags
 {
-	bool bDebug;
-	bool bPaint;
-	bool bProxy;
-	bool bChildren;
-	bool bParent;
-	bool bToolTip;
-	bool bCursor;
-	bool bSlateAttribute;
-	bool bMouseEventsHandler;
+	bool bDebug = false;
+	bool bPaint = false;
+	bool bProxy = false;
+	bool bChildren = false;
+	bool bParent = false;
+	bool bToolTip = false;
+	bool bCursor = false;
+	bool bSlateAttribute = false;
+	bool bMouseEventsHandler = false;
+	bool bActiveTimers = false;
 
 	void Parse(const FString& Arg)
 	{
@@ -89,6 +90,11 @@ struct FLogAllWidgetsDebugInfoFlags
 		{
 			return;
 		}
+		
+		if (FParse::Bool(*Arg, TEXT("ActiveTimers="), bActiveTimers))
+		{
+			return;
+		}
 	}
 };
 
@@ -133,6 +139,11 @@ void LogAllWidgetsDebugInfoImpl(FOutputDevice& Ar, const FLogAllWidgetsDebugInfo
 	{
 		MessageBuilder << TEXT(";MouseButtonDown;MouseButtonUp;MouseMove;MouseDblClick;MouseEnter;MouseLeave");
 	}
+	if (DebugInfoFlags.bActiveTimers)
+	{
+		MessageBuilder << TEXT(";HasActiveTimers");
+	}
+
 	Ar.Log(MessageBuilder.ToString());
 
 	const TArray<const SWidget*>& WidgetList = FWidgetList::GetAllWidgets();
@@ -282,6 +293,11 @@ void LogAllWidgetsDebugInfoImpl(FOutputDevice& Ar, const FLogAllWidgetsDebugInfo
 			}
 		}
 
+		if (DebugInfoFlags.bActiveTimers)
+		{
+			MessageBuilder << (Widget->HasActiveTimers() ? TEXT(";true") : TEXT(";false"));
+		}
+
 		Ar.Log(MessageBuilder.ToString());
 	}
 }
@@ -337,6 +353,7 @@ void FWidgetList::ExportToCSV(FStringView OutputFilename)
 	Flags.bCursor = true;
 	Flags.bSlateAttribute = true;
 	Flags.bMouseEventsHandler = true;
+	Flags.bActiveTimers = true;
 
 	FOutputDeviceFile OutputDeviceFile(*WriteToString<256>(OutputFilename), true /*bDisableBackup*/);
 	OutputDeviceFile.SetSuppressEventTag(true);
diff --git a/Engine/Source/Runtime/SlateCore/Private/Types/SlateActiveTimersMetaData.h b/Engine/Source/Runtime/SlateCore/Private/Types/SlateActiveTimersMetaData.h
new file mode 100644
--- /dev/null
+++ b/Engine/Source/Runtime/SlateCore/Private/Types/SlateActiveTimersMetaData.h
@@ -0,0 +1,22 @@
+// Copyright Epic Games, Inc. All Rights Reserved.
+
+#pragma once
+
+#include "CoreMinimal.h"
+#include "Application/ActiveTimerHandle.h"
+#include "Types/ISlateMetaData.h"
+
+namespace UE::Slate
+{
+
+/** Metadata that holds the active timer content for a widget. */
+class FActiveTimersMetaData : public ISlateMetaData
+{
+public:
+	SLATE_METADATA_TYPE(FActiveTimersMetaData, ISlateMetaData)
+
+	/** The list of active timer handles for this widget. */
+	TArray<TSharedRef<FActiveTimerHandle>> ActiveTimers;
+};
+
+}
\ No newline at end of file
diff --git a/Engine/Source/Runtime/SlateCore/Private/Widgets/SWidget.cpp b/Engine/Source/Runtime/SlateCore/Private/Widgets/SWidget.cpp
--- a/Engine/Source/Runtime/SlateCore/Private/Widgets/SWidget.cpp
+++ b/Engine/Source/Runtime/SlateCore/Private/Widgets/SWidget.cpp
@@ -20,6 +20,7 @@
 #include "Widgets/SWindow.h"
 #include "Trace/SlateTrace.h"
 #include "Types/SlateAttributeMetaData.h"
+#include "Types/SlateActiveTimersMetaData.h"
 #include "Types/SlateCursorMetaData.h"
 #include "Types/SlateMouseEventsMetaData.h"
 #include "Types/ReflectionMetadata.h"
@@ -224,6 +225,7 @@ SWidget::SWidget()
 	, bEnabledAttributesUpdate(true)
 	, bIsDeclarativeSyntaxConstructionCompleted(false)
 	, bIsHoveredAttributeSet(false)
+	, bHasActiveTimers(false)
 	, bHasCustomPrepass(false)
 	, bHasRelativeLayoutScale(false)
 	, bVolatilityAlwaysInvalidatesPrepass(false)
@@ -284,31 +286,37 @@ SWidget::~SWidget()
 	// Unregister all ActiveTimers so they aren't left stranded in the Application's list.
 	if (FSlateApplicationBase::IsInitialized())
 	{
-		for (const auto& ActiveTimerHandle : ActiveTimers)
+		FSlateApplicationBase& SlateApplication = FSlateApplicationBase::Get();
+		if (bHasActiveTimers)
 		{
-			FSlateApplicationBase::Get().UnRegisterActiveTimer(ActiveTimerHandle);
-		}
-
-		// Warn the invalidation root
-		if (FSlateInvalidationRoot* InvalidationRoot = FastPathProxyHandle.GetInvalidationRootHandle().GetInvalidationRoot())
-		{
-			InvalidationRoot->OnWidgetDestroyed(this);
+			TSharedPtr<UE::Slate::FActiveTimersMetaData> ActiveTimersMetaData = GetMetaData<UE::Slate::FActiveTimersMetaData>();
+			check(ActiveTimersMetaData);
+			for (const TSharedRef<FActiveTimerHandle>& ActiveTimerHandle : ActiveTimersMetaData->ActiveTimers)
+			{
+				SlateApplication.UnRegisterActiveTimer(ActiveTimerHandle);
+			}
 		}
 
-		// Reset handle
-		FastPathProxyHandle = FWidgetProxyHandle();
-
-		// Note: this would still be valid if a widget was painted and then destroyed in the same frame.  
-		// In that case invalidation hasn't taken place for added widgets so the invalidation panel doesn't know about their cached element data to clean it up
-		PersistentState.CachedElementHandle.RemoveFromCache();
-
 #if WITH_ACCESSIBILITY
-		FSlateApplicationBase::Get().GetAccessibleMessageHandler()->OnWidgetRemoved(this);
+		SlateApplication.GetAccessibleMessageHandler()->OnWidgetRemoved(this);
 #endif
 		// Only clear if initialized because SNullWidget's destructor may be called after annotations are deleted
 		ClearSparseAnnotationsForWidget(this);
 	}
 
+	// Warn the invalidation root
+	if (FSlateInvalidationRoot* InvalidationRoot = FastPathProxyHandle.GetInvalidationRootHandle().GetInvalidationRoot())
+	{
+		InvalidationRoot->OnWidgetDestroyed(this);
+	}
+
+	// Reset handle
+	FastPathProxyHandle = FWidgetProxyHandle();
+
+	// Note: this would still be valid if a widget was painted and then destroyed in the same frame.  
+	// In that case invalidation hasn't taken place for added widgets so the invalidation panel doesn't know about their cached element data to clean it up
+	PersistentState.CachedElementHandle.RemoveFromCache();
+
 #if ENABLE_STATNAMEDEVENTS
 	delete[] StatIDStringStorage;
 	StatIDStringStorage = nullptr;
@@ -345,7 +353,10 @@ void SWidget::SWidgetConstruct(const FSlateBaseNamedArgs& Args)
 		SetToolTipText(Args._ToolTipText);
 	}
 
-	SetCursor(Args._Cursor);
+	if (Args._Cursor.IsSet())
+	{
+		SetCursor(Args._Cursor);
+	}
 
 #if WITH_ACCESSIBILITY
 	// If custom text is provided, force behavior to custom. Otherwise, use the passed-in behavior and set their default text.
@@ -355,7 +366,7 @@ void SWidget::SWidgetConstruct(const FSlateBaseNamedArgs& Args)
 		{
 			SetCanChildrenBeAccessible(AccessibleParams.bCanChildrenBeAccessible);
 			SetAccessibleBehavior(AccessibleParams.AccessibleText.IsSet() ? EAccessibleBehavior::Custom : AccessibleParams.AccessibleBehavior, AccessibleParams.AccessibleText, EAccessibleType::Main);
-		SetAccessibleBehavior(AccessibleParams.AccessibleSummaryText.IsSet() ? EAccessibleBehavior::Custom : AccessibleParams.AccessibleSummaryBehavior, AccessibleParams.AccessibleSummaryText, EAccessibleType::Summary);
+			SetAccessibleBehavior(AccessibleParams.AccessibleSummaryText.IsSet() ? EAccessibleBehavior::Custom : AccessibleParams.AccessibleSummaryBehavior, AccessibleParams.AccessibleSummaryText, EAccessibleType::Summary);
 		};
 		if (Args._AccessibleText.IsSet())
 		{
@@ -1416,7 +1427,6 @@ int32 SWidget::Paint(const FPaintArgs& Args, const FGeometry& AllottedGeometry,
 		UE_TRACE_SCOPED_SLATE_WIDGET_UPDATE(this);
 		if (HasAnyUpdateFlags(EWidgetUpdateFlags::NeedsActiveTimerUpdate))
 		{
-
 			SCOPE_CYCLE_COUNTER(STAT_SlateExecuteActiveTimers);
 			MutableThis->ExecuteActiveTimers(Args.GetCurrentTime(), Args.GetDeltaTime());
 		}
@@ -1797,7 +1807,18 @@ TSharedRef<FActiveTimerHandle> SWidget::RegisterActiveTimer(float TickPeriod, FW
 {
 	TSharedRef<FActiveTimerHandle> ActiveTimerHandle = MakeShared<FActiveTimerHandle>(TickPeriod, TickFunction, FSlateApplicationBase::Get().GetCurrentTime() + TickPeriod);
 	FSlateApplicationBase::Get().RegisterActiveTimer(ActiveTimerHandle);
-	ActiveTimers.Add(ActiveTimerHandle);
+
+	if (bHasActiveTimers)
+	{
+		GetMetaData<UE::Slate::FActiveTimersMetaData>()->ActiveTimers.Add(ActiveTimerHandle);
+	}
+	else
+	{
+		TSharedRef<UE::Slate::FActiveTimersMetaData> NewActiveTimers = MakeShared<UE::Slate::FActiveTimersMetaData>();
+		NewActiveTimers->ActiveTimers.Add(ActiveTimerHandle);
+		AddMetadata(NewActiveTimers);
+		bHasActiveTimers = true;
+	}	
 
 	AddUpdateFlags(EWidgetUpdateFlags::NeedsActiveTimerUpdate);
 
@@ -1806,44 +1827,56 @@ TSharedRef<FActiveTimerHandle> SWidget::RegisterActiveTimer(float TickPeriod, FW
 
 void SWidget::UnRegisterActiveTimer(const TSharedRef<FActiveTimerHandle>& ActiveTimerHandle)
 {
-	if (FSlateApplicationBase::IsInitialized())
+	if (bHasActiveTimers)
 	{
-		FSlateApplicationBase::Get().UnRegisterActiveTimer(ActiveTimerHandle);
-		ActiveTimers.Remove(ActiveTimerHandle);
-
-		if (ActiveTimers.Num() == 0)
+		TSharedRef<UE::Slate::FActiveTimersMetaData> ActiveTimersMetaData = GetMetaData<UE::Slate::FActiveTimersMetaData>().ToSharedRef();
+		ActiveTimersMetaData->ActiveTimers.RemoveSingle(ActiveTimerHandle);
+		if (ActiveTimersMetaData->ActiveTimers.Num() == 0)
 		{
+			RemoveMetaData(ActiveTimersMetaData);
+			bHasActiveTimers = false;
 			RemoveUpdateFlags(EWidgetUpdateFlags::NeedsActiveTimerUpdate);
 		}
 	}
+
+	if (FSlateApplicationBase::IsInitialized())
+	{
+		FSlateApplicationBase::Get().UnRegisterActiveTimer(ActiveTimerHandle);
+	}
 }
 
 void SWidget::ExecuteActiveTimers(double CurrentTime, float DeltaTime)
 {
+	checkf(bHasActiveTimers, TEXT("The flag EWidgetUpdateFlags::NeedsActiveTimerUpdate should match with the bHasActiveTimers flag"));
+
 	// loop over the registered tick handles and execute them, removing them if necessary.
-	for (int32 i = 0; i < ActiveTimers.Num();)
+	TSharedRef<UE::Slate::FActiveTimersMetaData> ActiveTimersMetaData = GetMetaData<UE::Slate::FActiveTimersMetaData>().ToSharedRef();
+	for (int32 Index = 0; Index < ActiveTimersMetaData->ActiveTimers.Num();)
 	{
-		EActiveTimerReturnType Result = ActiveTimers[i]->ExecuteIfPending(CurrentTime, DeltaTime);
+		TWeakPtr<FActiveTimerHandle> WeakActiveTimerHandle = ActiveTimersMetaData->ActiveTimers[Index];
+		EActiveTimerReturnType Result = ActiveTimersMetaData->ActiveTimers[Index]->ExecuteIfPending(CurrentTime, DeltaTime);
 		if (Result == EActiveTimerReturnType::Continue)
 		{
-			++i;
+			++Index;
 		}
 		else
 		{
-			// Possible that execution unregistered the timer 
-			if (ActiveTimers.IsValidIndex(i))
+			// Possible that execution unregistered the timer
+			if (TSharedPtr<FActiveTimerHandle> ActiveTimerHandle = WeakActiveTimerHandle.Pin())
 			{
 				if (FSlateApplicationBase::IsInitialized())
 				{
-					FSlateApplicationBase::Get().UnRegisterActiveTimer(ActiveTimers[i]);
+					FSlateApplicationBase::Get().UnRegisterActiveTimer(ActiveTimerHandle.ToSharedRef());
 				}
-				ActiveTimers.RemoveAt(i);
+				ActiveTimersMetaData->ActiveTimers.RemoveSingle(ActiveTimerHandle.ToSharedRef());
 			}
 		}
 	}
 
-	if (ActiveTimers.Num() == 0)
+	if (ActiveTimersMetaData->ActiveTimers.Num() == 0)
 	{
+		RemoveMetaData(ActiveTimersMetaData);
+		bHasActiveTimers = false;
 		RemoveUpdateFlags(EWidgetUpdateFlags::NeedsActiveTimerUpdate);
 	}
 }
diff --git a/Engine/Source/Runtime/SlateCore/Public/Widgets/SWidget.h b/Engine/Source/Runtime/SlateCore/Public/Widgets/SWidget.h
--- a/Engine/Source/Runtime/SlateCore/Public/Widgets/SWidget.h
+++ b/Engine/Source/Runtime/SlateCore/Public/Widgets/SWidget.h
@@ -1642,7 +1642,10 @@ public:
 	SLATECORE_API void UnRegisterActiveTimer( const TSharedRef<FActiveTimerHandle>& ActiveTimerHandle );
 	
 	/** Does this widget have any active timers? */
-	bool HasActiveTimers() const { return ActiveTimers.Num() > 0; }
+	bool HasActiveTimers() const
+	{
+		return bHasActiveTimers;
+	}
 
 private:
 
@@ -1728,6 +1731,9 @@ private:
 	/** Is the attribute IsHovered is set? */
 	uint8 bIsHoveredAttributeSet : 1;
 
+	/** Is there at least one ActiveTimer currently registered. */
+	uint8 bHasActiveTimers : 1;
+
 protected:
 	uint8 bHasCustomPrepass : 1;
 
@@ -1782,9 +1788,6 @@ private:
 
 	mutable FSlateWidgetPersistentState PersistentState;
 
-	/** The list of active timer handles for this widget. */
-	TArray<TSharedRef<FActiveTimerHandle>> ActiveTimers;
-
 	/** Stores the ideal size this widget wants to be. */
 	TOptional<FVector2f> DesiredSize;
 
