Slate: Reorder SWidget's member to save a bit of memory. Changed DesiredSize and PrepassLayoutScaleMultiplier types. They are TOptional but the extra bool was making them harder to pack in memory.

diff --git a/Engine/Source/Editor/GraphEditor/Private/SGraphPanel.cpp b/Engine/Source/Editor/GraphEditor/Private/SGraphPanel.cpp
--- a/Engine/Source/Editor/GraphEditor/Private/SGraphPanel.cpp
+++ b/Engine/Source/Editor/GraphEditor/Private/SGraphPanel.cpp
@@ -1506,7 +1506,7 @@ void SGraphPanel::MoveNodesToAnchorPinAtGraphPosition(TArrayView<UEdGraphNode* c
 				if (bNodeNeedsPrepass || bNoDesiredSize)
 				{
 					const int32 ChildIndex = Panel->Children.Find(OwningNodeWidget.ToSharedRef());
-					const float SelfLayoutScaleMultiplier = Panel->PrepassLayoutScaleMultiplier.Get(1.f);
+					const float SelfLayoutScaleMultiplier = Panel->GetPrepassLayoutScaleMultiplier();
 					const float ChildLayoutScaleMultiplier = Panel->bHasRelativeLayoutScale
 								? SelfLayoutScaleMultiplier * Panel->GetRelativeLayoutScale(ChildIndex, SelfLayoutScaleMultiplier)
 								: SelfLayoutScaleMultiplier;
@@ -2148,7 +2148,7 @@ void SGraphPanel::AddNode(UEdGraphNode* Node, AddNodeBehavior Behavior)
 	// We also need to take a bit of care to pass through the same layout scale multiplier as Prepass_ChildLoop() would have so that the zoom level
 	// scale is used, otherwise you'd still get a single frame of jitter while the graph is zoomed out.
 	const int32 ChildIndex = Children.Num() - 1;
-	const float SelfLayoutScaleMultiplier = PrepassLayoutScaleMultiplier.Get(1.f);
+	const float SelfLayoutScaleMultiplier = GetPrepassLayoutScaleMultiplier();
 	const float ChildLayoutScaleMultiplier = bHasRelativeLayoutScale
 		? SelfLayoutScaleMultiplier * GetRelativeLayoutScale(ChildIndex, SelfLayoutScaleMultiplier)
 		: SelfLayoutScaleMultiplier;
diff --git a/Engine/Source/Runtime/SlateCore/Private/Debugging/WidgetList.cpp b/Engine/Source/Runtime/SlateCore/Private/Debugging/WidgetList.cpp
--- a/Engine/Source/Runtime/SlateCore/Private/Debugging/WidgetList.cpp
+++ b/Engine/Source/Runtime/SlateCore/Private/Debugging/WidgetList.cpp
@@ -43,6 +43,7 @@ struct FLogAllWidgetsDebugInfoFlags
 	bool bSlateAttribute = false;
 	bool bMouseEventsHandler = false;
 	bool bActiveTimers = false;
+	bool bRenderTransform = false;
 
 	void Parse(const FString& Arg)
 	{
@@ -95,6 +96,11 @@ struct FLogAllWidgetsDebugInfoFlags
 		{
 			return;
 		}
+		
+		if (FParse::Bool(*Arg, TEXT("RenderTransform="), bRenderTransform))
+		{
+			return;
+		}
 	}
 };
 
@@ -143,6 +149,10 @@ void LogAllWidgetsDebugInfoImpl(FOutputDevice& Ar, const FLogAllWidgetsDebugInfo
 	{
 		MessageBuilder << TEXT(";HasActiveTimers");
 	}
+	if (DebugInfoFlags.bRenderTransform)
+	{
+		MessageBuilder << TEXT(";RenderTransformSet;RenderTransformPivot");
+	}
 
 	Ar.Log(MessageBuilder.ToString());
 
@@ -297,6 +307,12 @@ void LogAllWidgetsDebugInfoImpl(FOutputDevice& Ar, const FLogAllWidgetsDebugInfo
 		{
 			MessageBuilder << (Widget->HasActiveTimers() ? TEXT(";true") : TEXT(";false"));
 		}
+		
+		if (DebugInfoFlags.bRenderTransform)
+		{
+			MessageBuilder << (Widget->GetRenderTransform().IsSet() ? TEXT(";true") : TEXT(";false"));
+			MessageBuilder.Appendf(TEXT("%f,%f;"), Widget->GetRenderTransformPivot().X, Widget->GetRenderTransformPivot().Y);
+		}
 
 		Ar.Log(MessageBuilder.ToString());
 	}
diff --git a/Engine/Source/Runtime/SlateCore/Private/FastUpdate/WidgetProxy.cpp b/Engine/Source/Runtime/SlateCore/Private/FastUpdate/WidgetProxy.cpp
--- a/Engine/Source/Runtime/SlateCore/Private/FastUpdate/WidgetProxy.cpp
+++ b/Engine/Source/Runtime/SlateCore/Private/FastUpdate/WidgetProxy.cpp
@@ -141,11 +141,11 @@ void FWidgetProxy::ProcessLayoutInvalidation(FSlateInvalidationWidgetPostHeap& U
 	{
 		if (WidgetPtr->NeedsPrepass())
 		{
-			WidgetPtr->SlatePrepass(WidgetPtr->PrepassLayoutScaleMultiplier.Get(1.0f));
+			WidgetPtr->SlatePrepass(WidgetPtr->GetPrepassLayoutScaleMultiplier());
 		}
 		else
 		{
-			WidgetPtr->CacheDesiredSize(WidgetPtr->PrepassLayoutScaleMultiplier.Get(1.0f));
+			WidgetPtr->CacheDesiredSize(WidgetPtr->GetPrepassLayoutScaleMultiplier());
 		}
 
 		NewDesiredSize = WidgetPtr->GetDesiredSize();
@@ -191,7 +191,7 @@ bool FWidgetProxy::ProcessPostInvalidation(FSlateInvalidationWidgetPostHeap& Upd
 	bool bWidgetNeedsRepaint = false;
 	SWidget* WidgetPtr = GetWidget();
 
-	if (Visibility.IsVisible() && ParentIndex != FSlateInvalidationWidgetIndex::Invalid && !WidgetPtr->PrepassLayoutScaleMultiplier.IsSet())
+	if (Visibility.IsVisible() && ParentIndex != FSlateInvalidationWidgetIndex::Invalid && !WidgetPtr->bPrepassLayoutScaleMultiplierSet)
 	{
 		SCOPE_CYCLE_SWIDGET(WidgetPtr);
 		// If this widget has never been prepassed make sure the parent prepasses it to set the correct multiplier
diff --git a/Engine/Source/Runtime/SlateCore/Private/Widgets/SWidget.cpp b/Engine/Source/Runtime/SlateCore/Private/Widgets/SWidget.cpp
--- a/Engine/Source/Runtime/SlateCore/Private/Widgets/SWidget.cpp
+++ b/Engine/Source/Runtime/SlateCore/Private/Widgets/SWidget.cpp
@@ -226,6 +226,8 @@ SWidget::SWidget()
 	, bIsDeclarativeSyntaxConstructionCompleted(false)
 	, bIsHoveredAttributeSet(false)
 	, bHasActiveTimers(false)
+	, bDesiredSizeSet(false)
+	, bPrepassLayoutScaleMultiplierSet(false)
 	, bHasCustomPrepass(false)
 	, bHasRelativeLayoutScale(false)
 	, bVolatilityAlwaysInvalidatesPrepass(false)
@@ -239,10 +241,11 @@ SWidget::SWidget()
 	, FlowDirectionPreference(EFlowDirectionPreference::Inherit)
 	// Note we are defaulting to tick for backwards compatibility
 	, UpdateFlags(EWidgetUpdateFlags::NeedsTick)
-	, DesiredSize()
 	, VisibilityAttribute(*this, EVisibility::Visible)
 	, EnabledStateAttribute(*this, true)
 	, HoveredAttribute(*this, false)
+	, PrepassLayoutScaleMultiplierValue(1.0f)
+	, DesiredSize(FVector2f(0.0f, 0.0f))
 	, RenderTransformPivotAttribute(*this, FVector2D::ZeroVector)
 	, RenderTransformAttribute(*this)
 	, CullingBoundsExtension()
@@ -267,6 +270,7 @@ SWidget::SWidget()
 	UE_TRACE_SLATE_WIDGET_ADDED(this);
 }
 
+PRAGMA_DISABLE_DEPRECATION_WARNINGS
 SWidget::~SWidget()
 {
 #if WITH_SLATE_DEBUGGING
@@ -327,6 +331,7 @@ SWidget::~SWidget()
 	DEC_DWORD_STAT(STAT_SlateTotalWidgets);
 	DEC_MEMORY_STAT_BY(STAT_SlateSWidgetAllocSize, AllocSize);
 }
+PRAGMA_ENABLE_DEPRECATION_WARNINGS
 
 void SWidget::SWidgetConstruct(const FSlateBaseNamedArgs& Args)
 {
@@ -712,7 +717,7 @@ void SWidget::InvalidateChildRemovedFromTree(SWidget& Child)
 
 UE::Slate::FDeprecateVector2DResult SWidget::GetDesiredSize() const
 {
-	return UE::Slate::FDeprecateVector2DResult(DesiredSize.Get(FVector2f::ZeroVector));
+	return bDesiredSizeSet ? UE::Slate::FDeprecateVector2DResult(DesiredSize) : FVector2f::ZeroVector;
 }
 
 void SWidget::AssignParentWidget(TSharedPtr<SWidget> InParent)
@@ -1268,7 +1273,7 @@ void SWidget::Invalidate(EInvalidateWidgetReason InvalidateReason)
 		InvalidateReason |= EInvalidateWidgetReason::Layout;
 	}
 
-	if (EnumHasAnyFlags(InvalidateReason, EInvalidateWidgetReason::ChildOrder) || !PrepassLayoutScaleMultiplier.IsSet())
+	if (EnumHasAnyFlags(InvalidateReason, EInvalidateWidgetReason::ChildOrder) || !bPrepassLayoutScaleMultiplierSet)
 	{
 		MarkPrepassAsDirty();
 		InvalidateReason |= EInvalidateWidgetReason::Prepass;
@@ -1726,7 +1731,8 @@ void SWidget::ArrangeChildren(const FGeometry& AllottedGeometry, FArrangedChildr
 
 void SWidget::Prepass_Internal(float InLayoutScaleMultiplier)
 {
-	PrepassLayoutScaleMultiplier = InLayoutScaleMultiplier;
+	PrepassLayoutScaleMultiplierValue = InLayoutScaleMultiplier;
+	bPrepassLayoutScaleMultiplierSet = true;
 
 	bool bShouldPrepassChildren = true;
 	if (bHasCustomPrepass)
@@ -1746,7 +1752,7 @@ void SWidget::Prepass_Internal(float InLayoutScaleMultiplier)
 
 	{
 		// Cache this widget's desired size.
-		CacheDesiredSize(PrepassLayoutScaleMultiplier.Get(1.0f));
+		CacheDesiredSize(GetPrepassLayoutScaleMultiplier());
 		bNeedsPrepass = false;
 	}
 }
@@ -1793,9 +1799,10 @@ void SWidget::Prepass_ChildLoop(float InLayoutScaleMultiplier, FChildren* MyChil
 			// it is finally visible and invalidate it's prepass so that it gets that when its visibility
 			// is finally invalidated.
 			Child.MarkPrepassAsDirty();
-			Child.PrepassLayoutScaleMultiplier = Self->bHasRelativeLayoutScale
+			Child.PrepassLayoutScaleMultiplierValue = Self->bHasRelativeLayoutScale
 				? InLayoutScaleMultiplier * Self->GetRelativeLayoutScale(ChildIndex, InLayoutScaleMultiplier)
 				: InLayoutScaleMultiplier;
+			Child.bPrepassLayoutScaleMultiplierSet = true;
 		}
 		++ChildIndex;
 	};
diff --git a/Engine/Source/Runtime/SlateCore/Public/Widgets/SWidget.h b/Engine/Source/Runtime/SlateCore/Public/Widgets/SWidget.h
--- a/Engine/Source/Runtime/SlateCore/Public/Widgets/SWidget.h
+++ b/Engine/Source/Runtime/SlateCore/Public/Widgets/SWidget.h
@@ -738,6 +738,7 @@ private:
 	void SetDesiredSize(const FVector2D& InDesiredSize)
 	{
 		DesiredSize = FVector2f(InDesiredSize);
+		bDesiredSizeSet = true;
 	}
 
 #if STATS || ENABLE_STATNAMEDEVENTS
@@ -1272,8 +1273,29 @@ public:
 	}
 
 	/** Gets the desired flow direction for the layout. */
-	EFlowDirectionPreference GetFlowDirectionPreference() const { return FlowDirectionPreference; }
+	EFlowDirectionPreference GetFlowDirectionPreference() const
+	{
+		return FlowDirectionPreference;
+	}
+
+protected:
+	/** Establishes a new flow direction potentially, if this widget has a particular preference for it and all its children. */
+	EFlowDirection ComputeFlowDirection() const
+	{
+		switch (FlowDirectionPreference)
+		{
+		case EFlowDirectionPreference::Culture:
+			return FLayoutLocalization::GetLocalizedLayoutDirection();
+		case EFlowDirectionPreference::LeftToRight:
+			return EFlowDirection::LeftToRight;
+		case EFlowDirectionPreference::RightToLeft:
+			return EFlowDirection::RightToLeft;
+		}
+
+		return GSlateFlowDirection;
+	}
 
+public:
 	/** Set the tool tip that should appear when this widget is hovered. */
 	SLATECORE_API void SetToolTipText(const TAttribute<FText>& ToolTipText);
 
@@ -1295,7 +1317,6 @@ protected:
 	SLATECORE_API virtual TOptional<EMouseCursor::Type> GetCursor() const;
 
 public:
-
 	/**
 	 * Get the metadata of the type provided.
 	 * @return the first metadata of the type supplied that we encounter
@@ -1605,8 +1626,10 @@ private:
 	SLATECORE_API void Prepass_Internal(float LayoutScaleMultiplier);
 
 protected:
-
-	float GetPrepassLayoutScaleMultiplier() const { return PrepassLayoutScaleMultiplier.Get(1.0f); }
+	float GetPrepassLayoutScaleMultiplier() const
+	{
+		return bPrepassLayoutScaleMultiplierSet ? PrepassLayoutScaleMultiplierValue : 1.0f;
+	}
 	
 	SLATECORE_API void Prepass_ChildLoop(float InLayoutScaleMultiplier, FChildren* MyChildren);
 
@@ -1734,6 +1757,12 @@ private:
 	/** Is there at least one ActiveTimer currently registered. */
 	uint8 bHasActiveTimers : 1;
 
+	/** Is the DesiredSize set. */
+	uint8 bDesiredSizeSet : 1;
+
+	/** Is the PrepassLayoutScaleMultiplier set. */
+	uint8 bPrepassLayoutScaleMultiplierSet : 1;
+
 protected:
 	uint8 bHasCustomPrepass : 1;
 
@@ -1761,36 +1790,13 @@ protected:
 	 */
 	EWidgetPixelSnapping PixelSnappingMethod;
 
-protected:
-	/** Establishes a new flow direction potentially, if this widget has a particular preference for it and all its children. */
-	EFlowDirection ComputeFlowDirection() const
-	{
-		switch (FlowDirectionPreference)
-		{
-		case EFlowDirectionPreference::Culture:
-			return FLayoutLocalization::GetLocalizedLayoutDirection();
-		case EFlowDirectionPreference::LeftToRight:
-			return EFlowDirection::LeftToRight;
-		case EFlowDirectionPreference::RightToLeft:
-			return EFlowDirection::RightToLeft;
-		}
-
-		return GSlateFlowDirection;
-	}
-
 private:
-
 	/** Flow direction preference */
 	EFlowDirectionPreference FlowDirectionPreference;
 
 	/** The different updates this widget needs next frame. */
 	EWidgetUpdateFlags UpdateFlags;
 
-	mutable FSlateWidgetPersistentState PersistentState;
-
-	/** Stores the ideal size this widget wants to be. */
-	TOptional<FVector2f> DesiredSize;
-
 	/** Is this widget visible, hidden or collapsed */
 	TSlateAttribute<EVisibility> VisibilityAttribute;
 
@@ -1800,6 +1806,14 @@ private:
 	/** Whether or not this widget is hovered */
 	TSlateAttribute<bool> HoveredAttribute;
 
+	//~ See PrepassLayoutScaleMultiplier
+	float PrepassLayoutScaleMultiplierValue;
+
+	mutable FSlateWidgetPersistentState PersistentState;
+
+	/** Stores the ideal size this widget wants to be. See bDesiredSizeSet. */
+	FVector2f DesiredSize;
+
 	/** Render transform pivot of this widget (in normalized local space) */
 	TSlateAttribute<FVector2D> RenderTransformPivotAttribute;
 
@@ -1807,8 +1821,6 @@ private:
 	TSlateAttribute<TOptional<FSlateRenderTransform>> RenderTransformAttribute;
 
 protected:
-
-	TOptional<float> PrepassLayoutScaleMultiplier;
 	/**
 	* Can be used to enlarge the culling bounds of this widget (pre-intersection), this can be useful if you've got
 	* children that you know are using rendering transforms to render outside their standard bounds, if that happens
@@ -1833,13 +1845,25 @@ private:
 	/** Debugging information on the type of widget we're creating for the Widget Reflector. */
 	FName TypeOfWidget;
 
-private: 
+protected:
+#if WITH_EDITORONLY_DATA
+	UE_DEPRECATED(5.5, "Direct access to PrepassLayoutScaleMultiplier is deprecated. Use the GetPrepassLayoutScaleMultiplier. The scale multitplier can only be set in the prepass.")
+	TOptional<float> PrepassLayoutScaleMultiplier;
+#endif
 
+private: 
 #if !UE_BUILD_SHIPPING
 	/** Full file path (and line) in which this widget was created */
 	FName CreatedInLocation;
 #endif
 
+#if WITH_SLATE_DEBUGGING
+	/** The last time this widget got painted. */
+	uint32 LastPaintFrame = 0;
+	/** Flag to help detect when we access an invalid Widget. */
+	uint8 Debug_DestroyedTag = 0xDC;
+#endif // WITH_SLATE_DEBUGGING
+
 #if UE_SLATE_TRACE_ENABLED
 	/**
 	 * If the widget info is sent when a trace is not active, it will be ignored.
@@ -1849,13 +1873,6 @@ private:
 	mutable uint8 Debug_LastTraceInfoSent = 0;
 #endif // UE_SLATE_TRACE_ENABLED
 
-#if WITH_SLATE_DEBUGGING
-	/** The last time this widget got painted. */
-	uint32 LastPaintFrame = 0;
-	/** Flag to help detect when we access an invalid Widget. */
-	uint8 Debug_DestroyedTag = 0xDC;
-#endif // WITH_SLATE_DEBUGGING
-
 #if UE_SLATE_WITH_WIDGET_UNIQUE_IDENTIFIER
 	/** The widget's id */
 	uint64 UniqueIdentifier;
