Added single channel and approximate SDF font rendering modes, two-pass SDF generator

#tests Tested in editor and PIE
#rnx <Release Notes exclude>

diff --git a/Engine/Source/Runtime/Engine/Private/Font.cpp b/Engine/Source/Runtime/Engine/Private/Font.cpp
--- a/Engine/Source/Runtime/Engine/Private/Font.cpp
+++ b/Engine/Source/Runtime/Engine/Private/Font.cpp
@@ -391,22 +391,14 @@ void UFont::GetResourceSizeEx(FResourceSizeEx& CumulativeResourceSize)
 
 bool UFont::IsSdfFont() const
 {
-	switch (FontRasterizationMode)
-	{
-		case EFontRasterizationMode::Msdf:
-			return IsSlateSdfTextFeatureEnabled();
-		default:
-			return false;
-	}
+	return IsSdfFontRasterizationMode(FontRasterizationMode) && IsSlateSdfTextFeatureEnabled();
 }
 
 EFontRasterizationMode UFont::GetFontRasterizationMode() const
 {
-	switch (FontRasterizationMode)
+	if (IsSdfFontRasterizationMode(FontRasterizationMode) && !IsSlateSdfTextFeatureEnabled())
 	{
-		case EFontRasterizationMode::Msdf:
-			return IsSlateSdfTextFeatureEnabled() ? EFontRasterizationMode::Msdf : EFontRasterizationMode::Bitmap;
-		default:
-			return FontRasterizationMode;
+		return EFontRasterizationMode::Bitmap;
 	}
+	return FontRasterizationMode;
 }
diff --git a/Engine/Source/Runtime/SlateCore/Private/Fonts/FontCache.cpp b/Engine/Source/Runtime/SlateCore/Private/Fonts/FontCache.cpp
--- a/Engine/Source/Runtime/SlateCore/Private/Fonts/FontCache.cpp
+++ b/Engine/Source/Runtime/SlateCore/Private/Fonts/FontCache.cpp
@@ -137,15 +137,17 @@ FShapedGlyphEntryKey::FShapedGlyphEntryKey(const FShapedGlyphFaceData& InFontFac
 *
 */
 
-FSdfGlyphEntryKey::FSdfGlyphEntryKey(const TWeakPtr<FFreeTypeFace> InFontFace, uint32 InGlyphIndex, int32 InPpem, float InEmOuterSpread, float InEmInnerSpread)
+FSdfGlyphEntryKey::FSdfGlyphEntryKey(const TWeakPtr<FFreeTypeFace> InFontFace, uint32 InGlyphIndex, ESlateFontAtlasContentType InAtlasContentType, int32 InPpem, float InEmOuterSpread, float InEmInnerSpread)
 	: FontFace(InFontFace)
 	, GlyphIndex(InGlyphIndex)
+	, AtlasContentType(InAtlasContentType)
 	, Ppem(InPpem)
 	, SpreadCategory(GetSpreadCategory(InEmOuterSpread, InEmInnerSpread))
 	, KeyHash(0)
 {
 	KeyHash = HashCombine(KeyHash, GetTypeHash(FontFace));
 	KeyHash = HashCombine(KeyHash, GetTypeHash(GlyphIndex));
+	KeyHash = HashCombine(KeyHash, GetTypeHash(AtlasContentType));
 	KeyHash = HashCombine(KeyHash, GetTypeHash(Ppem));
 	KeyHash = HashCombine(KeyHash, GetTypeHash(SpreadCategory));
 }
@@ -276,26 +278,16 @@ SIZE_T FShapedGlyphSequence::GetAllocatedSize() const
 
 bool FShapedGlyphSequence::IsSdfFont() const
 {
-	switch (RasterizationMode)
-	{
-		case EFontRasterizationMode::Msdf:
-			return IsSlateSdfTextFeatureEnabled();
-		default:
-			return false;
-	}
+	return IsSdfFontRasterizationMode(RasterizationMode) && IsSlateSdfTextFeatureEnabled();
 }
 
 EFontRasterizationMode FShapedGlyphSequence::GetRasterizationMode() const
 {
-	switch (RasterizationMode)
+	if (IsSdfFontRasterizationMode(RasterizationMode) && !IsSlateSdfTextFeatureEnabled())
 	{
-		case EFontRasterizationMode::Msdf:
-			return IsSlateSdfTextFeatureEnabled() ?
-				EFontRasterizationMode::Msdf :
-				EFontRasterizationMode::Bitmap;
-		default:
-			return RasterizationMode;
+		return EFontRasterizationMode::Bitmap;
 	}
+	return RasterizationMode;
 }
 
 bool FShapedGlyphSequence::IsDirty() const
@@ -1201,7 +1193,7 @@ static float CeilSdfHalfSpread(float InEmSpread, float Ppem)
 	return PixelSpread/Ppem;
 }
 
-FSdfGlyphFontAtlasData FSlateFontCache::GetSdfGlyphFontAtlasData(const FShapedGlyphEntry& InShapedGlyph, const FFontOutlineSettings& InOutlineSettings, const FFontSdfSettings& InSdfSettings)
+FSdfGlyphFontAtlasData FSlateFontCache::GetSdfGlyphFontAtlasData(const FShapedGlyphEntry& InShapedGlyph, const FFontOutlineSettings& InOutlineSettings, EFontRasterizationMode InRasterizationMode, const FFontSdfSettings& InSdfSettings)
 {
 	if (!IsSlateSdfTextFeatureEnabled()
 		|| !InShapedGlyph.FontFaceData->bSupportsSdf
@@ -1212,6 +1204,7 @@ FSdfGlyphFontAtlasData FSlateFontCache::GetSdfGlyphFontAtlasData(const FShapedGl
 	const int32 CachedAtlasDataThreadIndex = static_cast<int32>(OwningThread);
 	check(CachedAtlasDataThreadIndex >= 0 && CachedAtlasDataThreadIndex < 2);
 
+	const bool bPlaceholderOnly = InRasterizationMode == EFontRasterizationMode::SdfApproximation;
 #if WITH_FREETYPE
 	const float TargetPpem = static_cast<float>(FreeTypeUtils::ComputeFontPixelSize(InShapedGlyph.FontFaceData->FontSize, InShapedGlyph.FontFaceData->FontScale));
 #else
@@ -1224,11 +1217,34 @@ FSdfGlyphFontAtlasData FSlateFontCache::GetSdfGlyphFontAtlasData(const FShapedGl
 	float EmInnerSpread = FMath::Clamp(2.f/TargetPpem, 0.05f, 8.f);
 	float EmOuterSpread = FMath::Min(EmInnerSpread + EmOutlineSize, 8.f);
 
+	FSlateSdfGenerator::ESdfType SdfType = FSlateSdfGenerator::ESdfType::MultichannelAndSimple;
+	ESlateFontAtlasContentType AtlasType = ESlateFontAtlasContentType::Msdf;
+	TArray<uint8>* FontAtlasIndicesPtr = &MsdfFontAtlasIndices;
+	switch (InRasterizationMode)
+	{
+		case EFontRasterizationMode::Bitmap:
+		default:
+			checkNoEntry();
+			break;
+		case EFontRasterizationMode::Msdf:
+			SdfType = FSlateSdfGenerator::ESdfType::MultichannelAndSimple;
+			AtlasType = ESlateFontAtlasContentType::Msdf;
+			FontAtlasIndicesPtr = &MsdfFontAtlasIndices;
+			break;
+		case EFontRasterizationMode::Sdf:
+		case EFontRasterizationMode::SdfApproximation:
+			SdfType = FSlateSdfGenerator::ESdfType::Simple;
+			AtlasType = ESlateFontAtlasContentType::Alpha;
+			FontAtlasIndicesPtr = &GrayscaleFontAtlasIndices;
+			break;
+	}
+
 	auto TryRespawn = [&](const FSdfGlyphFontAtlasData& GlyphAtlasData) -> bool {
-		const FSdfGlyphEntryKey GlyphKey(InShapedGlyph.FontFaceData->FontFace, InShapedGlyph.GlyphIndex, InSdfSettings.GetClampedPpem(), GlyphAtlasData.EmOuterSpread, GlyphAtlasData.EmInnerSpread);
+		const FSdfGlyphEntryKey GlyphKey(InShapedGlyph.FontFaceData->FontFace, InShapedGlyph.GlyphIndex, AtlasType, InSdfSettings.GetClampedPpem(), GlyphAtlasData.EmOuterSpread, GlyphAtlasData.EmInnerSpread);
 		FSlateSdfGenerator::FRequestDescriptor SdfRequestDescriptor = { 
 			GlyphKey.FontFace,
 			GlyphKey.GlyphIndex,
+			SdfType,
 			GlyphAtlasData.EmOuterSpread,
 			GlyphAtlasData.EmInnerSpread,
 			GlyphKey.Ppem
@@ -1292,7 +1308,7 @@ FSdfGlyphFontAtlasData FSlateFontCache::GetSdfGlyphFontAtlasData(const FShapedGl
 	};
 	// Not cached on the glyph, so create a key for to look up this glyph, as it may
 	// have already been cached by another shaped text sequence
-	const FSdfGlyphEntryKey GlyphKey(InShapedGlyph.FontFaceData->FontFace, InShapedGlyph.GlyphIndex, InSdfSettings.GetClampedPpem(), EmOuterSpread, EmInnerSpread);
+	const FSdfGlyphEntryKey GlyphKey(InShapedGlyph.FontFaceData->FontFace, InShapedGlyph.GlyphIndex, AtlasType, InSdfSettings.GetClampedPpem(), EmOuterSpread, EmInnerSpread);
 
 	{
 		// Has the atlas data already been cached by another shaped text sequence?
@@ -1311,13 +1327,16 @@ FSdfGlyphFontAtlasData FSlateFontCache::GetSdfGlyphFontAtlasData(const FShapedGl
 	FSlateSdfGenerator::FRequestDescriptor SdfRequestDescriptor = { 
 		GlyphKey.FontFace,
 		GlyphKey.GlyphIndex,
+		SdfType,
 		EmOuterSpread,
 		EmInnerSpread,
 		GlyphKey.Ppem
 	};
 	FSlateSdfGenerator::FRequestOutputInfo SdfRequestOutputInfo = {};
 	FCharacterRenderData PlaceholderRenderData;
-	const FSlateSdfGenerator::ERequestResponse SpawnStatus = SdfGenerator->SpawnWithPlaceholder(SdfRequestDescriptor, SdfRequestOutputInfo, PlaceholderRenderData.RawPixels);
+	const FSlateSdfGenerator::ERequestResponse SpawnStatus = bPlaceholderOnly ?
+		SdfGenerator->MakePlaceholder(SdfRequestDescriptor, SdfRequestOutputInfo, PlaceholderRenderData.RawPixels) :
+		SdfGenerator->SpawnWithPlaceholder(SdfRequestDescriptor, SdfRequestOutputInfo, PlaceholderRenderData.RawPixels);
 	TSharedRef<FSdfGlyphFontAtlasData> NewAtlasData = MakeShareable(new FSdfGlyphFontAtlasData());
 
 	switch (SpawnStatus)
@@ -1333,7 +1352,7 @@ FSdfGlyphFontAtlasData FSlateFontCache::GetSdfGlyphFontAtlasData(const FShapedGl
 			// This can happen if font is still loading
 			return { };
 		case FSlateSdfGenerator::ERequestResponse::PLACEHOLDER_ONLY:
-			NewAtlasData->bPendingRespawn = true;
+			NewAtlasData->bPendingRespawn = !bPlaceholderOnly;
 			break;
 		case FSlateSdfGenerator::ERequestResponse::BAD_REQUEST:
 		default:
@@ -1361,9 +1380,13 @@ FSdfGlyphFontAtlasData FSlateFontCache::GetSdfGlyphFontAtlasData(const FShapedGl
 		PlaceholderRenderData.SizeY = SdfRequestOutputInfo.ImageHeight;
 		PlaceholderRenderData.HorizontalOffset = SdfRequestOutputInfo.BearingX;
 		PlaceholderRenderData.VerticalOffset = SdfRequestOutputInfo.BearingY;
-		PlaceholderRenderData.ContentType = ESlateFontAtlasContentType::Msdf;
+		PlaceholderRenderData.ContentType = AtlasType;
 		PlaceholderRenderData.bSupportsOutline = true;
 	}
+	else if (bPlaceholderOnly)
+	{
+		return { };
+	}
 
 	auto TryAddToAtlas = [&](FSlateFontAtlas& FontAtlas, uint8 FontAtlasIndex) -> bool
 	{
@@ -1379,7 +1402,7 @@ FSdfGlyphFontAtlasData FSlateFontCache::GetSdfGlyphFontAtlasData(const FShapedGl
 				return true;
 			}
 		}
-		else
+		else if (!bPlaceholderOnly)
 		{
 			// Deferred add character begins by reserving the character box/slot in the texture
 			FDeferredCharacterRenderData RenderData;
@@ -1397,7 +1420,8 @@ FSdfGlyphFontAtlasData FSlateFontCache::GetSdfGlyphFontAtlasData(const FShapedGl
 		return false;
 	};
 
-	for (const uint8 FontAtlasIndex : MsdfFontAtlasIndices)
+	check(FontAtlasIndicesPtr);
+	for (const uint8 FontAtlasIndex : *FontAtlasIndicesPtr)
 	{
 		FSlateFontAtlas& FontAtlas = static_cast<FSlateFontAtlas&>(AllFontTextures[FontAtlasIndex].Get());
 		if (TryAddToAtlas(FontAtlas, FontAtlasIndex))
@@ -1414,9 +1438,9 @@ FSdfGlyphFontAtlasData FSlateFontCache::GetSdfGlyphFontAtlasData(const FShapedGl
 		return { /*AtlasHasTooManyTexturesTryAgainLater*/ };
 	}
 
-	TSharedRef<FSlateFontAtlas> FontAtlas = FontAtlasFactory->CreateFontAtlas(ESlateFontAtlasContentType::Msdf);
+	TSharedRef<FSlateFontAtlas> FontAtlas = FontAtlasFactory->CreateFontAtlas(AtlasType);
 	const uint8 TextureIndex = (uint8)AllFontTextures.Add(FontAtlas);
-	MsdfFontAtlasIndices.Add(TextureIndex);
+	FontAtlasIndicesPtr->Add(TextureIndex);
 
 	INC_DWORD_STAT_BY(STAT_SlateNumFontAtlases, 1);
 
@@ -1717,8 +1741,22 @@ void FSlateFontCache::UpdateCache()
 		SdfGenerator->Update(
 			[this](const FSlateSdfGenerator::FRequestDescriptor InRequestDescriptor, TArray<uint8> InRawPixels)
 			{
-				const FSdfGlyphEntryKey GlyphKey(InRequestDescriptor.FontFace, 
-												 InRequestDescriptor.GlyphIndex, 
+				ESlateFontAtlasContentType AtlasType = ESlateFontAtlasContentType::Msdf;
+				switch (InRequestDescriptor.SdfType)
+				{
+					case FSlateSdfGenerator::ESdfType::Simple:
+					case FSlateSdfGenerator::ESdfType::Pseudodistance:
+						AtlasType = ESlateFontAtlasContentType::Alpha;
+						break;
+					case FSlateSdfGenerator::ESdfType::MultichannelAndSimple:
+						AtlasType = ESlateFontAtlasContentType::Msdf;
+						break;
+					default:
+						checkNoEntry();
+				}
+				const FSdfGlyphEntryKey GlyphKey(InRequestDescriptor.FontFace,
+												 InRequestDescriptor.GlyphIndex,
+												 AtlasType,
 												 InRequestDescriptor.Ppem,
 												 InRequestDescriptor.EmOuterSpread,
 												 InRequestDescriptor.EmInnerSpread);
@@ -1729,7 +1767,7 @@ void FSlateFontCache::UpdateCache()
 					if (FoundAtlasData && (*FoundAtlasData)->bSupportsSdf)
 					{
 						check(AllFontTextures.IsValidIndex((*FoundAtlasData)->TextureIndex) &&
-							  MsdfFontAtlasIndices.Contains((*FoundAtlasData)->TextureIndex));
+							  (InRequestDescriptor.SdfType == FSlateSdfGenerator::ESdfType::MultichannelAndSimple ? MsdfFontAtlasIndices : GrayscaleFontAtlasIndices).Contains((*FoundAtlasData)->TextureIndex));
 						FSlateFontAtlas& FontAtlas = static_cast<FSlateFontAtlas&>(AllFontTextures[(*FoundAtlasData)->TextureIndex].Get());
 						FDeferredCharacterRenderData RenderData;
 						RenderData.StartU = (*FoundAtlasData)->StartU;
diff --git a/Engine/Source/Runtime/SlateCore/Private/Fonts/SlateSdfGenerator.cpp b/Engine/Source/Runtime/SlateCore/Private/Fonts/SlateSdfGenerator.cpp
--- a/Engine/Source/Runtime/SlateCore/Private/Fonts/SlateSdfGenerator.cpp
+++ b/Engine/Source/Runtime/SlateCore/Private/Fonts/SlateSdfGenerator.cpp
@@ -336,6 +336,21 @@ public:
 
 };
 
+static int32 GetNumSdfChannels(FSlateSdfGenerator::ESdfType InSdfType)
+{
+	switch (InSdfType)
+	{
+		case FSlateSdfGenerator::ESdfType::Simple:
+		case FSlateSdfGenerator::ESdfType::Pseudodistance:
+			return 1;
+		case FSlateSdfGenerator::ESdfType::MultichannelAndSimple:
+			return 4;
+		default:
+			checkNoEntry();
+	}
+	return 0;
+}
+
 FSlateSdfGenerator::ERequestResponse FSdfGeneratorTask::Prepare(const FSlateSdfGenerator::FRequestDescriptor& InDescriptor, FSlateSdfGenerator::FRequestOutputInfo& OutOutputInfo, bool bCsvTrace)
 {
 	#if WITH_FREETYPE
@@ -440,35 +455,66 @@ void FSdfGeneratorTask::DoOutlineDecomposition()
 	const bool OverlappedContourSupport = true;
 	const int32 TargetWidth = GlyphSdfMapping.GetSdfWidth();
 	const int32 TargetHeight = GlyphSdfMapping.GetSdfHeight();
+	const int32 TargetChannels = GetNumSdfChannels(Descriptor.SdfType);
 	TArray<float> FloatPixels;
-	FloatPixels.SetNumUninitialized(4 * TargetWidth * TargetHeight);
-	OutputPixels.SetNumUninitialized(4 * TargetWidth * TargetHeight);
+	FloatPixels.SetNumUninitialized(TargetChannels * TargetWidth * TargetHeight);
+	OutputPixels.SetNumUninitialized(TargetChannels * TargetWidth * TargetHeight);
 
-	TArray<uint8> ECBuffer;
-	ECBuffer.SetNumUninitialized (TargetWidth * TargetHeight);
-
-	msdfgen::edgeColoringInkTrap(MsdfgenShape, SDF_CORNER_ANGLE_THRESHOLD);
-
-	msdfgen::BitmapRef<float, 4> FloatBitmap(FloatPixels.GetData(), TargetWidth, TargetHeight);
-	msdfgen::generateMTSDF(
-		FloatBitmap,
-		MsdfgenShape,
-		GlyphSdfMapping.GetMsdfgenProjection(),
-		1.0,
-		msdfgen::MSDFGeneratorConfig(
-			OverlappedContourSupport,
-			msdfgen::ErrorCorrectionConfig(
-				msdfgen::ErrorCorrectionConfig::EDGE_PRIORITY,
-				msdfgen::ErrorCorrectionConfig::CHECK_DISTANCE_AT_EDGE,
-				msdfgen::ErrorCorrectionConfig::defaultMinDeviationRatio,
-				msdfgen::ErrorCorrectionConfig::defaultMinImproveRatio,
-				ECBuffer.GetData()
-			)
-		)
-	);
+	switch (Descriptor.SdfType)
+	{
+		case FSlateSdfGenerator::ESdfType::Simple:
+		{
+			msdfgen::BitmapRef<float, 1> FloatBitmap(FloatPixels.GetData(), TargetWidth, TargetHeight);
+			msdfgen::generateSDF(
+				FloatBitmap,
+				MsdfgenShape,
+				GlyphSdfMapping.GetMsdfgenProjection(),
+				1.0,
+				msdfgen::GeneratorConfig(OverlappedContourSupport)
+			);
+			break;
+		}
+		case FSlateSdfGenerator::ESdfType::Pseudodistance:
+		{
+			msdfgen::BitmapRef<float, 1> FloatBitmap(FloatPixels.GetData(), TargetWidth, TargetHeight);
+			msdfgen::generatePseudoSDF(
+				FloatBitmap,
+				MsdfgenShape,
+				GlyphSdfMapping.GetMsdfgenProjection(),
+				1.0,
+				msdfgen::GeneratorConfig(OverlappedContourSupport)
+			);
+			break;
+		}
+		case FSlateSdfGenerator::ESdfType::MultichannelAndSimple:
+		{
+			msdfgen::edgeColoringInkTrap(MsdfgenShape, SDF_CORNER_ANGLE_THRESHOLD);
+
+			msdfgen::BitmapRef<float, 4> FloatBitmap(FloatPixels.GetData(), TargetWidth, TargetHeight);
+			msdfgen::generateMTSDF(
+				FloatBitmap,
+				MsdfgenShape,
+				GlyphSdfMapping.GetMsdfgenProjection(),
+				1.0,
+				msdfgen::MSDFGeneratorConfig(
+					OverlappedContourSupport,
+					msdfgen::ErrorCorrectionConfig(
+						msdfgen::ErrorCorrectionConfig::EDGE_PRIORITY,
+						msdfgen::ErrorCorrectionConfig::CHECK_DISTANCE_AT_EDGE,
+						msdfgen::ErrorCorrectionConfig::defaultMinDeviationRatio,
+						msdfgen::ErrorCorrectionConfig::defaultMinImproveRatio,
+						OutputPixels.GetData() // Temporarily repurpose output buffer as error correction buffer
+					)
+				)
+			);
+			break;
+		}
+		default:
+			checkNoEntry();
+	}
 
 	const float* Src = FloatPixels.GetData();
-	for (uint8* Dst = OutputPixels.GetData(), * End = Dst+4*TargetWidth*TargetHeight; Dst < End; ++Dst, ++Src)
+	for (uint8* Dst = OutputPixels.GetData(), * End = Dst+TargetChannels*TargetWidth*TargetHeight; Dst < End; ++Dst, ++Src)
 	{
 		*Dst = GlyphSdfMapping.EncodeDistance(*Src);
 	}
@@ -476,10 +522,11 @@ void FSdfGeneratorTask::DoOutlineDecomposition()
 
 void FSdfGeneratorTask::MakePlaceholder(TArray<uint8>& OutRawPixels) const
 {
+	const int32 TargetChannels = GetNumSdfChannels(Descriptor.SdfType);
 	const int32 TargetWidth = GlyphSdfMapping.GetSdfWidth();
 	const int32 TargetHeight = GlyphSdfMapping.GetSdfHeight();
 	const int32 TargetArea = TargetWidth * TargetHeight;
-	const int32 TotalSubpixels = 4 * TargetArea;
+	const int32 TotalSubpixels = TargetChannels * TargetArea;
 	TArray<float> FloatPixels;
 	FloatPixels.SetNumUninitialized(TargetArea);
 	OutRawPixels.SetNumUninitialized(TotalSubpixels);
@@ -491,10 +538,30 @@ void FSdfGeneratorTask::MakePlaceholder(TArray<uint8>& OutRawPixels) const
 		GlyphSdfMapping.GetMsdfgenOuterRange(),
 		GlyphSdfMapping.GetMsdfgenInnerRange()
 	);
+
 	const float* Src = FloatPixels.GetData();
-	for (uint8* Dst = OutRawPixels.GetData(), * End = Dst + TotalSubpixels; Dst < End; Dst += 4, ++Src)
+	switch (TargetChannels)
 	{
-		Dst[3] = Dst[2] = Dst[1] = Dst[0] = msdfgen::pixelFloatToByte(*Src);
+		case 1:
+			for (uint8* Dst = OutRawPixels.GetData(), * End = Dst + TotalSubpixels; Dst < End; ++Dst, ++Src)
+			{
+				*Dst = msdfgen::pixelFloatToByte(*Src);
+			}
+			break;
+		case 3:
+			for (uint8* Dst = OutRawPixels.GetData(), * End = Dst + TotalSubpixels; Dst < End; Dst += 3, ++Src)
+			{
+				Dst[2] = Dst[1] = Dst[0] = msdfgen::pixelFloatToByte(*Src);
+			}
+			break;
+		case 4:
+			for (uint8* Dst = OutRawPixels.GetData(), * End = Dst + TotalSubpixels; Dst < End; Dst += 4, ++Src)
+			{
+				Dst[3] = Dst[2] = Dst[1] = Dst[0] = msdfgen::pixelFloatToByte(*Src);
+			}
+			break;
+		default:
+			checkNoEntry();
 	}
 }
 
@@ -513,6 +580,7 @@ public:
 	virtual ERequestResponse Spawn(const FRequestDescriptor& InRequest, FRequestOutputInfo& OutCharInfo) override;
 	virtual ERequestResponse SpawnWithPlaceholder(const FRequestDescriptor& InRequest, FRequestOutputInfo& OutCharInfo, TArray<uint8>& OutRawPixels) override;
 	virtual ERequestResponse Respawn(const FRequestDescriptor& InRequest, const FRequestOutputInfo& InCharInfo) override;
+	virtual ERequestResponse MakePlaceholder(const FRequestDescriptor& InRequest, FRequestOutputInfo& OutCharInfo, TArray<uint8>& OutRawPixels) override;
 	virtual void Update(const FForEachRequestDoneCallback& InEnumerator) override;
 	virtual void Flush() override;
 private:
@@ -599,11 +667,9 @@ FSlateSdfGenerator::ERequestResponse FSlateSdfGeneratorImpl::SpawnWithPlaceholde
 {
 	if (FreeTasks.IsEmpty())
 	{
-		SdfUtils::FSdfGeneratorTask PlaceholderTask;
-		const ERequestResponse Result = PlaceholderTask.Prepare(InRequest, OutCharInfo, false);
+		const ERequestResponse Result = MakePlaceholder(InRequest, OutCharInfo, OutRawPixels);
 		if (Result == ERequestResponse::SUCCESS)
 		{
-			PlaceholderTask.MakePlaceholder(OutRawPixels);
 			return ERequestResponse::PLACEHOLDER_ONLY;
 		}
 		return Result;
@@ -656,6 +722,17 @@ FSlateSdfGenerator::ERequestResponse FSlateSdfGeneratorImpl::Respawn(const FRequ
 	return Result;
 }
 
+FSlateSdfGenerator::ERequestResponse FSlateSdfGeneratorImpl::MakePlaceholder(const FRequestDescriptor& InRequest, FRequestOutputInfo& OutCharInfo, TArray<uint8>& OutRawPixels)
+{
+	SdfUtils::FSdfGeneratorTask PlaceholderTask;
+	const ERequestResponse Result = PlaceholderTask.Prepare(InRequest, OutCharInfo, false);
+	if (Result == ERequestResponse::SUCCESS)
+	{
+		PlaceholderTask.MakePlaceholder(OutRawPixels);
+	}
+	return Result;
+}
+
 void FSlateSdfGeneratorImpl::Update(const FForEachRequestDoneCallback& InEnumerator)
 {
 	for (auto It = StartedTasks.CreateIterator(); It; ++It)
diff --git a/Engine/Source/Runtime/SlateCore/Private/Fonts/SlateSdfGenerator.h b/Engine/Source/Runtime/SlateCore/Private/Fonts/SlateSdfGenerator.h
--- a/Engine/Source/Runtime/SlateCore/Private/Fonts/SlateSdfGenerator.h
+++ b/Engine/Source/Runtime/SlateCore/Private/Fonts/SlateSdfGenerator.h
@@ -11,6 +11,7 @@ public:
 	/** A unique pointer to a FSlateSdfGenerator object. */
 	using Ptr = TUniquePtr<FSlateSdfGenerator>;
 
+	/** Generator's response to caller's request. */
 	enum class ERequestResponse
 	{
 		/** Task spawned successfully (and placeholder generated if requested) */
@@ -25,6 +26,17 @@ public:
 		BAD_REQUEST
 	};
 
+	/** Type of requested signed distance field. */
+	enum class ESdfType
+	{
+		/** Simple single-channel signed distance field based on true Euclidean distance (1 channel/pixel) */
+		Simple,
+		/** Single-channel signed distance field based on the "pseudodistance" metric, which is always orthogonal to an edge, and produces mitered rather than rounded outlines (1 channel/pixel) */
+		Pseudodistance,
+		/** Multi-channel signed distance field with simple (true) signed distance field in additional channel (4 channels/pixel, also known as MTSDF) */
+		MultichannelAndSimple
+	};
+
 	/** Glyph metrics made available immediately after spawning a new task. */
 	struct FRequestOutputInfo
 	{
@@ -45,6 +57,8 @@ public:
 		TWeakPtr<class FFreeTypeFace> FontFace;
 		/** Numeric index of the requested glyph. */
 		uint32 GlyphIndex;
+		/** Type of requested distance field. */
+		ESdfType SdfType;
 		/** Outer portion of the width of representable distances in the output distance field expressed in em. */
 		float EmOuterSpread;
 		/** Inner portion of the width of representable distances in the output distance field expressed in em. */
@@ -65,6 +79,8 @@ public:
 	virtual ERequestResponse SpawnWithPlaceholder(const FRequestDescriptor& InRequest, FRequestOutputInfo& OutCharInfo, TArray<uint8>& OutRawPixels) = 0;
 	/** Attempts to start generating again if previous attempt failed but produced a placeholder, whose FRequestOutputInfo must match. */
 	virtual ERequestResponse Respawn(const FRequestDescriptor& InRequest, const FRequestOutputInfo& InCharInfo) = 0;
+	/** Only generates a placeholder-quality signed distance field immediately into OutRawPixels without spawning. */
+	virtual ERequestResponse MakePlaceholder(const FRequestDescriptor& InRequest, FRequestOutputInfo& OutCharInfo, TArray<uint8>& OutRawPixels) = 0;
 	/** Checks for finished tasks and processes each of them by calling InEnumerator. */
 	virtual void Update(const FForEachRequestDoneCallback& InEnumerator) = 0;
 	/** Flushes all started tasks. */
diff --git a/Engine/Source/Runtime/SlateCore/Private/Rendering/ElementBatcher.cpp b/Engine/Source/Runtime/SlateCore/Private/Rendering/ElementBatcher.cpp
--- a/Engine/Source/Runtime/SlateCore/Private/Rendering/ElementBatcher.cpp
+++ b/Engine/Source/Runtime/SlateCore/Private/Rendering/ElementBatcher.cpp
@@ -1391,7 +1391,9 @@ void FSlateElementBatcher::AddTextElement(const FSlateTextElement& DrawElement)
 							ShaderType = ESlateShader::ColorFont;
 							break;
 						case ESlateFontAtlasContentType::Msdf:
-							check(IsSlateSdfTextFeatureEnabled());
+							// SDF / MSDF not supported here because we don't have SDF attributes available
+							checkNoEntry();
+							// Still use SDF shader to prevent raw MSDF texture from showing
 							ShaderType = !bEnableOutline || InOutlineSettings.bMiteredCorners ? ESlateShader::MsdfFont : ESlateShader::SdfFont;
 							break;
 						default:
@@ -3281,7 +3283,7 @@ void FSlateElementBatcher::BuildShapedTextSequence(const FShapedTextBuildContext
 			if (bIsSdfGlyph)
 			{
 				const FFontSdfSettings& FontSdfSettings = GlyphSequenceToRender->GetFontSdfSettings();
-				const FSdfGlyphFontAtlasData SdfGlyphAtlasData = Context.FontCache->GetSdfGlyphFontAtlasData(GlyphToRender, *Context.OutlineSettings, FontSdfSettings);
+				const FSdfGlyphFontAtlasData SdfGlyphAtlasData = Context.FontCache->GetSdfGlyphFontAtlasData(GlyphToRender, *Context.OutlineSettings, GlyphSequenceToRender->GetRasterizationMode(), FontSdfSettings);
 				bCanRenderGlyph = (SdfGlyphAtlasData.Valid && SdfGlyphAtlasData.bSupportsSdf);
 				if (bCanRenderGlyph)
 				{
@@ -3358,18 +3360,19 @@ void FSlateElementBatcher::BuildShapedTextSequence(const FShapedTextBuildContext
 
 					const ESlateFontAtlasContentType ContentType = SlateFontTexture->GetContentType();
 					Tint = ContentType == ESlateFontAtlasContentType::Color ? FColor::White : Context.FontTint;
-					check(bIsSdfGlyph == (ContentType == ESlateFontAtlasContentType::Msdf));
 
 					ESlateShader ShaderType = ESlateShader::Default;
 					switch (ContentType)
 					{
 						case ESlateFontAtlasContentType::Alpha:
-							ShaderType = ESlateShader::GrayscaleFont;
+							ShaderType = bIsSdfGlyph ? ESlateShader::SdfFont : ESlateShader::GrayscaleFont;
 							break;
 						case ESlateFontAtlasContentType::Color:
+							check(!bIsSdfGlyph);
 							ShaderType = ESlateShader::ColorFont;
 							break;
 						case ESlateFontAtlasContentType::Msdf:
+							check(bIsSdfGlyph);
 							ShaderType = !bOutlineFont || Context.OutlineSettings->bMiteredCorners ? ESlateShader::MsdfFont : ESlateShader::SdfFont;
 							break;
 						default:
diff --git a/Engine/Source/Runtime/SlateCore/Public/Fonts/FontCache.h b/Engine/Source/Runtime/SlateCore/Public/Fonts/FontCache.h
--- a/Engine/Source/Runtime/SlateCore/Public/Fonts/FontCache.h
+++ b/Engine/Source/Runtime/SlateCore/Public/Fonts/FontCache.h
@@ -7,6 +7,7 @@
 #include "Fonts/ShapedTextFwd.h"
 #include "UObject/ObjectMacros.h"
 #include "Fonts/SlateFontInfo.h"
+#include "Fonts/FontRasterizationMode.h"
 #include "Fonts/FontSdfSettings.h"
 #include "Textures/TextureAtlas.h"
 #include "Fonts/FontTypes.h"
@@ -253,12 +254,13 @@ private:
 struct FSdfGlyphEntryKey
 {
 public:
-	FSdfGlyphEntryKey(const TWeakPtr<FFreeTypeFace> InFontFace, uint32 InGlyphIndex, int32 InPpem, float InEmOuterSpread, float InEmInnerSpread);
+	FSdfGlyphEntryKey(const TWeakPtr<FFreeTypeFace> InFontFace, uint32 InGlyphIndex, ESlateFontAtlasContentType InAtlasContentType, int32 InPpem, float InEmOuterSpread, float InEmInnerSpread);
 
 	FORCEINLINE bool operator==(const FSdfGlyphEntryKey& Other) const
 	{
 		return FontFace == Other.FontFace
 			&& GlyphIndex == Other.GlyphIndex
+			&& AtlasContentType == Other.AtlasContentType
 			&& Ppem == Other.Ppem
 			&& SpreadCategory == Other.SpreadCategory;
 	}
@@ -277,6 +279,8 @@ public:
 	const TWeakPtr<FFreeTypeFace> FontFace;
 	/** The index of this glyph in the FreeType face */
 	const uint32 GlyphIndex;
+	/** Type of glyph pixel content */
+	const ESlateFontAtlasContentType AtlasContentType;
 	/** The pixel size at which the sdf glyph is generated */
 	const int32 Ppem;
 	/** The spread category. The spreads of a glyph entry can be arbitrary but similar values will share the same category and therefore glyph entry */
@@ -832,7 +836,7 @@ public:
 	/**
 	 * Get the atlas information and the scaled metrics of a given shaped sdf glyph. This information will be cached if required.
 	 */
-	SLATECORE_API FSdfGlyphFontAtlasData GetSdfGlyphFontAtlasData(const FShapedGlyphEntry& InShapedGlyph, const FFontOutlineSettings& InOutlineSettings, const FFontSdfSettings& InSdfSettings);
+	SLATECORE_API FSdfGlyphFontAtlasData GetSdfGlyphFontAtlasData(const FShapedGlyphEntry& InShapedGlyph, const FFontOutlineSettings& InOutlineSettings, EFontRasterizationMode InRasterizationMode, const FFontSdfSettings& InSdfSettings);
 
 	/**
 	 * Gets the overflow glyph sequence for a given font. The overflow sequence is used to replace characters that are clipped
diff --git a/Engine/Source/Runtime/SlateCore/Public/Fonts/FontProviderInterface.h b/Engine/Source/Runtime/SlateCore/Public/Fonts/FontProviderInterface.h
--- a/Engine/Source/Runtime/SlateCore/Public/Fonts/FontProviderInterface.h
+++ b/Engine/Source/Runtime/SlateCore/Public/Fonts/FontProviderInterface.h
@@ -6,6 +6,7 @@
 #include "UObject/Interface.h"
 #include "UObject/ObjectMacros.h"
 #include "UObject/UObjectGlobals.h"
+#include "FontRasterizationMode.h"
 #include "FontSdfSettings.h"
 
 #include "FontProviderInterface.generated.h"
diff --git a/Engine/Source/Runtime/SlateCore/Public/Fonts/FontRasterizationMode.h b/Engine/Source/Runtime/SlateCore/Public/Fonts/FontRasterizationMode.h
new file mode 100644
--- /dev/null
+++ b/Engine/Source/Runtime/SlateCore/Public/Fonts/FontRasterizationMode.h
@@ -0,0 +1,37 @@
+// Copyright Epic Games, Inc. All Rights Reserved.
+#pragma once
+
+#include "CoreMinimal.h"
+#include "UObject/ObjectMacros.h"
+#include "FontRasterizationMode.generated.h"
+
+/** Enumerates supported font rasterization modes. */
+UENUM()
+enum class EFontRasterizationMode : uint8
+{
+	/** Glyphs are rasterized directly into alpha mask bitmaps per size and skew. */
+	Bitmap,
+
+	/** Glyphs are rasterized into multi-channel signed distance fields, which are size and skew agnostic. */
+	Msdf UMETA(DisplayName = "Multi-Channel Distance Field"),
+
+	/** Glyphs are rasterized into single-channel signed distance fields, which are size and skew agnostic. More memory efficient but corners may appear rounded. */
+	Sdf UMETA(DisplayName = "Signed Distance Field"),
+
+	/** Glyphs are rasterized into approximate distance fields, which are size and skew agnostic. More memory and computationally efficient but lower quality. */
+	SdfApproximation UMETA(DisplayName = "Approximate Signed Distance Field (fast)")
+};
+
+/** Returns true if the font rasterization Mode is one of the distance field-based modes, which are part of the "Slate SDF text" feature */
+inline bool IsSdfFontRasterizationMode(EFontRasterizationMode Mode)
+{
+	switch (Mode)
+	{
+		case EFontRasterizationMode::Msdf:
+		case EFontRasterizationMode::Sdf:
+		case EFontRasterizationMode::SdfApproximation:
+			return true;
+		default:
+			return false;
+	}
+}
diff --git a/Engine/Source/Runtime/SlateCore/Public/Fonts/FontSdfSettings.h b/Engine/Source/Runtime/SlateCore/Public/Fonts/FontSdfSettings.h
--- a/Engine/Source/Runtime/SlateCore/Public/Fonts/FontSdfSettings.h
+++ b/Engine/Source/Runtime/SlateCore/Public/Fonts/FontSdfSettings.h
@@ -6,17 +6,6 @@
 #include "UObject/Class.h"
 #include "FontSdfSettings.generated.h"
 
-/** Enumerates supported font rasterization modes. */
-UENUM()
-enum class EFontRasterizationMode : uint8
-{
-	/** Glyphs are rasterized directly into alpha mask bitmaps per size and skew. */
-	Bitmap,
-
-	/** Glyphs are first rasterized into multi-channel signed distance fields, which are size and skew agnostic. */
-	Msdf UMETA(DisplayName = "Multi-Channel Distance Field")
-};
-
 /**
  * Settings for signed distance field fonts.
  */
