Slate Text
- Fix incoherency in ellipsis behavior between a text and its outline. As it was larger, the outline could be cut earlier than main text. Now the main text will respect the constraint from the (larger) outline.
- Fixed misplaced position of the ellipsis in Left to right texts.
- A bit of code clean up (removed unused structure members).

#tests Updated the FontOverflow engine tests, validated that there was no regression in the tests or in the editor.
[RN] Fixed an issue in the text renderer with OverflowPolicy "Ellipsis" when using outline, where the outline could get clipped at a different glyph than the "main" text.

#virtualized

diff --git a/Engine/Source/Runtime/SlateCore/Private/Rendering/ElementBatcher.cpp b/Engine/Source/Runtime/SlateCore/Private/Rendering/ElementBatcher.cpp
--- a/Engine/Source/Runtime/SlateCore/Private/Rendering/ElementBatcher.cpp
+++ b/Engine/Source/Runtime/SlateCore/Private/Rendering/ElementBatcher.cpp
@@ -1281,17 +1281,6 @@ void FSlateElementBatcher::AddTextElement(const FSlateTextElement& DrawElement)
 	FSlateFontCache& FontCache = *RenderingPolicy->GetFontCache();
 	FSlateShaderResourceManager& ResourceManager = *RenderingPolicy->GetResourceManager();
 
-	const UObject* BaseFontMaterial = DrawElement.GetFontInfo().FontMaterial;
-	const UObject* OutlineFontMaterial = OutlineSettings.OutlineMaterial;
-
-#if SLATE_CHECK_UOBJECT_RENDER_RESOURCES
-	SlateElementBatcher::CheckUObject(DrawElement, BaseFontMaterial);
-	SlateElementBatcher::CheckUObject(DrawElement, OutlineFontMaterial);
-#endif
-
-	bool bOutlineFont = OutlineSettings.OutlineSize > 0;
-	const int32 OutlineSize = OutlineSettings.OutlineSize;
-
 	auto BuildFontGeometry = [&](const FFontOutlineSettings& InOutlineSettings, const FColor& InTint, const UObject* FontMaterial, int32 InLayer, float InOutlineHorizontalOffset)
 	{
 		FCharacterList& CharacterList = FontCache.GetCharacterList(DrawElement.GetFontInfo(), FontScale, InOutlineSettings);
@@ -1525,10 +1514,19 @@ void FSlateElementBatcher::AddTextElement(const FSlateTextElement& DrawElement)
 		}
 	};
 
+	const UObject* BaseFontMaterial = DrawElement.GetFontInfo().FontMaterial;
+	const UObject* OutlineFontMaterial = OutlineSettings.OutlineMaterial;
+
+#if SLATE_CHECK_UOBJECT_RENDER_RESOURCES
+	SlateElementBatcher::CheckUObject(DrawElement, BaseFontMaterial);
+	SlateElementBatcher::CheckUObject(DrawElement, OutlineFontMaterial);
+#endif
+
+	const bool bOutlineFont = OutlineSettings.OutlineSize > 0;
 	if (bOutlineFont)
 	{
 		//The fill area was measured without an outline so it must be shifted by the scaled outline size
-		const float HorizontalOffset = FMath::RoundToFloat((float)OutlineSize * FontScale);
+		const float HorizontalOffset = FMath::RoundToFloat((float)OutlineSettings.OutlineSize * FontScale);
 
 		// Build geometry for the outline
 		BuildFontGeometry(OutlineSettings, PackVertexColor(OutlineSettings.OutlineColor), OutlineFontMaterial, Layer, HorizontalOffset);
@@ -1624,13 +1622,7 @@ void FSlateElementBatcher::AddShapedTextElement( const FSlateShapedTextElement&
 	const FSlateRenderTransform RenderTransform = Concatenate(Inverse(FontScale), DrawElement.GetRenderTransform());
 	BuildContext.RenderTransform = &RenderTransform;
 
-	const UObject* BaseFontMaterial = ShapedGlyphSequence->GetFontMaterial();
-	const UObject* OutlineFontMaterial = OutlineSettings.OutlineMaterial;
-
-	bool bOutlineFont = OutlineSettings.OutlineSize > 0;
-	const int32 OutlineSize = OutlineSettings.OutlineSize;
-
-	auto BuildFontGeometry = [&](const FFontOutlineSettings& InOutlineSettings, const FColor& InTint, const UObject* FontMaterial, int32 InLayer, float InHorizontalOffset)
+	auto BuildFontGeometry = [&](const FFontOutlineSettings& InOutlineSettings, const FColor& InTint, const UObject* FontMaterial, int32 InLayer, float InHorizontalOffset, int32 MaxGlyphCountToRender = -1) -> int32
 	{
 		FVector2f TopLeft(0, 0);
 
@@ -1638,7 +1630,6 @@ void FSlateElementBatcher::AddShapedTextElement( const FSlateShapedTextElement&
 		float PosY = TopLeft.Y;
 
 		BuildContext.FontMaterial = FontMaterial;
-		BuildContext.OutlineFontMaterial = OutlineFontMaterial;
 	
 		BuildContext.OutlineSettings = &InOutlineSettings;
 		BuildContext.StartLineX = PosX;
@@ -1646,13 +1637,12 @@ void FSlateElementBatcher::AddShapedTextElement( const FSlateShapedTextElement&
 		BuildContext.LayerId = InLayer;
 		BuildContext.FontTint = InTint;
 
-		BuildContext.bEnableOutline = InOutlineSettings.OutlineSize > 0;
-
 		// Optimize by culling
 		// Todo: this doesn't work with cached clipping
 		BuildContext.bEnableCulling = false;
 		BuildContext.bForceEllipsis = DrawElement.OverflowArgs.bIsLastVisibleBlock && DrawElement.OverflowArgs.bIsNextBlockClipped;
 		BuildContext.OverflowDirection = DrawElement.OverflowArgs.OverflowDirection;
+		BuildContext.MaxGlyphCountToRender = MaxGlyphCountToRender;
 
 		if (ShapedGlyphSequence->GetGlyphsToRender().Num() > 200 || (OverflowGlyphSequence && BuildContext.OverflowDirection != ETextOverflowDirection::NoOverflow))
 		{
@@ -1688,20 +1678,25 @@ void FSlateElementBatcher::AddShapedTextElement( const FSlateShapedTextElement&
 			}
 		}
 
-		BuildShapedTextSequence<Rounding>(BuildContext);
+		return BuildShapedTextSequence<Rounding>(BuildContext);
 	};
 
 	STAT((bSdfFont ? ElementStat_ShapedTextSdf : ElementStat_ShapedText)++);
 
+	const UObject* BaseFontMaterial = ShapedGlyphSequence->GetFontMaterial();
+	const bool bOutlineFont = OutlineSettings.OutlineSize > 0;
 	if (bOutlineFont)
 	{
-		//The fill area was measured without an outline so it must be shifted by the scaled outline size
-		const float HorizontalOffset = FMath::RoundToFloat((float)OutlineSize * FontScale);
+		//The fill area was measured without an outline so it must be shifted by the scaled outline size.
+		//The bounding box of the text is bigger than a standard one, due to the outline (we add the scaled outline size * 2), 
+		//so the horizontal offset ensure that the outline won't bleed outside of the box and will be correctly positioned inside the box.
+		const float HorizontalOffset = FMath::RoundToFloat((float)OutlineSettings.OutlineSize * FontScale);
+		const UObject* OutlineFontMaterial = OutlineSettings.OutlineMaterial;
 
 		// Build geometry for the outline
-		BuildFontGeometry(OutlineSettings, PackVertexColor(DrawElement.GetOutlineTint()), OutlineFontMaterial, Layer, HorizontalOffset);
+		const int32 GlyphsRendered = BuildFontGeometry(OutlineSettings, PackVertexColor(DrawElement.GetOutlineTint()), OutlineFontMaterial, Layer, HorizontalOffset);
 		// Build geometry for the base font which is always rendered on top of the outline 
-		BuildFontGeometry(FFontOutlineSettings::NoOutline, BaseTint, BaseFontMaterial, Layer+1, HorizontalOffset);
+		BuildFontGeometry(FFontOutlineSettings::NoOutline, BaseTint, BaseFontMaterial, Layer+1, HorizontalOffset, GlyphsRendered);
 	}
 	else
 	{
@@ -3204,7 +3199,7 @@ const FSlateClippingState* FSlateElementBatcher::ResolveClippingState(const FSla
 }
 
 template<ESlateVertexRounding Rounding>
-void FSlateElementBatcher::BuildShapedTextSequence(const FShapedTextBuildContext& Context)
+int32 FSlateElementBatcher::BuildShapedTextSequence(const FShapedTextBuildContext& Context)
 {
 	const FShapedGlyphSequence* GlyphSequenceToRender = Context.ShapedGlyphSequence;
 
@@ -3244,8 +3239,10 @@ void FSlateElementBatcher::BuildShapedTextSequence(const FShapedTextBuildContext
 	const int32 MaxPreAllocatedGlyphIndicies = 4096 + 1024;
 	const bool bUseStaticIndicies = GlyphSequenceToRender->GetGlyphsToRender().Num() < MaxPreAllocatedGlyphIndicies;
 
+	const bool bOutlineFont = Context.OutlineSettings->OutlineSize > 0;
 	const int32 NumGlyphs = GlyphSequenceToRender->GetGlyphsToRender().Num();
 	const TArray<FShapedGlyphEntry>& GlyphsToRender = GlyphSequenceToRender->GetGlyphsToRender();
+	int32 GlyphsRendered = 0;
 	for (int32 GlyphIndex = 0; GlyphIndex < NumGlyphs; ++GlyphIndex)
 	{
 		const FShapedGlyphEntry& GlyphToRender = GlyphsToRender[GlyphIndex];
@@ -3263,9 +3260,7 @@ void FSlateElementBatcher::BuildShapedTextSequence(const FShapedTextBuildContext
 		float SizeV = 0;
 
 		bool bIsVisible = GlyphToRender.bIsVisible;
-
 		bool bCanRenderGlyph = bIsVisible;
-		const bool bOutlineFont = Context.OutlineSettings->OutlineSize > 0;
 
 		FVector2f SpriteSize(0.f, 0.f);
 		FVector2f SpriteOffset(0.f, 0.f);
@@ -3310,7 +3305,7 @@ void FSlateElementBatcher::BuildShapedTextSequence(const FShapedTextBuildContext
 			else
 			{
 				const FShapedGlyphFontAtlasData GlyphAtlasData = Context.FontCache->GetShapedGlyphFontAtlasData(GlyphToRender, *Context.OutlineSettings);
-				bCanRenderGlyph = (GlyphAtlasData.Valid && (!Context.bEnableOutline || GlyphAtlasData.SupportsOutline));
+				bCanRenderGlyph = (GlyphAtlasData.Valid && (!bOutlineFont || GlyphAtlasData.SupportsOutline));
 				if (bCanRenderGlyph)
 				{
 					NextAtlasDataTextureIndex = GlyphAtlasData.TextureIndex;
@@ -3484,7 +3479,7 @@ void FSlateElementBatcher::BuildShapedTextSequence(const FShapedTextBuildContext
 		{
 			// If we are on the last glyph don't bother checking if the ellipsis can fit. If the last glyph can fit there is no need for ellipsis
 			float OverflowSequenceNeededSize = GlyphIndex < NumGlyphs - 1 ? Context.OverflowGlyphSequence->GetMeasuredWidth() : 0;
-			if(X + OverflowTestWidth + OverflowSequenceNeededSize >= Context.LocalClipBoundingBoxRight)
+			if(X + OverflowTestWidth + OverflowSequenceNeededSize >= Context.LocalClipBoundingBoxRight || (Context.MaxGlyphCountToRender >= 0 && GlyphIndex >= Context.MaxGlyphCountToRender))
 			{
 				bNeedEllipsis = true;
 				// We subtract out any whitespace advance. This avoids the ellipsis from ever floating out in the middle of a block of whitespace.
@@ -3494,12 +3489,13 @@ void FSlateElementBatcher::BuildShapedTextSequence(const FShapedTextBuildContext
 				// No characters to render after the ellipsis on the right side
 				break;
 			}
+			GlyphsRendered++;
 		}
 		else if(OverflowDirection == ETextOverflowDirection::RightToLeft)
 		{
 			bool bClipped = false;
 			// Right to left overflow
-			if (X < Context.LocalClipBoundingBoxLeft)
+			if (X < Context.LocalClipBoundingBoxLeft || (Context.MaxGlyphCountToRender >= 0 && GlyphIndex < (NumGlyphs - Context.MaxGlyphCountToRender)))
 			{
 				// This glyph is in the clipped region or is not visible so just advance. It cannot be shown
 				bClipped = true;
@@ -3511,22 +3507,23 @@ void FSlateElementBatcher::BuildShapedTextSequence(const FShapedTextBuildContext
 
 				// Can the ellipsis fit in the free spot by skipping the previous glyph(s)
 				const float EllipsisWidth = Context.OverflowGlyphSequence->GetMeasuredWidth();
-				const float AvailableX = X + SizeX - Context.LocalClipBoundingBoxLeft;
+				const float AvailableX = X + GlyphToRender.XAdvance - Context.LocalClipBoundingBoxLeft;
 				if (AvailableX >= EllipsisWidth)
 				{
 					// The available area can fit the ellipsis. Mark that we need an ellipsis and stop checking for overflow. The rest of the text can be built normally
 					bNeedSpaceForEllipsis = false;
 				}
 
-				//Always try to put the ellipsis, wether it fits or not: it's better to have an ellipsis a bit clipped than no feedback at all.
-					bNeedEllipsis = true;
-					EllipsisLineX = (LineX + SizeX - EllipsisWidth);
-					EllipsisLineY = LineY;
-				}
-				else
-				{
-					OverflowDirection = ETextOverflowDirection::NoOverflow;
-				}
+				//Always try to put the ellipsis, whether it fits or not: it's better to have an ellipsis a bit clipped than no feedback at all.
+				bNeedEllipsis = true;
+				EllipsisLineX = (LineX + GlyphToRender.XAdvance - EllipsisWidth);
+				EllipsisLineY = LineY;
+			}
+			else
+			{
+				OverflowDirection = ETextOverflowDirection::NoOverflow;
+				GlyphsRendered++;
+			}
 
 			// If we just clipped a glyph omit all characters in said glyph. Otherwise floating diacritics would be visible above the ellipsis. This is common in complex languages.
 			if (bClipped && GlyphToRender.NumCharactersInGlyph > 1)
@@ -3538,6 +3535,10 @@ void FSlateElementBatcher::BuildShapedTextSequence(const FShapedTextBuildContext
 
 			bCanRenderGlyph = !bClipped;
 		}
+		else
+		{
+			GlyphsRendered++;
+		}
 
 		if(bCanRenderGlyph && RenderBatch)
 		{
@@ -3627,6 +3628,8 @@ void FSlateElementBatcher::BuildShapedTextSequence(const FShapedTextBuildContext
 
 		BuildShapedTextSequence<Rounding>(EllipsisContext);
 	}
+
+	return GlyphsRendered;
 }
 
 void FSlateElementBatcher::ResetBatches()
diff --git a/Engine/Source/Runtime/SlateCore/Public/Rendering/ElementBatcher.h b/Engine/Source/Runtime/SlateCore/Public/Rendering/ElementBatcher.h
--- a/Engine/Source/Runtime/SlateCore/Public/Rendering/ElementBatcher.h
+++ b/Engine/Source/Runtime/SlateCore/Public/Rendering/ElementBatcher.h
@@ -423,7 +423,6 @@ private:
 		const class FShapedGlyphSequence* ShapedGlyphSequence;
 		const class FShapedGlyphSequence* OverflowGlyphSequence;
 		const UObject* FontMaterial;
-		const UObject* OutlineFontMaterial;
 		const struct FFontOutlineSettings* OutlineSettings;
 		const FSlateDrawElement* DrawElement;
 		class FSlateFontCache* FontCache;
@@ -434,17 +433,17 @@ private:
 		float StartLineY;
 		float LocalClipBoundingBoxLeft = 0;
 		float LocalClipBoundingBoxRight = 0;
+		int32 MaxGlyphCountToRender = -1;
 		int32 LayerId;
 		FColor FontTint;
 		ETextOverflowDirection OverflowDirection;
-		bool bEnableOutline : 1;
 		bool bEnableCulling : 1;
 		bool bForceEllipsis : 1;
 		
 	};
 
 	template<ESlateVertexRounding Rounding>
-	void BuildShapedTextSequence(const FShapedTextBuildContext& Context);
+	int32 BuildShapedTextSequence(const FShapedTextBuildContext& Context);
 private:
 	/** Uncached Batch data currently being filled in */
 	FSlateBatchData* BatchData;
