Slate text: Fix ellipsis position for left to right overflow direction that was displayed after the last whitespaces of a line in case of wrapping text instead of right after the last visible glyph, which could lead to a part of the ellipsis to be clipped.
#tests Validated via functionnal test that the ellipsis was better positionned than before in specific case. No regression in this text.
[RN] Fixed the ellipsis position for texts that are both left to right and wrapping . The ellipsis was displayed after the last whitespaces of a line instead of right after the last visible glyph, which could result in the ellipsis to be partially clipped.

diff --git a/Engine/Source/Runtime/SlateCore/Private/Rendering/ElementBatcher.cpp b/Engine/Source/Runtime/SlateCore/Private/Rendering/ElementBatcher.cpp
--- a/Engine/Source/Runtime/SlateCore/Private/Rendering/ElementBatcher.cpp
+++ b/Engine/Source/Runtime/SlateCore/Private/Rendering/ElementBatcher.cpp
@@ -3610,7 +3610,7 @@ int32 FSlateElementBatcher::BuildShapedTextSequence(const FShapedTextBuildContex
 	if (!bNeedEllipsis && Context.bForceEllipsis)
 	{
 		bNeedEllipsis = true;
-		EllipsisLineX = LineX;
+		EllipsisLineX = LineX - PreviousWhitespaceAdvance; //It's ok to always substract whitespace advance (LTR or RTL overflow direction), as the value will always be 0 for RTL overflow direction.
 		EllipsisLineY = LineY;
 	}
 
