Slate SDF font - fix for excessive signed distance spread when text is tiny

diff --git a/Engine/Source/Runtime/SlateCore/Private/Fonts/FontCache.cpp b/Engine/Source/Runtime/SlateCore/Private/Fonts/FontCache.cpp
--- a/Engine/Source/Runtime/SlateCore/Private/Fonts/FontCache.cpp
+++ b/Engine/Source/Runtime/SlateCore/Private/Fonts/FontCache.cpp
@@ -1206,15 +1206,15 @@ FSdfGlyphFontAtlasData FSlateFontCache::GetSdfGlyphFontAtlasData(const FShapedGl
 
 	const bool bPlaceholderOnly = InRasterizationMode == EFontRasterizationMode::SdfApproximation;
 #if WITH_FREETYPE
-	const float TargetPpem = static_cast<float>(FreeTypeUtils::ComputeFontPixelSize(InShapedGlyph.FontFaceData->FontSize, InShapedGlyph.FontFaceData->FontScale));
+	const float TargetPpem = FMath::Max(1.f, static_cast<float>(FreeTypeUtils::ComputeFontPixelSize(InShapedGlyph.FontFaceData->FontSize, InShapedGlyph.FontFaceData->FontScale)));
 #else
-	const float TargetPpem = 1.f/72.f*FontConstants::RenderDPI*InShapedGlyph.FontFaceData->FontSize*InShapedGlyph.FontFaceData->FontScale;
+	const float TargetPpem = FMath::Max(1.f, 1.f/72.f*static_cast<float>(FontConstants::RenderDPI)*InShapedGlyph.FontFaceData->FontSize*InShapedGlyph.FontFaceData->FontScale);
 #endif
 	const float SdfPpem = static_cast<float>(InSdfSettings.GetClampedPpem());
 	const float Scale = TargetPpem/SdfPpem;
 	const float EmOutlineSize = FMath::RoundToFloat(InOutlineSettings.OutlineSize * InShapedGlyph.FontFaceData->FontScale)/TargetPpem;
 
-	float EmInnerSpread = FMath::Clamp(2.f/TargetPpem, 0.05f, 8.f);
+	float EmInnerSpread = FMath::Clamp(2.f/TargetPpem, 0.05f, 0.25f);
 	float EmOuterSpread = FMath::Min(EmInnerSpread + EmOutlineSize, 8.f);
 
 	FSlateSdfGenerator::ESdfType SdfType = FSlateSdfGenerator::ESdfType::MultichannelAndSimple;
