Slate: Added the ability to draw dotted and dashed line segments

In order to support this the following changes have been made:
 - A new API for FSlateDrawElement::MakeDashedLines.
 - Anti-aliased lines now upload texture coordinates dictating each vertex's position along the line, and the dash length in addition to the existing antialias gradient. This shouldn't add any additional bandwidth overhead since these texture coordinates were always being uploaded before anyway, just with a value of (1,1)
 - Additional shader instructions for anti-aliased lines are now also run to compute the opacity for each pixel based on the dash length and position. These are always run, even for non-dashed, antialiased lines. Profiling showed negligible performance difference for a test with heavy line-rendering.

diff --git a/Engine/Shaders/Private/SlateElementPixelShader.usf b/Engine/Shaders/Private/SlateElementPixelShader.usf
--- a/Engine/Shaders/Private/SlateElementPixelShader.usf
+++ b/Engine/Shaders/Private/SlateElementPixelShader.usf
@@ -374,6 +374,35 @@ half4 GetRoundedBoxElementColor( VertexToPixelInterpolants InVertex )
 #endif
 
 #if SHADER_TYPE == ST_LineSegment
+
+/**
+ * Computes 0-1 value that interpolates between dashes of length DashLength based on Distance with 1px of antialiasing between dashes and gaps
+ * 
+ * ie:
+ * 1 -----------------                     -----------------                     -----------------                     
+ *                    \                   /                 \                   /                 \                   /
+ * 0                   \_________________/                   \_________________/                   \_________________/ 
+ *
+ *   <  DashLength   >  <  DashLength   >  <  DashLength   >  <  DashLength   >  <  DashLength   >  <  DashLength   >
+ 
+ */
+float ModulateDashedLine(float Distance, float DashLength)
+{
+	// Add one for antialias pixel
+	DashLength += 1.f;
+
+	const float HalfDash  = 0.5f*DashLength;
+	const float Period    = 2.0f*DashLength;
+
+	// Offset the interpolation so that the dash starts at distance = 0;
+	const float Interp = Distance - HalfDash + 0.5f;
+	const float Mod = Interp - Period*floor(Interp/Period); // mod(Interp, Period)
+
+	// Compute the dash alpha using a clamped triangle wave
+	const float TriWave = 2.0f * abs(Mod - DashLength) - DashLength;
+	return clamp(TriWave, -1.0f, 1.0f)*0.5f + 0.5f;
+}
+
 /**
  * Generates an anti-aliased line segment pixel
  * This is based on the fast prefiltered lines technique published in GPU Gems 2
@@ -387,16 +416,21 @@ half4 GetRoundedBoxElementColor( VertexToPixelInterpolants InVertex )
 half4 GetLineSegmentElementColor( VertexToPixelInterpolants InVertex )
 {
 	const float2 Gradient = SLATE_UV0(InVertex);
+	const float2 DashParams = SLATE_UV1(InVertex);
 
 	// Get coverage based on distance from segment sides and ends
 	const float2 OutsideFilterUV = float2(1.0f, 1.0f);
 	const float2 InsideFilterUV = float2(ShaderParams.x, 0.0f);
 	const float2 LineCoverage = smoothstep(OutsideFilterUV, InsideFilterUV, abs(Gradient));
 
+	const float DashLength = DashParams.y;
+	const float DashAlpha = ModulateDashedLine(DashParams.x, DashParams.y);
+
 	half4 Color = InVertex.Color;
-	Color.a *= LineCoverage.x * LineCoverage.y;
+	Color.a *= LineCoverage.x * LineCoverage.y * DashAlpha;
 	return Color;
 }
+
 #endif
 
 HALF4_TYPE Main( VertexToPixelInterpolants VIn ) : SV_Target0
diff --git a/Engine/Shaders/StandaloneRenderer/D3D/SlateElementPixelShader.hlsl b/Engine/Shaders/StandaloneRenderer/D3D/SlateElementPixelShader.hlsl
--- a/Engine/Shaders/StandaloneRenderer/D3D/SlateElementPixelShader.hlsl
+++ b/Engine/Shaders/StandaloneRenderer/D3D/SlateElementPixelShader.hlsl
@@ -184,16 +184,59 @@ float4 GetRoundedBoxElementColor( VertexOut InVertex )
 	return OutColor;
 }
 
+/**
+ * Computes 0-1 value that interpolates between dashes of length DashLength based on Distance with 1px of antialiasing between dashes and gaps
+ * 
+ * ie:
+ * 1 -----------------                     -----------------                     -----------------                     
+ *                    \                   /                 \                   /                 \                   /
+ * 0                   \_________________/                   \_________________/                   \_________________/ 
+ *
+ *   <  DashLength   >  <  DashLength   >  <  DashLength   >  <  DashLength   >  <  DashLength   >  <  DashLength   >
+ 
+ */
+float ModulateDashedLine(float Distance, float DashLength)
+{
+	// Add one for antialias pixel
+	DashLength += 1.f;
+
+	const float HalfDash  = 0.5f*DashLength;
+	const float Period    = 2.0f*DashLength;
+
+	// Offset the interpolation so that the dash starts at distance = 0;
+	const float Interp = Distance - HalfDash + 0.5f;
+	const float Mod = Interp - Period*floor(Interp/Period); // mod(Interp, Period)
+
+	// Compute the dash alpha using a clamped triangle wave
+	const float TriWave = 2.0f * abs(Mod - DashLength) - DashLength;
+	return clamp(TriWave, -1.0f, 1.0f)*0.5f + 0.5f;
+}
+
+/**
+ * Generates an anti-aliased line segment pixel
+ * This is based on the fast prefiltered lines technique published in GPU Gems 2
+ *
+ * Instead of passing edge function coefficients as uniform parameters we use
+ * texture coordinates to define the line
+ *
+ * When rasterized, the submitted geometry must contain every pixel that
+ * could have coverage > 0. Failing to do so will create sharp, aliased edges
+ */
 float4 GetLineSegmentElementColor( VertexOut InVertex )
 {
 	const float2 Gradient = InVertex.TextureCoordinates;
+	const float2 DashParams = InVertex.TextureCoordinates.zw;
 
+	// Get coverage based on distance from segment sides and ends
 	const float2 OutsideFilterUV = float2(1.0f, 1.0f);
 	const float2 InsideFilterUV = float2(ShaderParams.x, 0.0f);
 	const float2 LineCoverage = smoothstep(OutsideFilterUV, InsideFilterUV, abs(Gradient));
 
+	const float DashLength = DashParams.y;
+	const float DashAlpha = ModulateDashedLine(DashParams.x, DashParams.y);
+
 	float4 Color = InVertex.Color;
-	Color.a *= LineCoverage.x * LineCoverage.y;
+	Color.a *= LineCoverage.x * LineCoverage.y * DashAlpha;
 	return Color;
 }
 
diff --git a/Engine/Shaders/StandaloneRenderer/OpenGL/SlateElementPixelShader.glsl b/Engine/Shaders/StandaloneRenderer/OpenGL/SlateElementPixelShader.glsl
--- a/Engine/Shaders/StandaloneRenderer/OpenGL/SlateElementPixelShader.glsl
+++ b/Engine/Shaders/StandaloneRenderer/OpenGL/SlateElementPixelShader.glsl
@@ -287,16 +287,60 @@ vec4 GetRoundedBoxElementColor()
 	return OutColor;
 }
 
+
+/**
+ * Computes 0-1 value that interpolates between dashes of length DashLength based on Distance with 1px of antialiasing between dashes and gaps
+ * 
+ * ie:
+ * 1 -----------------                     -----------------                     -----------------                     
+ *                    \                   /                 \                   /                 \                   /
+ * 0                   \_________________/                   \_________________/                   \_________________/ 
+ *
+ *   <  DashLength   >  <  DashLength   >  <  DashLength   >  <  DashLength   >  <  DashLength   >  <  DashLength   >
+ 
+ */
+float ModulateDashedLine(float Distance, float DashLength)
+{
+	// Add one for antialias pixel
+	DashLength += 1.0;
+
+	float HalfDash  = 0.50*DashLength;
+	float Period    = 2.00*DashLength;
+
+	// Offset the interpolation so that the dash starts at distance = 0;
+	float Interp = Distance - HalfDash + 0.50;
+	float Mod = Interp - Period*floor(Interp/Period); // mod(Interp, Period)
+
+	// Compute the dash alpha using a clamped triangle wave
+	float TriWave = 2.0 * abs(Mod - DashLength) - DashLength;
+	return clamp(TriWave, -1.0, 1.0)*0.5 + 0.5;
+}
+
+/**
+ * Generates an anti-aliased line segment pixel
+ * This is based on the fast prefiltered lines technique published in GPU Gems 2
+ *
+ * Instead of passing edge function coefficients as uniform parameters we use
+ * texture coordinates to define the line
+ *
+ * When rasterized, the submitted geometry must contain every pixel that
+ * could have coverage > 0. Failing to do so will create sharp, aliased edges
+ */
 vec4 GetLineSegmentElementColor()
 {
 	vec2 Gradient = TexCoords.xy;
+	vec2 DashParams = TexCoords.zw;
 
-	vec2 OutsideFilterUV = vec2(1.0, 1.0);
-	vec2 InsideFilterUV = vec2(ShaderParams.x, 0.0);
+	// Get coverage based on distance from segment sides and ends
+	vec2 OutsideFilterUV = float2(1.0, 1.0);
+	vec2 InsideFilterUV = float2(ShaderParams.x, 0.0);
 	vec2 LineCoverage = smoothstep(OutsideFilterUV, InsideFilterUV, abs(Gradient));
 
+	float DashLength = DashParams.y;
+	float DashAlpha = ModulateDashedLine(DashParams.x, DashParams.y);
+
 	vec4 OutColor = Color;
-	OutColor.a *= LineCoverage.x * LineCoverage.y;
+	OutColor.a *= LineCoverage.x * LineCoverage.y * DashAlpha;
 	return OutColor;
 }
 
diff --git a/Engine/Source/Runtime/AppFramework/Private/Widgets/Testing/STestSuite.cpp b/Engine/Source/Runtime/AppFramework/Private/Widgets/Testing/STestSuite.cpp
--- a/Engine/Source/Runtime/AppFramework/Private/Widgets/Testing/STestSuite.cpp
+++ b/Engine/Source/Runtime/AppFramework/Private/Widgets/Testing/STestSuite.cpp
@@ -830,12 +830,18 @@ public:
 					.OnPaintHandler( this, &SElementTesting::TestGradientElement ) 
 				]
 				+ SVerticalBox::Slot()
-				.FillHeight(1)
+				.AutoHeight()
 				[
 					SNew( SCustomPaintWidget )
 					.OnPaintHandler( this, &SElementTesting::TestSplineElement ) 
 				]
 				+ SVerticalBox::Slot()
+				.AutoHeight()
+				[
+					SNew( SCustomPaintWidget )
+					.OnPaintHandler( this, &SElementTesting::TestDottedLineElement ) 
+				]
+				+ SVerticalBox::Slot()
 				.FillHeight(3)
 				[
 					SNew( SCustomPaintWidget )
@@ -1089,6 +1095,137 @@ private:
 
 	}
 
+	int32 TestDottedLineElement( const FOnPaintHandlerParams& InParams )
+	{
+		const FVector2f Start(10,10);
+		const FVector2f StartDir(InParams.Geometry.GetLocalSize().X * 1000 / 600,0);
+		const FVector2f End(InParams.Geometry.GetLocalSize().X/4, InParams.Geometry.GetLocalSize().Y-10);
+		const FVector2f EndDir(InParams.Geometry.GetLocalSize().X * 1000 / 600,0);
+
+		static float CurTime = 0; 
+		CurTime += FSlateApplication::Get().GetDeltaTime();
+
+		float DashLengths[] = { 5.f, 20.f, 50.f };
+		float DashLength = DashLengths[FMath::FloorToInt(CurTime) % 3];
+
+		FVector2f LineStart =  FVector2f( InParams.Geometry.GetLocalSize().X/4, 10.0f );
+
+		TArray<FVector2f> LinePoints;
+		TArray<FLinearColor> LineColors;
+		LinePoints.Add(LineStart); LineColors.Add(FLinearColor::Red);
+		LinePoints.Add( LineStart + FVector2f( 100.0f, 50.0f ) );
+		LinePoints.Add( LineStart + FVector2f( 200.0f, 10.0f ) );
+		LinePoints.Add( LineStart + FVector2f( 300.0f, 50.0f ) );
+		LinePoints.Add( LineStart + FVector2f( 400.0f, 10.0f ) );
+
+
+		FSlateDrawElement::MakeDashedLines( 
+			InParams.OutDrawElements,
+			InParams.Layer,
+			InParams.Geometry.ToPaintGeometry(),
+			CopyTemp(LinePoints),
+			InParams.bEnabled ? ESlateDrawEffect::None : ESlateDrawEffect::DisabledEffect,
+			FColor::White,
+			InParams.Geometry.Scale,
+			DashLength,
+			-CurTime * 10.f
+		);
+
+		FSlateDrawElement::MakeDashedLines(
+			InParams.OutDrawElements,
+			InParams.Layer,
+			InParams.Geometry.ToOffsetPaintGeometry(FVector2D(0.f, 10.f)),
+			CopyTemp(LinePoints),
+			InParams.bEnabled ? ESlateDrawEffect::None : ESlateDrawEffect::DisabledEffect,
+			FColor::White,
+			5.f * InParams.Geometry.Scale,
+			DashLength,
+			-CurTime * 20.f
+		);
+
+		FSlateDrawElement::MakeDashedLines(
+			InParams.OutDrawElements,
+			InParams.Layer,
+			InParams.Geometry.ToOffsetPaintGeometry(FVector2D(0.f, 20.f)),
+			CopyTemp(LinePoints),
+			InParams.bEnabled ? ESlateDrawEffect::None : ESlateDrawEffect::DisabledEffect,
+			FColor::White,
+			10.f * InParams.Geometry.Scale,
+			DashLength,
+			-CurTime * 30.f
+		);
+
+		FSlateDrawElement::MakeDashedLines(
+			InParams.OutDrawElements,
+			InParams.Layer,
+			InParams.Geometry.ToOffsetPaintGeometry(FVector2D(0.f, 34.f)),
+			CopyTemp(LinePoints),
+			InParams.bEnabled ? ESlateDrawEffect::None : ESlateDrawEffect::DisabledEffect,
+			FColor::White,
+			20.f * InParams.Geometry.Scale,
+			DashLength,
+			-CurTime * 40.f
+		);
+			
+
+		LineStart =  LinePoints[ LinePoints.Num() - 1 ] + FVector2f(50.f,10.f);
+		LinePoints.Empty();
+
+		float CurTimeModTwoPi = FMath::Fmod(CurTime, 2*PI);
+		for( float I = 0; I < 10*PI; I+=.1f)
+		{
+			LinePoints.Add( LineStart + FVector2f( I*15 , 15*FMath::Sin( I + CurTimeModTwoPi) ) );
+		}
+
+		static FColor Color = FColor::MakeRandomColor();
+		FSlateDrawElement::MakeDashedLines( 
+			InParams.OutDrawElements,
+			InParams.Layer+1,
+			InParams.Geometry.ToPaintGeometry(),
+			CopyTemp(LinePoints),
+			InParams.bEnabled ? ESlateDrawEffect::NoPixelSnapping : ESlateDrawEffect::NoPixelSnapping|ESlateDrawEffect::DisabledEffect,
+			Color,
+			InParams.Geometry.Scale,
+			DashLength,
+			CurTime*15.f);
+
+		FSlateDrawElement::MakeDashedLines(
+			InParams.OutDrawElements,
+			InParams.Layer,
+			InParams.Geometry.ToOffsetPaintGeometry(FVector2D(0.f, 10.f)),
+			CopyTemp(LinePoints),
+			InParams.bEnabled ? ESlateDrawEffect::NoPixelSnapping : ESlateDrawEffect::NoPixelSnapping | ESlateDrawEffect::DisabledEffect,
+			Color,
+			2.f * InParams.Geometry.Scale,
+			DashLength,
+			CurTime*15.f);
+
+		FSlateDrawElement::MakeDashedLines(
+			InParams.OutDrawElements,
+			InParams.Layer,
+			InParams.Geometry.ToOffsetPaintGeometry(FVector2D(0.f, 20.f)),
+			CopyTemp(LinePoints),
+			InParams.bEnabled ? ESlateDrawEffect::NoPixelSnapping : ESlateDrawEffect::NoPixelSnapping | ESlateDrawEffect::DisabledEffect,
+			Color,
+			3.f * InParams.Geometry.Scale,
+			DashLength,
+			CurTime*15.f);
+
+		FSlateDrawElement::MakeDashedLines(
+			InParams.OutDrawElements,
+			InParams.Layer,
+			InParams.Geometry.ToOffsetPaintGeometry(FVector2D(0.f, 34.f)),
+			CopyTemp(LinePoints),
+			InParams.bEnabled ? ESlateDrawEffect::NoPixelSnapping : ESlateDrawEffect::NoPixelSnapping | ESlateDrawEffect::DisabledEffect,
+			Color,
+			4.f * InParams.Geometry.Scale,
+			DashLength,
+			CurTime*15.f);
+
+		return InParams.Layer;
+
+	}
+
 	void MakeRotationExample( const FOnPaintHandlerParams& InParams )
 	{
 		const FSlateBrush* CenterBrush = FTestStyle::Get().GetBrush("TestRotation40px");
diff --git a/Engine/Source/Runtime/SlateCore/Private/Rendering/DrawElementTypes.cpp b/Engine/Source/Runtime/SlateCore/Private/Rendering/DrawElementTypes.cpp
--- a/Engine/Source/Runtime/SlateCore/Private/Rendering/DrawElementTypes.cpp
+++ b/Engine/Source/Runtime/SlateCore/Private/Rendering/DrawElementTypes.cpp
@@ -600,6 +600,29 @@ void FSlateDrawElement::MakeLines(FSlateWindowElementList& ElementList, uint32 I
 	Element.Init(ElementList, EElementType::ET_Line, InLayer, PaintGeometry, DrawEffects);
 }
 
+void FSlateDrawElement::MakeDashedLines(FSlateWindowElementList& ElementList, uint32 InLayer, const FPaintGeometry& PaintGeometry, TArray<FVector2f>&& Points, ESlateDrawEffect InDrawEffects, const FLinearColor& InTint, float Thickness, float DashLengthPx, float DashScreenOffset)
+{
+	PaintGeometry.CommitTransformsIfUsingLegacyConstructor();
+
+	if (ShouldCull(ElementList) || Points.Num() < 2)
+	{
+		return;
+	}
+
+	FSlateLineElement& Element = ElementList.AddUninitialized<EElementType::ET_Line>();
+
+	Element.SetTint(InTint);
+	Element.SetThickness(Thickness);
+	Element.SetLines(MoveTemp(Points), true);
+
+	Element.DashLength = DashLengthPx;
+	Element.DashOffset = DashScreenOffset;
+
+	ESlateDrawEffect DrawEffects = InDrawEffects | ESlateDrawEffect::NoPixelSnapping;
+
+	Element.Init(ElementList, EElementType::ET_Line, InLayer, PaintGeometry, DrawEffects);
+}
+
 #if UE_ENABLE_SLATE_VECTOR_DEPRECATION_MECHANISMS
 void FSlateDrawElement::MakeLines(FSlateWindowElementList& ElementList, uint32 InLayer, const FPaintGeometry& PaintGeometry, const TArray<FVector2d>& Points, const TArray<FLinearColor>& PointColors, ESlateDrawEffect InDrawEffects, const FLinearColor& InTint, bool bAntialias, float Thickness)
 {
@@ -636,6 +659,29 @@ void FSlateDrawElement::MakeLines( FSlateWindowElementList& ElementList, uint32
 	Element.Init(ElementList, EElementType::ET_Line, InLayer, PaintGeometry, InDrawEffects);
 }
 
+void FSlateDrawElement::MakeDashedLines(FSlateWindowElementList& ElementList, uint32 InLayer, const FPaintGeometry& PaintGeometry, TArray<FVector2f>&& Points, TArray<FLinearColor>&& PointColors, ESlateDrawEffect InDrawEffects, const FLinearColor& InTint, float Thickness, float DashLengthPx, float DashScreenOffset)
+{
+	PaintGeometry.CommitTransformsIfUsingLegacyConstructor();
+
+	if (ShouldCull(ElementList) || Points.Num() < 2)
+	{
+		return;
+	}
+
+	FSlateLineElement& Element = ElementList.AddUninitialized<EElementType::ET_Line>();
+
+	Element.SetTint(InTint);
+	Element.SetThickness(Thickness);
+	Element.SetLines(MoveTemp(Points), true, MoveTemp(PointColors));
+
+	Element.DashLength = DashLengthPx;
+	Element.DashOffset = DashScreenOffset;
+
+	ESlateDrawEffect DrawEffects = InDrawEffects | ESlateDrawEffect::NoPixelSnapping;
+
+	Element.Init(ElementList, EElementType::ET_Line, InLayer, PaintGeometry, DrawEffects);
+}
+
 void FSlateDrawElement::MakeViewport( FSlateWindowElementList& ElementList, uint32 InLayer, const FPaintGeometry& PaintGeometry, TSharedPtr<const ISlateViewport> Viewport, ESlateDrawEffect InDrawEffects, const FLinearColor& InTint )
 {
 	PaintGeometry.CommitTransformsIfUsingLegacyConstructor();
diff --git a/Engine/Source/Runtime/SlateCore/Private/Rendering/ElementBatcher.cpp b/Engine/Source/Runtime/SlateCore/Private/Rendering/ElementBatcher.cpp
--- a/Engine/Source/Runtime/SlateCore/Private/Rendering/ElementBatcher.cpp
+++ b/Engine/Source/Runtime/SlateCore/Private/Rendering/ElementBatcher.cpp
@@ -2181,14 +2181,20 @@ struct FLineBuilder
 	 * @param FilterRadius		Antialiasing filter radius in screenspace pixels.
 	 * @param AngleCosineLimit	Miter Angle Limit after being passed through AngleCosine.
 	 */
-	FLineBuilder(FSlateRenderBatch& InRenderBatch, const FSlateRenderTransform& InRenderTransform, float ElementScale, float HalfThickness, float FilterRadius, float MiterAngleLimit) :
+	FLineBuilder(FSlateRenderBatch& InRenderBatch, const FSlateRenderTransform& InRenderTransform, float ElementScale, float HalfThickness, float FilterRadius, float MiterAngleLimit, float InDashLength, float InDashOffset) :
 		RenderBatch(InRenderBatch),
 		RenderTransform(InRenderTransform),
 		LocalHalfThickness((HalfThickness + FilterRadius) / ElementScale),
 		LocalFilterRadius(FilterRadius / ElementScale),
 		LocalCapLength((FilterRadius / ElementScale) * 2.0f),
-		AngleCosineLimit(FMath::DegreesToRadians((180.0f - MiterAngleLimit) * 0.5f))
+		AngleCosineLimit(FMath::DegreesToRadians((180.0f - MiterAngleLimit) * 0.5f)),
+		DashLength(InDashLength),
+		DashOffset(InDashOffset)
 	{
+		if (DashLength < 1.f)
+		{
+			DashLength = SolidLineDashLength;
+		}
 	}
 
 	/**
@@ -2196,6 +2202,8 @@ struct FLineBuilder
 	 */
 	void NumElements(const TArray<FVector2f>& Points, uint32& OutNumVertex, uint32& OutNumIndex)
 	{
+		const bool bIsDashed = DashLength != SolidLineDashLength;
+
 		FVector2f Position = Points[0];
 		FVector2f NextPosition = Points[1];
 
@@ -2233,7 +2241,7 @@ struct FLineBuilder
 				DirDotMiterNormal >= AngleCosineLimit &&
 				(MinSegmentLength * 0.5f * DirDotMiterNormal) >= FMath::Abs(DistanceToMiterLine))
 			{
-				OutNumVertex += 2;
+				OutNumVertex += bIsDashed ? 4 : 2;
 				OutNumIndex += 6;
 			}
 			else
@@ -2258,8 +2266,10 @@ struct FLineBuilder
 	 * edge is kept parallel to the opposite side. Not doing so would result in the
 	 * UVs shearing apart at the diagonal where the two triangles meet.
 	 */
-	void BuildLineGeometry(const TArray<FVector2f>& Points, const TArray<FColor>& PackedColors, const FColor& PackedTint, ESlateVertexRounding Rounding) const
+	void BuildLineGeometry(const TArray<FVector2f>& Points, const TArray<FColor>& PackedColors, const FColor& PackedTint, ESlateVertexRounding Rounding)
 	{
+		const bool bIsDashed = DashLength != SolidLineDashLength;
+
 		FColor PointColor = PackedColors.Num() ? PackedColors[0] : PackedTint;
 		FVector2f Position = Points[0];
 		FVector2f NextPosition = Points[1];
@@ -2269,9 +2279,16 @@ struct FLineBuilder
 		(NextPosition - Position).ToDirectionAndLength(Direction, Length);
 		FVector2f Up = GetRotated90(Direction) * LocalHalfThickness;
 
+		// Start the position along the line with the user-provided dash offset.
+		//    This allows users to 'anchor' the dashes in screen space based on some external virtual position
+		PositionAlongLine = DashOffset;
+
 		// Build the start cap at the first point
 		MakeStartCap(Position, Direction, Length, Up, PointColor, Rounding);
 
+		// Increase the position for the next point
+		PositionAlongLine += Length;
+
 		// @TODO: Since the vertex - index relationship is not homogenous whenever 
 		// a miter angle break occurs we cannot pre-allocate indicies here
 		// However if we were to create a new renderbatch on this break then it would be possible
@@ -2281,6 +2298,7 @@ struct FLineBuilder
 		for (int32 Point = 1; Point < LastPointIndex; ++Point)
 		{
 			const FVector2f LastDirection = Direction;
+			const FVector2f LastPosition = Position;
 			const FVector2f LastUp = Up;
 			const float LastLength = Length;
 
@@ -2310,9 +2328,16 @@ struct FLineBuilder
 				const float ParallelDistance = DistanceToMiterLine / DirDotMiterNormal;
 				const FVector2f MiterUp = Up - (Direction * ParallelDistance);
 
-				RenderBatch.EmplaceVertex(FSlateVertex::Make(RenderTransform, FVector2f(Position + MiterUp), FVector2f(1.0f, 0.0f), PointColor, {}, Rounding));
-				RenderBatch.EmplaceVertex(FSlateVertex::Make(RenderTransform, FVector2f(Position - MiterUp), FVector2f(-1.0f, 0.0f), PointColor, {}, Rounding));
+				const float MiterOffset = FVector2f::DotProduct(Direction, MiterUp);
+				RenderBatch.EmplaceVertex(FSlateVertex::Make(RenderTransform, FVector2f(Position + MiterUp), FVector2f(1.0f, 0.0f), FVector2f(PositionAlongLine-MiterOffset, DashLength), PointColor, {}, Rounding));
+				RenderBatch.EmplaceVertex(FSlateVertex::Make(RenderTransform, FVector2f(Position - MiterUp), FVector2f(-1.0f, 0.0f), FVector2f(PositionAlongLine+MiterOffset, DashLength), PointColor, {}, Rounding));
 				AddQuadIndices(RenderBatch);
+
+				if (bIsDashed)
+				{
+					RenderBatch.EmplaceVertex(FSlateVertex::Make(RenderTransform, FVector2f(Position + MiterUp), FVector2f(1.0f, 0.0f), FVector2f(PositionAlongLine+MiterOffset, DashLength), PointColor, {}, Rounding));
+					RenderBatch.EmplaceVertex(FSlateVertex::Make(RenderTransform, FVector2f(Position - MiterUp), FVector2f(-1.0f, 0.0f), FVector2f(PositionAlongLine-MiterOffset, DashLength), PointColor, {}, Rounding));
+				}
 			}
 			else
 			{
@@ -2320,6 +2345,8 @@ struct FLineBuilder
 				MakeEndCap(Position, LastDirection, LastLength, LastUp, PointColor, Rounding);
 				MakeStartCap(Position, Direction, Length, Up, PointColor, Rounding);
 			}
+
+			PositionAlongLine += Length;
 		}
 
 		// Build the last point's incoming segment and end cap
@@ -2398,13 +2425,14 @@ private:
 			// Center the cap over Position if possible, but never place
 			// vertices on the far side of this segment's midpoint
 			const float InwardDistance = FMath::Min(LocalFilterRadius, SegmentLength * 0.5f);
+			const float OutwardDistance = (InwardDistance - LocalCapLength);
 			const FVector2f CapInward = Direction * InwardDistance;
-			const FVector2f CapOutward = Direction * (InwardDistance - LocalCapLength);
+			const FVector2f CapOutward = Direction * OutwardDistance;
 
-			RenderBatch.EmplaceVertex(FSlateVertex::Make(RenderTransform, FVector2f(Position + CapOutward + Up), FVector2f(1.0f, -1.0f), Color, {}, Rounding));
-			RenderBatch.EmplaceVertex(FSlateVertex::Make(RenderTransform, FVector2f(Position + CapOutward - Up), FVector2f(-1.0f, -1.0f), Color, {}, Rounding));
-			RenderBatch.EmplaceVertex(FSlateVertex::Make(RenderTransform, FVector2f(Position + CapInward + Up), FVector2f(1.0f, 0.0f), Color, {}, Rounding));
-			RenderBatch.EmplaceVertex(FSlateVertex::Make(RenderTransform, FVector2f(Position + CapInward - Up), FVector2f(-1.0f, 0.0f), Color, {}, Rounding));
+			RenderBatch.EmplaceVertex(FSlateVertex::Make(RenderTransform, FVector2f(Position + CapOutward + Up), FVector2f(1.0f, -1.0f), FVector2f(PositionAlongLine + OutwardDistance, DashLength), Color, {}, Rounding));
+			RenderBatch.EmplaceVertex(FSlateVertex::Make(RenderTransform, FVector2f(Position + CapOutward - Up), FVector2f(-1.0f, -1.0f), FVector2f(PositionAlongLine + OutwardDistance, DashLength), Color, {}, Rounding));
+			RenderBatch.EmplaceVertex(FSlateVertex::Make(RenderTransform, FVector2f(Position + CapInward + Up), FVector2f(1.0f, 0.0f), FVector2f(PositionAlongLine + InwardDistance, DashLength), Color, {}, Rounding));
+			RenderBatch.EmplaceVertex(FSlateVertex::Make(RenderTransform, FVector2f(Position + CapInward - Up), FVector2f(-1.0f, 0.0f), FVector2f(PositionAlongLine + InwardDistance, DashLength), Color, {}, Rounding));
 			AddQuadIndices(RenderBatch);
 		}
 	}
@@ -2428,15 +2456,16 @@ private:
 			// Center the cap over Position if possible, but never place
 			// vertices on the far side of this segment's midpoint
 			const float InwardDistance = FMath::Min(LocalFilterRadius, SegmentLength * 0.5f);
+			const float OutwardDistance = (InwardDistance - LocalCapLength);
 			const FVector2f CapInward = Direction * -InwardDistance;
-			const FVector2f CapOutward = Direction * (LocalCapLength - InwardDistance);
+			const FVector2f CapOutward = Direction * OutwardDistance;
 
-			RenderBatch.EmplaceVertex(FSlateVertex::Make(RenderTransform, FVector2f(Position + CapInward + Up), FVector2f(1.0f, 0.0f), Color, {}, Rounding));
-			RenderBatch.EmplaceVertex(FSlateVertex::Make(RenderTransform, FVector2f(Position + CapInward - Up), FVector2f(-1.0f, 0.0f), Color, {}, Rounding));
+			RenderBatch.EmplaceVertex(FSlateVertex::Make(RenderTransform, FVector2f(Position + CapInward + Up), FVector2f(1.0f, 0.0f), FVector2f(PositionAlongLine - InwardDistance, DashLength), Color, {}, Rounding));
+			RenderBatch.EmplaceVertex(FSlateVertex::Make(RenderTransform, FVector2f(Position + CapInward - Up), FVector2f(-1.0f, 0.0f), FVector2f(PositionAlongLine - InwardDistance, DashLength), Color, {}, Rounding));
 			AddQuadIndices(RenderBatch);
 
-			RenderBatch.EmplaceVertex(FSlateVertex::Make(RenderTransform, FVector2f(Position + CapOutward + Up), FVector2f(1.0f, 1.0f), Color, {}, Rounding));
-			RenderBatch.EmplaceVertex(FSlateVertex::Make(RenderTransform, FVector2f(Position + CapOutward - Up), FVector2f(-1.0f, 1.0f), Color, {}, Rounding));
+			RenderBatch.EmplaceVertex(FSlateVertex::Make(RenderTransform, FVector2f(Position + CapOutward + Up), FVector2f(1.0f, 1.0f), FVector2f(PositionAlongLine + OutwardDistance, DashLength), Color, {}, Rounding));
+			RenderBatch.EmplaceVertex(FSlateVertex::Make(RenderTransform, FVector2f(Position + CapOutward - Up), FVector2f(-1.0f, 1.0f), FVector2f(PositionAlongLine + OutwardDistance, DashLength), Color, {}, Rounding));
 			AddQuadIndices(RenderBatch);
 		}
 	}
@@ -2462,6 +2491,10 @@ private:
 		return GetRotated90(InboundSegmentDir);
 	}
 
+	// A dashlength of -10 forces results in a computation of 1.f for the entire line in the pixel shader
+	// by feeding it an 'inside-out' trangle wave that modulates between +10 and +30
+	static constexpr float SolidLineDashLength = -10.f;
+
 	FSlateRenderBatch& RenderBatch;
 	const FSlateRenderTransform& RenderTransform;
 
@@ -2469,6 +2502,9 @@ private:
 	const float LocalFilterRadius;
 	const float LocalCapLength;
 	const float AngleCosineLimit;
+	float DashLength = SolidLineDashLength; // Default to no dash
+	float DashOffset = 0.f;
+	float PositionAlongLine = 0.f;
 };
 
 
@@ -2517,7 +2553,9 @@ void FSlateElementBatcher::AddLineElements( const FSlateDrawElementArray<FSlateL
 				DrawElement.GetScale(),
 				HalfThickness,
 				FilterRadius,
-				MiterAngleLimit);
+				MiterAngleLimit,
+				DrawElement.DashLength,
+				DrawElement.DashOffset);
 
 			LineBuilder.BuildLineGeometry(Points, PackedColors, PackedTint, Rounding);
 		}
@@ -2655,7 +2693,9 @@ void FSlateElementBatcher::AddLineElements( const FSlateDrawElementArray<FSlateL
 					DrawElement.GetScale(),
 					HalfThickness,
 					FilterRadius,
-					MiterAngleLimit);
+					MiterAngleLimit,
+					DrawElement.DashLength,
+					DrawElement.DashOffset);
 
 				LineBuilder.NumElements(Points, NumVertexes, NumIndices);
 			}
diff --git a/Engine/Source/Runtime/SlateCore/Public/Rendering/DrawElementTypes.h b/Engine/Source/Runtime/SlateCore/Public/Rendering/DrawElementTypes.h
--- a/Engine/Source/Runtime/SlateCore/Public/Rendering/DrawElementTypes.h
+++ b/Engine/Source/Runtime/SlateCore/Public/Rendering/DrawElementTypes.h
@@ -207,6 +207,21 @@ public:
 	SLATECORE_API static void MakeLines(FSlateWindowElementList& ElementList, uint32 InLayer, const FPaintGeometry& PaintGeometry, TArray<FVector2f> Points, ESlateDrawEffect InDrawEffects = ESlateDrawEffect::None, const FLinearColor& InTint = FLinearColor::White, bool bAntialias = true, float Thickness = 1.0f);
 
 
+	/**
+	 * Creates a dashed or dotted line defined by the provided points. Such lines are always renderd using the anti-aliased shader method.
+	 *
+	 * @param ElementList              The list in which to add elements
+	 * @param InLayer                  The layer to draw the element on
+	 * @param PaintGeometry            DrawSpace position and dimensions; see FPaintGeometry
+	 * @param Points                   Points that make up the lines.  The points are joined together. I.E if Points has A,B,C there the line is A-B-C.  To draw non-joining line segments call MakeLines multiple times
+	 * @param InDrawEffects            Optional draw effects to apply
+	 * @param InTint                   Color to tint the element
+	 * @param Thickness                The thickness of the line
+	 * @param DashLengthPx             The screen-space length of each dash (and each gap - irregular spacings are not supported)
+	 * @param DashScreenOffset         A screen space offset that can be used to 'anchor' dashes to prevent them sliding when moving lines drawn in a virtual space.
+	 */
+	SLATECORE_API static void MakeDashedLines(FSlateWindowElementList& ElementList, uint32 InLayer, const FPaintGeometry& PaintGeometry, TArray<FVector2f>&& Points, ESlateDrawEffect InDrawEffects = ESlateDrawEffect::None, const FLinearColor& InTint = FLinearColor::White, float Thickness = 1.0f, float DashLengthPx = 10.0f, float DashScreenOffset = 0.f);
+
 	/**
 	 * Creates a line defined by the provided points
 	 *
@@ -225,6 +240,22 @@ public:
 #endif
 	SLATECORE_API static void MakeLines(FSlateWindowElementList& ElementList, uint32 InLayer, const FPaintGeometry& PaintGeometry, TArray<FVector2f> Points, TArray<FLinearColor> PointColors, ESlateDrawEffect InDrawEffects = ESlateDrawEffect::None, const FLinearColor& InTint=FLinearColor::White, bool bAntialias = true, float Thickness = 1.0f);
 
+	/**
+	 * Creates a dashed or dotted line defined by the provided points. Such lines are always renderd using the anti-aliased shader method.
+	 *
+	 * @param ElementList              The list in which to add elements
+	 * @param InLayer                  The layer to draw the element on
+	 * @param PaintGeometry            DrawSpace position and dimensions; see FPaintGeometry
+	 * @param Points                   Points that make up the lines.  The points are joined together. I.E if Points has A,B,C there the line is A-B-C.  To draw non-joining line segments call MakeLines multiple times
+	 * @param PointColors              Vertex Color for each defined points
+	 * @param InDrawEffects            Optional draw effects to apply
+	 * @param InTint                   Color to tint the element
+	 * @param Thickness                The thickness of the line
+	 * @param DashLengthPx             The screen-space length of each dash (and each gap - irregular spacings are not supported)
+	 * @param DashScreenOffset         A screen space offset that can be used to 'anchor' dashes to prevent them sliding when moving lines drawn in a virtual space.
+	 */
+	SLATECORE_API static void MakeDashedLines(FSlateWindowElementList& ElementList, uint32 InLayer, const FPaintGeometry& PaintGeometry, TArray<FVector2f>&& Points, TArray<FLinearColor>&& PointColors, ESlateDrawEffect InDrawEffects = ESlateDrawEffect::None, const FLinearColor& InTint = FLinearColor::White, float Thickness = 1.0f, float DashLengthPx = 10.0f, float DashScreenOffset = 0.f);
+
 	/**
 	 * Creates a viewport element which is useful for rendering custom data in a texture into Slate
 	 *
@@ -563,6 +594,8 @@ struct FSlateLineElement : public FSlateDrawElement, public FSlateTintableElemen
 	TArray<FVector2f> Points;
 	TArray<FLinearColor> PointColors;
 	float Thickness;
+	float DashLength = 0.f;
+	float DashOffset = 0.f;
 
 	bool bAntialias;
 
