Slate: Fix typo in documentation.

diff --git a/Engine/Source/Runtime/SlateCore/Public/Types/SlateAttribute.h b/Engine/Source/Runtime/SlateCore/Public/Types/SlateAttribute.h
--- a/Engine/Source/Runtime/SlateCore/Public/Types/SlateAttribute.h
+++ b/Engine/Source/Runtime/SlateCore/Public/Types/SlateAttribute.h
@@ -122,7 +122,7 @@
  *
  *		//bEnabled invalidate paint, we need it to invalidate the Layout.
  *		AttributeInitializer.OverrideInvalidationReason("bEnabled",
-			FSlateAttributeDescriptor::FInvalidateWidgetReasonAttribute{EWidgetInvalidationReason::Layout | EWidgetInvalidationReason::Paint});
+			FSlateAttributeDescriptor::FInvalidateWidgetReasonAttribute{EInvalidateWidgetReason::Layout | EInvalidateWidgetReason::Paint});
  *	}
  *
  *
