DrawElementTypes: Added a new function MakeRotatedShapedText that allows the text to be rotated around a given point

Original review: 34813945

diff --git a/Engine/Source/Runtime/SlateCore/Private/Rendering/DrawElementTypes.cpp b/Engine/Source/Runtime/SlateCore/Private/Rendering/DrawElementTypes.cpp
--- a/Engine/Source/Runtime/SlateCore/Private/Rendering/DrawElementTypes.cpp
+++ b/Engine/Source/Runtime/SlateCore/Private/Rendering/DrawElementTypes.cpp
@@ -505,6 +505,60 @@ void FSlateDrawElement::MakeShapedText(FSlateWindowElementList& ElementList, uin
 	Element.Init(ElementList, EElementType::ET_ShapedText, InLayer, PaintGeometry, InDrawEffects);
 }
 
+void FSlateDrawElement::MakeRotatedShapedText(
+	FSlateWindowElementList& ElementList,
+	const uint32 InLayer,
+	const FPaintGeometry& PaintGeometry,
+	const FShapedGlyphSequenceRef& InShapedGlyphSequence,
+	ESlateDrawEffect InDrawEffects,
+	const FLinearColor& BaseTint,
+	const FLinearColor& OutlineTint,
+	const float Angle2D,
+	const UE::Slate::FDeprecateOptionalVector2DParameter& InRotationPoint,
+	const ERotationSpace RotationSpace,
+	const FTextOverflowArgs& TextOverflowArgs)
+{
+	SCOPE_CYCLE_COUNTER(STAT_SlateDrawElementMakeTime)
+	PaintGeometry.CommitTransformsIfUsingLegacyConstructor();
+	
+	if (InShapedGlyphSequence->GetGlyphsToRender().Num() == 0)
+	{
+		return;
+	}
+	
+	if (ShouldCull(ElementList, PaintGeometry))
+	{
+		return;
+	}
+	
+	// Don't do anything if the font would be completely transparent 
+	if ((BaseTint.A == 0 && InShapedGlyphSequence->GetFontOutlineSettings().OutlineSize == 0) || 
+		(BaseTint.A == 0 && OutlineTint.A == 0))
+	{
+		return;
+	}
+	
+#if SLATE_CHECK_UOBJECT_RENDER_RESOURCES
+	SlateDrawElement::CheckInvalidUMaterial(InShapedGlyphSequence->GetFontMaterial(), TEXT("Font Material"));
+	SlateDrawElement::CheckInvalidUMaterial(InShapedGlyphSequence->GetFontOutlineSettings().OutlineMaterial, TEXT("Outline Material"));
+#endif
+	
+	FSlateShapedTextElement& Element = ElementList.AddUninitialized<EElementType::ET_ShapedText>();
+	
+	Element.SetTint(BaseTint);
+	Element.SetShapedText(ElementList, InShapedGlyphSequence, OutlineTint);
+	Element.SetOverflowArgs(TextOverflowArgs);
+	
+	Element.Init(ElementList, EElementType::ET_ShapedText, InLayer, PaintGeometry, InDrawEffects);
+	
+	if (Angle2D != 0.0f)
+	{
+		const FVector2f RotationPoint = GetRotationPoint(PaintGeometry, InRotationPoint, RotationSpace);
+		const FSlateRenderTransform RotationTransform = Concatenate(Inverse(RotationPoint), FQuat2f(Angle2D), RotationPoint);
+		Element.SetRenderTransform(Concatenate(RotationTransform, Element.GetRenderTransform()));
+	}
+}
+
 void FSlateDrawElement::MakeGradient( FSlateWindowElementList& ElementList, uint32 InLayer, const FPaintGeometry& PaintGeometry, TArray<FSlateGradientStop> InGradientStops, EOrientation InGradientType, ESlateDrawEffect InDrawEffects, FVector4f CornerRadius)
 {
 	PaintGeometry.CommitTransformsIfUsingLegacyConstructor();
diff --git a/Engine/Source/Runtime/SlateCore/Public/Rendering/DrawElementTypes.h b/Engine/Source/Runtime/SlateCore/Public/Rendering/DrawElementTypes.h
--- a/Engine/Source/Runtime/SlateCore/Public/Rendering/DrawElementTypes.h
+++ b/Engine/Source/Runtime/SlateCore/Public/Rendering/DrawElementTypes.h
@@ -143,6 +143,20 @@ public:
 	 */
 	SLATECORE_API static void MakeShapedText( FSlateWindowElementList& ElementList, uint32 InLayer, const FPaintGeometry& PaintGeometry, const FShapedGlyphSequenceRef& InShapedGlyphSequence, ESlateDrawEffect InDrawEffects, const FLinearColor& BaseTint, const FLinearColor& OutlineTint, FTextOverflowArgs TextOverflowArgs = FTextOverflowArgs());
 
+	/**
+	 * Creates a text element which displays a series of shaped glyphs on the screen
+	 *
+	 * @param ElementList			The list in which to add elements
+	 * @param InLayer               The layer to draw the element on
+	 * @param PaintGeometry         DrawSpace position and dimensions; see FPaintGeometry
+	 * @param InShapedGlyphSequence The shaped glyph sequence to draw
+	 * @param InDrawEffects         Optional draw effects to apply
+	 * @param BaseTint              Color to tint the element
+	 * @param OutlineTint           Color to tint the outline
+	 * @param Angle2D               The rotation of the element in radians
+	 */
+	SLATECORE_API static void MakeRotatedShapedText(FSlateWindowElementList& ElementList, uint32 InLayer, const FPaintGeometry& PaintGeometry, const FShapedGlyphSequenceRef& InShapedGlyphSequence, ESlateDrawEffect InDrawEffects, const FLinearColor& BaseTint, const FLinearColor& OutlineTint, float Angle2D, const UE::Slate::FDeprecateOptionalVector2DParameter& InRotationPoint = TOptional<FVector2f>(), ERotationSpace RotationSpace = RelativeToElement, const FTextOverflowArgs& TextOverflowArgs = FTextOverflowArgs());
+	
 	/**
 	 * Creates a gradient element
 	 *
