Shrink FSlateBrush 16bytes

#rnx

diff --git a/Engine/Source/Runtime/SlateCore/Private/Styling/SlateBrush.cpp b/Engine/Source/Runtime/SlateCore/Private/Styling/SlateBrush.cpp
--- a/Engine/Source/Runtime/SlateCore/Private/Styling/SlateBrush.cpp
+++ b/Engine/Source/Runtime/SlateCore/Private/Styling/SlateBrush.cpp
@@ -8,9 +8,7 @@
 #include UE_INLINE_GENERATED_CPP_BY_NAME(SlateBrush)
 
 FSlateBrush::FSlateBrush()
-	: bIsDynamicallyLoaded(false)
-	, bHasUObject_DEPRECATED(false)
-	, bIsSet(true)
+	: TintColor(FLinearColor::White)
 	, DrawAs(ESlateBrushDrawType::Image)
 	, Tiling(ESlateBrushTileType::NoTile)
 	, Mirroring(ESlateBrushMirrorType::NoMirror)
@@ -20,10 +18,12 @@ FSlateBrush::FSlateBrush()
 #if WITH_EDITORONLY_DATA
 	, Tint_DEPRECATED(FLinearColor::White)
 #endif
-	, TintColor(FLinearColor::White)
 	, ResourceObject(nullptr)
-	, ResourceName(NAME_None)
 	, UVRegion(ForceInit)
+	, bIsDynamicallyLoaded(false)
+	, bHasUObject_DEPRECATED(false)
+	, bIsSet(true)
+	, ResourceName(NAME_None)
 {
 }
 
@@ -38,8 +38,7 @@ FSlateBrush::FSlateBrush( ESlateBrushDrawType::Type InDrawType,
 						  bool bInDynamicallyLoaded
 						)
 
-	: bIsDynamicallyLoaded(bInDynamicallyLoaded)
-	, bIsSet(true)
+	: TintColor( InTint )
 	, DrawAs(InDrawType)
 	, Tiling(InTiling)
 	, Mirroring(ESlateBrushMirrorType::NoMirror)
@@ -49,10 +48,11 @@ FSlateBrush::FSlateBrush( ESlateBrushDrawType::Type InDrawType,
 #if WITH_EDITORONLY_DATA
 	, Tint_DEPRECATED(FLinearColor::White)
 #endif
-	, TintColor( InTint )
 	, ResourceObject( InObjectResource )
-	, ResourceName( InResourceName )
 	, UVRegion( ForceInit )
+	, bIsDynamicallyLoaded(bInDynamicallyLoaded)
+	, bIsSet(true)
+	, ResourceName( InResourceName )
 {
 	bHasUObject_DEPRECATED = (InObjectResource != nullptr) || InResourceName.ToString().StartsWith(FSlateBrush::UTextureIdentifier());
 
@@ -74,8 +74,7 @@ FSlateBrush::FSlateBrush( ESlateBrushDrawType::Type InDrawType,
  						  bool bInDynamicallyLoaded
  						)
 
-	: bIsDynamicallyLoaded(bInDynamicallyLoaded)
-	, bIsSet(true)
+	: TintColor( InTint )
 	, DrawAs(InDrawType)
 	, Tiling(InTiling)
 	, Mirroring(ESlateBrushMirrorType::NoMirror)
@@ -85,10 +84,11 @@ FSlateBrush::FSlateBrush( ESlateBrushDrawType::Type InDrawType,
 #if WITH_EDITORONLY_DATA
 	, Tint_DEPRECATED(FLinearColor::White)
 #endif
-	, TintColor( InTint )
 	, ResourceObject( InObjectResource )
-	, ResourceName( InResourceName )
 	, UVRegion( ForceInit )
+	, bIsDynamicallyLoaded(bInDynamicallyLoaded)
+	, bIsSet(true)
+	, ResourceName( InResourceName )
 {
 	bHasUObject_DEPRECATED = (InObjectResource != nullptr) || InResourceName.ToString().StartsWith(FSlateBrush::UTextureIdentifier());
 
@@ -109,8 +109,7 @@ FSlateBrush::FSlateBrush( ESlateBrushDrawType::Type InDrawType,
 						  bool bInDynamicallyLoaded
  						)
 
-	: bIsDynamicallyLoaded(bInDynamicallyLoaded)
-	, bIsSet(true)
+	: TintColor(InTint)
 	, DrawAs(InDrawType)
 	, Tiling(InTiling)
 	, Mirroring(ESlateBrushMirrorType::NoMirror)
@@ -120,10 +119,11 @@ FSlateBrush::FSlateBrush( ESlateBrushDrawType::Type InDrawType,
 #if WITH_EDITORONLY_DATA
 	, Tint_DEPRECATED(FLinearColor::White)
 #endif
-	, TintColor(InTint)
 	, ResourceObject(InObjectResource)
-	, ResourceName(InResourceName)
 	, UVRegion(ForceInit)
+	, bIsDynamicallyLoaded(bInDynamicallyLoaded)
+	, bIsSet(true)
+	, ResourceName(InResourceName)
 {
 	bHasUObject_DEPRECATED = (InObjectResource != nullptr) || InResourceName.ToString().StartsWith(FSlateBrush::UTextureIdentifier());
 
diff --git a/Engine/Source/Runtime/SlateCore/Public/Styling/SlateBrush.h b/Engine/Source/Runtime/SlateCore/Public/Styling/SlateBrush.h
--- a/Engine/Source/Runtime/SlateCore/Public/Styling/SlateBrush.h
+++ b/Engine/Source/Runtime/SlateCore/Public/Styling/SlateBrush.h
@@ -225,29 +225,29 @@ namespace SlateBrushDefs
 	static const float DefaultImageSize = 32.0f;
 }
 
+// Hacky base class to avoid 8 bytes of padding after the vtable
+struct FSlateBrushFixLayout
+{
+	virtual ~FSlateBrushFixLayout() = default;
+};
+
 /**
- * A brush which contains information about how to draw a Slate element
- */
+* A brush which contains information about how to draw a Slate element
+*/
 USTRUCT(BlueprintType) //, meta = (HasNativeMake = ""))
-struct FSlateBrush
+struct FSlateBrush 
+#if CPP 
+: public FSlateBrushFixLayout
+#endif
 {
 	GENERATED_USTRUCT_BODY()
 
 	friend class FSlateShaderResourceManager;
 
-protected:
-	/** Whether or not the brush path is a path to a UObject */
-	UPROPERTY()
-	uint8 bIsDynamicallyLoaded:1;
-
-	/** Whether or not the brush has a UTexture resource */
-	UPROPERTY()
-	uint8 bHasUObject_DEPRECATED:1;
-
-	/** This is true for all constructed brushes except for optional brushes */
-	uint8 bIsSet : 1;
-
 public:
+	/** Tinting applied to the image. */
+	UPROPERTY(EditAnywhere, BlueprintReadWrite, Category=Brush, meta=( DisplayName="Tint", sRGB="true" ))
+	FSlateColor TintColor;
 
 	/** How to draw the image */
 	UPROPERTY(EditAnywhere, BlueprintReadWrite, Category=Brush)
@@ -280,15 +280,6 @@ public:
 	FLinearColor Tint_DEPRECATED;
 #endif
 
-	/** Tinting applied to the image. */
-	UPROPERTY(EditAnywhere, BlueprintReadWrite, Category=Brush, meta=( DisplayName="Tint", sRGB="true" ))
-	FSlateColor TintColor;
-
-public:
-	/** How to draw the outline.  Currently only used for RoundedBox type brushes. */
-	UPROPERTY(EditAnywhere, BlueprintReadWrite, Category=Brush)
-	FSlateBrushOutlineSettings OutlineSettings;
-
 public:
 
 	/**
@@ -485,11 +476,12 @@ private:
 	UPROPERTY(EditAnywhere, BlueprintReadWrite, Category=Brush, meta=( AllowPrivateAccess="true", DisplayThumbnail="true", DisplayName="Image", AllowedClasses="/Script/Engine.Texture,/Script/Engine.MaterialInterface,/Script/Engine.SlateTextureAtlasInterface", DisallowedClasses = "/Script/MediaAssets.MediaTexture"))
 	TObjectPtr<UObject> ResourceObject;
 
-protected:
-	/** The name of the rendering resource to use */
-	UPROPERTY(EditAnywhere, AdvancedDisplay, Category = Brush)
-	FName ResourceName;
+public:
+	/** How to draw the outline.  Currently only used for RoundedBox type brushes. */
+	UPROPERTY(EditAnywhere, BlueprintReadWrite, Category=Brush)
+	FSlateBrushOutlineSettings OutlineSettings;
 
+protected:
 	/** 
 	 *  Optional UV region for an image
 	 *  When valid - overrides UV region specified in resource proxy
@@ -497,11 +489,26 @@ protected:
 	UPROPERTY()
 	FBox2f UVRegion;
 
+	/** Whether or not the brush path is a path to a UObject */
+	UPROPERTY()
+	uint8 bIsDynamicallyLoaded:1;
+
+	/** Whether or not the brush has a UTexture resource */
+	UPROPERTY()
+	uint8 bHasUObject_DEPRECATED:1;
+
+	/** This is true for all constructed brushes except for optional brushes */
+	uint8 bIsSet : 1;
+
 public:
 
 	/** Rendering resource for this brush */
 	mutable FSlateResourceHandle ResourceHandle;
+
 protected:
+	/** The name of the rendering resource to use */
+	UPROPERTY(EditAnywhere, AdvancedDisplay, Category = Brush)
+	FName ResourceName;
 
 	/** 
 	 * This constructor is protected; use one of the deriving classes instead.
