Slate SDF fonts - moved from Font to Font Face and per device profile configuration

#tests Tested in editor
#rnx <Release Notes exclude>

diff --git a/Engine/Source/Editor/FontEditor/Private/FontEditor.cpp b/Engine/Source/Editor/FontEditor/Private/FontEditor.cpp
--- a/Engine/Source/Editor/FontEditor/Private/FontEditor.cpp
+++ b/Engine/Source/Editor/FontEditor/Private/FontEditor.cpp
@@ -465,7 +465,6 @@ void FFontEditor::PostUndo(bool bSuccess)
 void FFontEditor::NotifyPostChange( const FPropertyChangedEvent& PropertyChangedEvent, class FEditPropertyChain* PropertyThatChanged)
 {
 	static const FName FontCacheTypePropertyName = GET_MEMBER_NAME_CHECKED(UFont, FontCacheType);
-	static const FName FontRasterizationModePropertyName = GET_MEMBER_NAME_CHECKED(UFont, FontRasterizationMode);
 	static const FName CompositeFontPropertyName = GET_MEMBER_NAME_CHECKED(UFont, CompositeFont);
 	static const FName TexturePageWidthName = GET_MEMBER_NAME_CHECKED(FFontImportOptionsData, TexturePageWidth);
 	static const FName TexturePageMaxHeightName = GET_MEMBER_NAME_CHECKED(FFontImportOptionsData, TexturePageMaxHeight);
@@ -513,13 +512,6 @@ void FFontEditor::NotifyPostChange( const FPropertyChangedEvent& PropertyChanged
 		}
 	}
 
-	if (PropertyChangedEvent.Property && PropertyChangedEvent.Property->GetFName() == FontRasterizationModePropertyName)
-	{
-		// Show / hide SdfFont category of properties
-		UpdateLayout();
-		FontProperties->ForceRefresh();
-	}
-
 	if (PropertyChangedEvent.Property && PropertyChangedEvent.Property->GetFName() == DistanceFieldScaleFactorName)
 	{
 		const uint32 SignedInt32NumBits = 31;
@@ -1153,22 +1145,6 @@ bool FFontEditor::GetIsPropertyVisible(const FPropertyAndParent& PropertyAndPare
 	static const FName CategoryFName = "Category";
 	const FString& CategoryValue = PropertyAndParent.Property.GetMetaData(CategoryFName);
 
-	// If SDF feature is disabled, hide all SDF-related settings
-	if (!IsSlateSdfTextFeatureEnabled())
-	{
-		if (PropertyAndParent.Property.GetFName() == GET_MEMBER_NAME_CHECKED(UFont, FontRasterizationMode) ||
-			PropertyAndParent.Property.GetFName() == GET_MEMBER_NAME_CHECKED(UFont, SdfSettings))
-		{
-			return false;
-		}
-	}
-
-	// Hide SDF settings if font rasterization mode is not MSDF
-	if (!Font->IsSdfFont() && PropertyAndParent.Property.GetFName() == GET_MEMBER_NAME_CHECKED(UFont, SdfSettings))
-	{
-		return false;
-	}
-
 	// We need to hide the properties associated with the category that we're not currently using (either Offline or Runtime)
 	const FString CategoryToExclude = (Font->FontCacheType == EFontCacheType::Offline) ? TEXT("RuntimeFont") : TEXT("OfflineFont");
 
diff --git a/Engine/Source/Editor/FontEditor/Private/FontFaceDetailsCustomization.cpp b/Engine/Source/Editor/FontEditor/Private/FontFaceDetailsCustomization.cpp
--- a/Engine/Source/Editor/FontEditor/Private/FontFaceDetailsCustomization.cpp
+++ b/Engine/Source/Editor/FontEditor/Private/FontFaceDetailsCustomization.cpp
@@ -10,6 +10,7 @@
 #include "EditorFontGlyphs.h"
 #include "Engine/FontFace.h"
 #include "Fonts/CompositeFont.h"
+#include "Fonts/FontCache.h"
 #include "Fonts/SlateFontInfo.h"
 #include "Framework/Application/SlateApplication.h"
 #include "HAL/Platform.h"
@@ -92,6 +93,16 @@ void FFontFaceDetailsCustomization::CustomizeDetails(IDetailLayoutBuilder& Detai
 				]
 			];
 	}
+
+	// Hide distance field settings if disabled by CVar
+	if (!IsSlateSdfTextFeatureEnabled())
+	{
+		DetailBuilder.GetProperty(GET_MEMBER_NAME_CHECKED(UFontFace, bEnableDistanceFieldRendering))->MarkHiddenByCustomization();
+		DetailBuilder.GetProperty(GET_MEMBER_NAME_CHECKED(UFontFace, MinDistanceFieldPpem))->MarkHiddenByCustomization();
+		DetailBuilder.GetProperty(GET_MEMBER_NAME_CHECKED(UFontFace, MidDistanceFieldPpem))->MarkHiddenByCustomization();
+		DetailBuilder.GetProperty(GET_MEMBER_NAME_CHECKED(UFontFace, MaxDistanceFieldPpem))->MarkHiddenByCustomization();
+		DetailBuilder.GetProperty(GET_MEMBER_NAME_CHECKED(UFontFace, PlatformRasterizationModeOverrides))->MarkHiddenByCustomization();
+	}
 }
 
 FText FFontFaceDetailsCustomization::GetFontDisplayName() const
diff --git a/Engine/Source/Runtime/Engine/Classes/Engine/Font.h b/Engine/Source/Runtime/Engine/Classes/Engine/Font.h
--- a/Engine/Source/Runtime/Engine/Classes/Engine/Font.h
+++ b/Engine/Source/Runtime/Engine/Classes/Engine/Font.h
@@ -106,14 +106,6 @@ class UFont : public UObject, public IFontProviderInterface
 	UPROPERTY(EditAnywhere, Category=Font)
 	EFontCacheType FontCacheType;
 
-	/** The preferred rasterization method for this font (enable / disable MSDF) */
-	UPROPERTY(EditAnywhere, Category=RuntimeFont)
-	EFontRasterizationMode FontRasterizationMode = EFontRasterizationMode::Bitmap;
-
-	/** Settings for rendering this font using the sdf pipeline */
-	UPROPERTY(EditAnywhere, Category=RuntimeFont, meta = (DisplayName = "SDF Settings"))
-	FFontSdfSettings SdfSettings;
-
 	/** List of characters in the font.  For a MultiFont, this will include all characters in all sub-fonts!  Thus,
 		the number of characters in this array isn't necessary the number of characters available in the font */
 	UPROPERTY(EditAnywhere, Category=OfflineFont)
@@ -287,24 +279,6 @@ public:
 		return ScalingFactor;
 	}
 
-	/**
-	 *	Returns true if rasterization mode is signed distance field-based (and the feature is enabled)
-	 */
-	virtual bool IsSdfFont() const override;
-
-	/**
-	 *	Get the font rasterization mode (IFontProviderInterface)
-	 */
-	virtual EFontRasterizationMode GetFontRasterizationMode() const override;
-
-	/**
-	 *	Get the font SDF settings (IFontProviderInterface)
-	 */
-	virtual const FFontSdfSettings& GetSdfSettings() const override
-	{
-		return SdfSettings;
-	}
-
 	/** Returns the maximum height for any character in this font using this font's default size and scale. */
 	ENGINE_API float GetMaxCharHeight() const;
 	
diff --git a/Engine/Source/Runtime/Engine/Classes/Engine/FontFace.h b/Engine/Source/Runtime/Engine/Classes/Engine/FontFace.h
--- a/Engine/Source/Runtime/Engine/Classes/Engine/FontFace.h
+++ b/Engine/Source/Runtime/Engine/Classes/Engine/FontFace.h
@@ -6,11 +6,31 @@
 #include "UObject/ObjectMacros.h"
 #include "UObject/Object.h"
 #include "Fonts/FontFaceInterface.h"
+#include "Fonts/FontRasterizationMode.h"
 #include "FontFace.generated.h"
 
 class ITargetPlatform;
 struct FPropertyChangedEvent;
 
+/** Remapping of rasterization modes */
+USTRUCT(BlueprintType)
+struct FFontFacePlatformRasterizationOverrides
+{
+	GENERATED_BODY()
+
+	/** Rasterization mode to be used instead of Sharp (Multi-Channel SDF) */
+	UPROPERTY(EditAnywhere, Category=DistanceFieldMode, meta=(DisplayName="Override for Sharp"))
+	EFontRasterizationMode MsdfOverride = EFontRasterizationMode::Msdf;
+
+	/** Rasterization mode to be used instead of Smooth (Plain SDF) */
+	UPROPERTY(EditAnywhere, Category=DistanceFieldMode, meta=(DisplayName="Override for Smooth"))
+	EFontRasterizationMode SdfOverride = EFontRasterizationMode::Sdf;
+
+	/** Rasterization mode to be used instead of Fast (Approximate SDF) */
+	UPROPERTY(EditAnywhere, Category=DistanceFieldMode, meta=(DisplayName="Override for Fast"))
+	EFontRasterizationMode SdfApproximationOverride = EFontRasterizationMode::SdfApproximation;
+};
+
 /**
  * A font face asset contains the raw payload data for a source TTF/OTF file as used by FreeType.
  * During cook this asset type generates a ".ufont" file containing the raw payload data (unless loaded "Inline").
@@ -27,6 +47,7 @@ public:
 	//~ Begin UObject Interface
 	virtual void Serialize(FArchive& Ar) override;
 	virtual void GetResourceSizeEx(FResourceSizeEx& CumulativeResourceSize) override;
+	virtual void PostLoad() override;
 #if WITH_EDITOR
 	virtual void PostEditChangeProperty(FPropertyChangedEvent& PropertyChangedEvent) override;
 	virtual void PostEditUndo() override;
@@ -54,10 +75,12 @@ public:
 	virtual int32 GetDescendOverriddenValue() const override;
 	virtual int32 GetStrikeBrushHeightPercentage() const override;
 	virtual FFontFaceDataConstRef GetFontFaceData() const override;
+	virtual FFontRasterizationSettings GetRasterizationSettings() const override;
 	//~ End IFontFaceInterface interface
 
 private:
 	FString GetCookedFilename() const;
+	void UpdateDeviceRasterizationSettings();
 	//~ Begin UObject Interface
 #if WITH_EDITOR
 	virtual void CookAdditionalFilesOverride(const TCHAR* PackageFilename, const ITargetPlatform* TargetPlatform,
@@ -116,4 +139,29 @@ public:
 	UPROPERTY(VisibleAnywhere, Transient, Category=FontFace, AdvancedDisplay)
 	TArray<FString> SubFaces;
 #endif // WITH_EDITORONLY_DATA
+
+	/** Enables distance field rendering for this face (otherwise only Bitmap rendering is used) */
+	UPROPERTY(EditAnywhere, Category=DistanceFieldMode)
+	bool bEnableDistanceFieldRendering = false;
+	
+	/** Distance field px/em resolution "low" quality value */
+	UPROPERTY(EditAnywhere, Category=DistanceFieldMode, meta=(ClampMin=8, ClampMax=256, DisplayName="Low SDF resolution (px/em)"))
+	int32 MinDistanceFieldPpem = 32;
+
+	/** Distance field px/em resolution "medium" quality value */
+	UPROPERTY(EditAnywhere, Category=DistanceFieldMode, meta=(ClampMin=8, ClampMax=256, DisplayName="Medium SDF resolution (px/em)"))
+	int32 MidDistanceFieldPpem = 48;
+
+	/** Distance field px/em resolution "high" quality value */
+	UPROPERTY(EditAnywhere, Category=DistanceFieldMode, meta=(ClampMin=8, ClampMax=256, DisplayName="High SDF resolution (px/em)"))
+	int32 MaxDistanceFieldPpem = 64;
+
+	/** If set, allows to override distance field modes set in device profiles */
+	UPROPERTY(EditAnywhere, Category=DistanceFieldMode, meta=(ClampMin=8, ClampMax=256, DisplayName="Override platform rasterization mode"))
+	TOptional<FFontFacePlatformRasterizationOverrides> PlatformRasterizationModeOverrides;
+
+private:
+	/** Cached rasterization settings for the active device profile */
+	FFontRasterizationSettings DeviceRasterizationSettings;
+
 };
diff --git a/Engine/Source/Runtime/Engine/Classes/Engine/UserInterfaceSettings.h b/Engine/Source/Runtime/Engine/Classes/Engine/UserInterfaceSettings.h
--- a/Engine/Source/Runtime/Engine/Classes/Engine/UserInterfaceSettings.h
+++ b/Engine/Source/Runtime/Engine/Classes/Engine/UserInterfaceSettings.h
@@ -266,6 +266,14 @@ private:
 	bool bUseCustomFontDPI;
 #endif
 
+	UFUNCTION()
+	static bool ShouldDisplayDistanceFieldFontSettings();
+
+public:
+	UPROPERTY(config, EditAnywhere, Category = "UMG Fonts", meta = (EditCondition = "ShouldDisplayDistanceFieldFontSettings()", EditConditionHides))
+	bool bEnableDistanceFieldFontRasterization;
+
+private:
 	UPROPERTY(Transient)
 	TArray<TObjectPtr<UObject>> CursorClasses;
 
diff --git a/Engine/Source/Runtime/Engine/Private/Font.cpp b/Engine/Source/Runtime/Engine/Private/Font.cpp
--- a/Engine/Source/Runtime/Engine/Private/Font.cpp
+++ b/Engine/Source/Runtime/Engine/Private/Font.cpp
@@ -388,17 +388,3 @@ void UFont::GetResourceSizeEx(FResourceSizeEx& CumulativeResourceSize)
 		}
 	}
 }
-
-bool UFont::IsSdfFont() const
-{
-	return IsSdfFontRasterizationMode(FontRasterizationMode) && IsSlateSdfTextFeatureEnabled();
-}
-
-EFontRasterizationMode UFont::GetFontRasterizationMode() const
-{
-	if (IsSdfFontRasterizationMode(FontRasterizationMode) && !IsSlateSdfTextFeatureEnabled())
-	{
-		return EFontRasterizationMode::Bitmap;
-	}
-	return FontRasterizationMode;
-}
diff --git a/Engine/Source/Runtime/Engine/Private/FontFace.cpp b/Engine/Source/Runtime/Engine/Private/FontFace.cpp
--- a/Engine/Source/Runtime/Engine/Private/FontFace.cpp
+++ b/Engine/Source/Runtime/Engine/Private/FontFace.cpp
@@ -2,6 +2,7 @@
 
 #include "Engine/FontFace.h"
 #include "Engine/Font.h"
+#include "Engine/UserInterfaceSettings.h"
 #include "EditorFramework/AssetImportData.h"
 #include "Misc/Paths.h"
 #include "Framework/Application/SlateApplication.h"
@@ -17,6 +18,29 @@
 
 DEFINE_LOG_CATEGORY_STATIC(LogFontFace, Log, All);
 
+TAutoConsoleVariable<FString> CVarFontFaceDistanceFieldRasterizationMode(
+	TEXT("UI.SlateSDFText.RasterizationMode"),
+	"Bitmap",
+	TEXT("Sets the rasterization mode of font faces with distance field rasterization enabled. Possible values are: Bitmap, Msdf, Sdf, SdfApproximation."),
+	ECVF_Preview);
+
+TAutoConsoleVariable<int32> CVarFontFaceDistanceFieldResolutionLevel(
+	TEXT("UI.SlateSDFText.ResolutionLevel"),
+	2,
+	TEXT("Sets the resolution level (1 = low, 2 = medium, 3 = high) of font faces with distance field rasterization enabled."),
+	ECVF_Preview);
+
+EFontRasterizationMode GetDeviceFontFaceDistanceFieldRasterizationMode()
+{
+	const int64 EnumVal = StaticEnum<EFontRasterizationMode>()->GetValueByNameString(CVarFontFaceDistanceFieldRasterizationMode.GetValueOnAnyThread());
+	if (EnumVal != INDEX_NONE)
+	{
+		return (EFontRasterizationMode) EnumVal;
+	}
+	UE_LOG(LogFontFace, Warning, TEXT("Unexpected value of CVar UI.SlateSDFText.RasterizationMode - falling back to Bitmap."));
+	return EFontRasterizationMode::Bitmap;
+}
+
 UFontFace::UFontFace()
 	: AscendOverriddenValue(0)
 	, bIsAscendOverridden(false)
@@ -107,11 +131,20 @@ void UFontFace::GetResourceSizeEx(FResourceSizeEx& CumulativeResourceSize)
 
 }
 
+void UFontFace::PostLoad()
+{
+	Super::PostLoad();
+
+	UpdateDeviceRasterizationSettings();
+}
+
 #if WITH_EDITOR
 void UFontFace::PostEditChangeProperty(FPropertyChangedEvent& PropertyChangedEvent)
 {
 	Super::PostEditChangeProperty(PropertyChangedEvent);
 
+	UpdateDeviceRasterizationSettings();
+
 #if WITH_EDITORONLY_DATA
 	CacheSubFaces();
 #endif // WITH_EDITORONLY_DATA
@@ -247,6 +280,11 @@ FFontFaceDataConstRef UFontFace::GetFontFaceData() const
 	return FontFaceData;
 }
 
+FFontRasterizationSettings UFontFace::GetRasterizationSettings() const
+{
+	return DeviceRasterizationSettings;
+}
+
 FString UFontFace::GetCookedFilename() const
 {
 	// UFontFace assets themselves can't be localized, however that doesn't mean the package they're in isn't localized (ie, when they're upgraded into a UFont asset)
@@ -261,3 +299,45 @@ FString UFontFace::GetCookedFilename() const
 	return FPaths::GetPath(PackageFilename) / GetName() + TEXT(".ufont");
 }
 
+void UFontFace::UpdateDeviceRasterizationSettings()
+{
+	DeviceRasterizationSettings = FFontRasterizationSettings();
+	if (bEnableDistanceFieldRendering && GetDefault<UUserInterfaceSettings>()->bEnableDistanceFieldFontRasterization && IsSlateSdfTextFeatureEnabled())
+	{
+		DeviceRasterizationSettings.Mode = GetDeviceFontFaceDistanceFieldRasterizationMode();
+		if (PlatformRasterizationModeOverrides.IsSet())
+		{
+			switch (DeviceRasterizationSettings.Mode)
+			{
+				default:
+					checkNoEntry();
+					break;
+				case EFontRasterizationMode::Bitmap:
+					break;
+				case EFontRasterizationMode::Msdf:
+					DeviceRasterizationSettings.Mode = PlatformRasterizationModeOverrides->MsdfOverride;
+					break;
+				case EFontRasterizationMode::Sdf:
+					DeviceRasterizationSettings.Mode = PlatformRasterizationModeOverrides->SdfOverride;
+					break;
+				case EFontRasterizationMode::SdfApproximation:
+					DeviceRasterizationSettings.Mode = PlatformRasterizationModeOverrides->SdfApproximationOverride;
+					break;
+			}
+		}
+
+		const int32 DeviceResolutionLevel = CVarFontFaceDistanceFieldResolutionLevel.GetValueOnAnyThread();
+		if (DeviceResolutionLevel <= 1) // Low
+		{
+			DeviceRasterizationSettings.DistanceFieldPpem = MinDistanceFieldPpem;
+		}
+		else if (DeviceResolutionLevel >= 3) // High
+		{
+			DeviceRasterizationSettings.DistanceFieldPpem = MaxDistanceFieldPpem;
+		}
+		else // 2 = Medium
+		{
+			DeviceRasterizationSettings.DistanceFieldPpem = MidDistanceFieldPpem;
+		}
+	}
+}
diff --git a/Engine/Source/Runtime/Engine/Private/UserInterfaceSettings.cpp b/Engine/Source/Runtime/Engine/Private/UserInterfaceSettings.cpp
--- a/Engine/Source/Runtime/Engine/Private/UserInterfaceSettings.cpp
+++ b/Engine/Source/Runtime/Engine/Private/UserInterfaceSettings.cpp
@@ -3,6 +3,7 @@
 #include "Engine/UserInterfaceSettings.h"
 
 #include "Logging/MessageLog.h"
+#include "Fonts/FontCache.h"
 
 #include "Engine/DPICustomScalingRule.h"
 
@@ -21,6 +22,7 @@ UUserInterfaceSettings::UUserInterfaceSettings(const FObjectInitializer& ObjectI
 	, FontDPIPreset(ConvertToEFontDPI(CustomFontDPI))
 	, bUseCustomFontDPI(false)
 #endif
+	, bEnableDistanceFieldFontRasterization(false)
 {
 	SectionName = TEXT("UI");
 }
@@ -272,5 +274,9 @@ void UUserInterfaceSettings::ForceLoadResources(bool bForceLoadEverything)
 	}
 }
 
-#undef LOCTEXT_NAMESPACE
+bool UUserInterfaceSettings::ShouldDisplayDistanceFieldFontSettings()
+{
+	return IsSlateSdfTextFeatureEnabled();
+}
 
+#undef LOCTEXT_NAMESPACE
diff --git a/Engine/Source/Runtime/SlateCore/Private/Fonts/CompositeFont.cpp b/Engine/Source/Runtime/SlateCore/Private/Fonts/CompositeFont.cpp
--- a/Engine/Source/Runtime/SlateCore/Private/Fonts/CompositeFont.cpp
+++ b/Engine/Source/Runtime/SlateCore/Private/Fonts/CompositeFont.cpp
@@ -192,6 +192,16 @@ const UObject* FFontData::GetFontFaceAsset() const
 	return FontFaceAsset;
 }
 
+FFontRasterizationSettings FFontData::GetFontRasterizationSettings() const
+{
+	if (FontFaceAsset)
+	{
+		const IFontFaceInterface* FontFace = CastChecked<const IFontFaceInterface>(FontFaceAsset);
+		return FontFace->GetRasterizationSettings();
+	}
+	return FFontRasterizationSettings();
+}
+
 #if WITH_EDITORONLY_DATA
 bool FFontData::HasLegacyData() const
 {
diff --git a/Engine/Source/Runtime/SlateCore/Private/Fonts/FontCache.cpp b/Engine/Source/Runtime/SlateCore/Private/Fonts/FontCache.cpp
--- a/Engine/Source/Runtime/SlateCore/Private/Fonts/FontCache.cpp
+++ b/Engine/Source/Runtime/SlateCore/Private/Fonts/FontCache.cpp
@@ -173,7 +173,7 @@ FSdfGlyphTaskKey::FSdfGlyphTaskKey(const FSdfGlyphEntryKey& InSdfGlyphEntryKey,
 * 
 */
 
-FShapedGlyphSequence::FShapedGlyphSequence(TArray<FShapedGlyphEntry> InGlyphsToRender, const int16 InTextBaseline, const uint16 InMaxTextHeight, const UObject* InFontMaterial, const FFontOutlineSettings& InOutlineSettings, const EFontRasterizationMode InRasterizationMode, const FFontSdfSettings& InSdfSettings, const FSourceTextRange& InSourceTextRange)
+FShapedGlyphSequence::FShapedGlyphSequence(TArray<FShapedGlyphEntry> InGlyphsToRender, const int16 InTextBaseline, const uint16 InMaxTextHeight, const UObject* InFontMaterial, const FFontOutlineSettings& InOutlineSettings, const FSourceTextRange& InSourceTextRange)
 	: GlyphsToRender(MoveTemp(InGlyphsToRender))
 	, TextBaseline(InTextBaseline)
 	, MaxTextHeight(InMaxTextHeight)
@@ -182,8 +182,6 @@ FShapedGlyphSequence::FShapedGlyphSequence(TArray<FShapedGlyphEntry> InGlyphsToR
 	, SequenceWidth(0)
 	, GlyphFontFaces()
 	, SourceIndicesToGlyphData(InSourceTextRange)
-	, RasterizationMode(InRasterizationMode)
-	, SdfSettings(InSdfSettings)
 	, CachedFontSkew(GlyphsToRender.Num() > 0 ? GlyphsToRender[0].FontFaceData->FontSkew : 0.f)
 {
 	const int32 NumGlyphsToRender = GlyphsToRender.Num();
@@ -276,20 +274,6 @@ SIZE_T FShapedGlyphSequence::GetAllocatedSize() const
 	return GlyphsToRender.GetAllocatedSize() + GlyphFontFaces.GetAllocatedSize() + SourceIndicesToGlyphData.GetAllocatedSize();
 }
 
-bool FShapedGlyphSequence::IsSdfFont() const
-{
-	return IsSdfFontRasterizationMode(RasterizationMode) && IsSlateSdfTextFeatureEnabled();
-}
-
-EFontRasterizationMode FShapedGlyphSequence::GetRasterizationMode() const
-{
-	if (IsSdfFontRasterizationMode(RasterizationMode) && !IsSlateSdfTextFeatureEnabled())
-	{
-		return EFontRasterizationMode::Bitmap;
-	}
-	return RasterizationMode;
-}
-
 bool FShapedGlyphSequence::IsDirty() const
 {
 	for (const auto& GlyphFontFace : GlyphFontFaces)
@@ -522,7 +506,7 @@ FShapedGlyphSequencePtr FShapedGlyphSequence::GetSubSequence(const int32 InStart
 
 	if (EnumerateVisualGlyphsInSourceRange(InStartIndex, InEndIndex, GlyphCallback) == EEnumerateGlyphsResult::EnumerationComplete)
 	{
-		return MakeShared<FShapedGlyphSequence>(MoveTemp(SubGlyphsToRender), TextBaseline, MaxTextHeight, FontMaterial, OutlineSettings, RasterizationMode, SdfSettings, SubSequenceRange);
+		return MakeShared<FShapedGlyphSequence>(MoveTemp(SubGlyphsToRender), TextBaseline, MaxTextHeight, FontMaterial, OutlineSettings, SubSequenceRange);
 	}
 
 	return nullptr;
@@ -849,11 +833,14 @@ const FCharacterEntry* FCharacterList::CacheCharacter(TCHAR Character)
 			}
 
 			FShapedGlyphEntry ShapedGlyphEntry;
+			const FFontRasterizationSettings FontRasterizationSettings = FontDataPtr->GetFontRasterizationSettings();
 			ShapedGlyphEntry.FontFaceData = MakeShared<FShapedGlyphFaceData>(FaceGlyphData.FaceAndMemory, 
 																			 GlyphFlags, 
 																			 FontInfo.Size, 
 																			 FinalFontScale, 
-																			 0.f);
+																			 0.f,
+																			 FontRasterizationSettings.Mode,
+																			 FontRasterizationSettings.DistanceFieldPpem);
 			ShapedGlyphEntry.GlyphIndex = GlyphIndex;
 			ShapedGlyphEntry.XAdvance = XAdvance;
 			ShapedGlyphEntry.bIsVisible = !bIsWhitespace;
@@ -1193,25 +1180,25 @@ static float CeilSdfHalfSpread(float InEmSpread, float Ppem)
 	return PixelSpread/Ppem;
 }
 
-FSdfGlyphFontAtlasData FSlateFontCache::GetSdfGlyphFontAtlasData(const FShapedGlyphEntry& InShapedGlyph, const FFontOutlineSettings& InOutlineSettings, EFontRasterizationMode InRasterizationMode, const FFontSdfSettings& InSdfSettings)
+FSdfGlyphFontAtlasData FSlateFontCache::GetSdfGlyphFontAtlasData(const FShapedGlyphEntry& InShapedGlyph, const FFontOutlineSettings& InOutlineSettings)
 {
-	if (!IsSlateSdfTextFeatureEnabled()
-		|| !InShapedGlyph.FontFaceData->bSupportsSdf
-		|| !InShapedGlyph.bIsVisible)
+	if (!InShapedGlyph.bIsVisible
+		|| InShapedGlyph.FontFaceData->RasterizationMode == EFontRasterizationMode::Bitmap
+		|| !IsSlateSdfTextFeatureEnabled())
 	{
 		return { /* EarlyCheckIfCannotUseSdf */ };
 	}
 	const int32 CachedAtlasDataThreadIndex = static_cast<int32>(OwningThread);
 	check(CachedAtlasDataThreadIndex >= 0 && CachedAtlasDataThreadIndex < 2);
 
-	const bool bPlaceholderOnly = InRasterizationMode == EFontRasterizationMode::SdfApproximation;
+	const bool bPlaceholderOnly = InShapedGlyph.FontFaceData->RasterizationMode == EFontRasterizationMode::SdfApproximation;
 #if WITH_FREETYPE
 	const float TargetPpem = FMath::Max(1.f, static_cast<float>(FreeTypeUtils::ComputeFontPixelSize(InShapedGlyph.FontFaceData->FontSize, InShapedGlyph.FontFaceData->FontScale)));
 #else
 	const float TargetPpem = FMath::Max(1.f, 1.f/72.f*static_cast<float>(FontConstants::RenderDPI)*InShapedGlyph.FontFaceData->FontSize*InShapedGlyph.FontFaceData->FontScale);
 #endif
-	const float SdfPpem = static_cast<float>(InSdfSettings.GetClampedPpem());
-	const float Scale = TargetPpem/SdfPpem;
+	const int32 SdfPpem = InShapedGlyph.FontFaceData->SdfPpem;
+	const float Scale = TargetPpem/float(SdfPpem);
 	const float EmOutlineSize = FMath::RoundToFloat(InOutlineSettings.OutlineSize * InShapedGlyph.FontFaceData->FontScale)/TargetPpem;
 
 	float EmInnerSpread = FMath::Clamp(2.f/TargetPpem, 0.05f, 0.25f);
@@ -1220,7 +1207,7 @@ FSdfGlyphFontAtlasData FSlateFontCache::GetSdfGlyphFontAtlasData(const FShapedGl
 	FSlateSdfGenerator::ESdfType SdfType = FSlateSdfGenerator::ESdfType::MultichannelAndSimple;
 	ESlateFontAtlasContentType AtlasType = ESlateFontAtlasContentType::Msdf;
 	TArray<uint8>* FontAtlasIndicesPtr = &MsdfFontAtlasIndices;
-	switch (InRasterizationMode)
+	switch (InShapedGlyph.FontFaceData->RasterizationMode)
 	{
 		case EFontRasterizationMode::Bitmap:
 		default:
@@ -1239,8 +1226,9 @@ FSdfGlyphFontAtlasData FSlateFontCache::GetSdfGlyphFontAtlasData(const FShapedGl
 			break;
 	}
 
-	auto TryRespawn = [&](const FSdfGlyphFontAtlasData& GlyphAtlasData) -> bool {
-		const FSdfGlyphEntryKey GlyphKey(InShapedGlyph.FontFaceData->FontFace, InShapedGlyph.GlyphIndex, AtlasType, InSdfSettings.GetClampedPpem(), GlyphAtlasData.EmOuterSpread, GlyphAtlasData.EmInnerSpread);
+	auto TryRespawn = [&](const FSdfGlyphFontAtlasData& GlyphAtlasData) -> bool
+	{
+		const FSdfGlyphEntryKey GlyphKey(InShapedGlyph.FontFaceData->FontFace, InShapedGlyph.GlyphIndex, AtlasType, SdfPpem, GlyphAtlasData.EmOuterSpread, GlyphAtlasData.EmInnerSpread);
 		FSlateSdfGenerator::FRequestDescriptor SdfRequestDescriptor = { 
 			GlyphKey.FontFace,
 			GlyphKey.GlyphIndex,
@@ -1286,8 +1274,8 @@ FSdfGlyphFontAtlasData FSlateFontCache::GetSdfGlyphFontAtlasData(const FShapedGl
 		}
 	}
 
-	EmOuterSpread = CeilSdfHalfSpread(EmOuterSpread, SdfPpem);
-	EmInnerSpread = CeilSdfHalfSpread(EmInnerSpread, SdfPpem);
+	EmOuterSpread = CeilSdfHalfSpread(EmOuterSpread, float(SdfPpem));
+	EmInnerSpread = CeilSdfHalfSpread(EmInnerSpread, float(SdfPpem));
 
 	auto FinalizeEntry = [&](const TSharedRef<FSdfGlyphFontAtlasData>& InFontData) -> FSdfGlyphFontAtlasData
 	{
@@ -1308,7 +1296,7 @@ FSdfGlyphFontAtlasData FSlateFontCache::GetSdfGlyphFontAtlasData(const FShapedGl
 	};
 	// Not cached on the glyph, so create a key for to look up this glyph, as it may
 	// have already been cached by another shaped text sequence
-	const FSdfGlyphEntryKey GlyphKey(InShapedGlyph.FontFaceData->FontFace, InShapedGlyph.GlyphIndex, AtlasType, InSdfSettings.GetClampedPpem(), EmOuterSpread, EmInnerSpread);
+	const FSdfGlyphEntryKey GlyphKey(InShapedGlyph.FontFaceData->FontFace, InShapedGlyph.GlyphIndex, AtlasType, SdfPpem, EmOuterSpread, EmInnerSpread);
 
 	{
 		// Has the atlas data already been cached by another shaped text sequence?
diff --git a/Engine/Source/Runtime/SlateCore/Private/Fonts/SlateTextShaper.cpp b/Engine/Source/Runtime/SlateCore/Private/Fonts/SlateTextShaper.cpp
--- a/Engine/Source/Runtime/SlateCore/Private/Fonts/SlateTextShaper.cpp
+++ b/Engine/Source/Runtime/SlateCore/Private/Fonts/SlateTextShaper.cpp
@@ -308,14 +308,11 @@ FShapedGlyphSequenceRef FSlateTextShaper::FinalizeTextShaping(TArray<FShapedGlyp
 	}
 #endif // WITH_FREETYPE
 
-	const IFontProviderInterface* const FontObject = Cast<const IFontProviderInterface>(InFontInfo.FontObject);
 	return MakeShared<FShapedGlyphSequence>(MoveTemp(InGlyphsToRender), 
 											TextBaseline, 
 											MaxHeight, 
 											InFontInfo.FontMaterial.Get(),
 											InFontInfo.OutlineSettings,
-											FontObject ? FontObject->GetFontRasterizationMode() : EFontRasterizationMode::Bitmap,
-											FontObject ? FontObject->GetSdfSettings() : FFontSdfSettings(),
 											InSourceTextRange);
 }
 
@@ -425,7 +422,8 @@ void FSlateTextShaper::PerformKerningOnlyTextShaping(const TCHAR* InText, const
 			}
 
 			FreeTypeUtils::ApplySizeAndScale(KerningOnlyTextSequenceEntry.FaceAndMemory->GetFace(), InFontInfo.Size, FinalFontScale);
-			TSharedRef<FShapedGlyphFaceData> ShapedGlyphFaceData = MakeShared<FShapedGlyphFaceData>(KerningOnlyTextSequenceEntry.FaceAndMemory, GlyphFlags, InFontInfo.Size, FinalFontScale, InFontInfo.GetClampSkew());
+			const FFontRasterizationSettings FontRasterizationSettings = KerningOnlyTextSequenceEntry.FontDataPtr->GetFontRasterizationSettings();
+			TSharedRef<FShapedGlyphFaceData> ShapedGlyphFaceData = MakeShared<FShapedGlyphFaceData>(KerningOnlyTextSequenceEntry.FaceAndMemory, GlyphFlags, InFontInfo.Size, FinalFontScale, InFontInfo.GetClampSkew(), FontRasterizationSettings.Mode, FontRasterizationSettings.DistanceFieldPpem);
 			TSharedPtr<FFreeTypeKerningCache> KerningCache = FTCacheDirectory->GetKerningCache(KerningOnlyTextSequenceEntry.FaceAndMemory->GetFace(), FT_KERNING_DEFAULT, InFontInfo.Size, FinalFontScale);
 			TSharedRef<FFreeTypeAdvanceCache> AdvanceCache = FTCacheDirectory->GetAdvanceCache(KerningOnlyTextSequenceEntry.FaceAndMemory->GetFace(), GlyphFlags, InFontInfo.Size, FinalFontScale);
 
@@ -708,11 +706,14 @@ void FSlateTextShaper::PerformHarfBuzzTextShaping(const TCHAR* InText, const int
 			{
 				continue;
 			}
+			const FFontRasterizationSettings FontRasterizationSettings = HarfBuzzTextSequenceEntry.FontDataPtr->GetFontRasterizationSettings();
 			TSharedRef<FShapedGlyphFaceData> ShapedGlyphFaceData = MakeShared<FShapedGlyphFaceData>(HarfBuzzTextSequenceEntry.FaceAndMemory, 
 																									GlyphFlags, 
 																									InFontInfo.Size, 
 																									FinalFontScale, 
-																									InFontInfo.GetClampSkew());
+																									InFontInfo.GetClampSkew(),
+																									FontRasterizationSettings.Mode,
+																									FontRasterizationSettings.DistanceFieldPpem);
 			TSharedPtr<FFreeTypeKerningCache> KerningCache = FTCacheDirectory->GetKerningCache(HarfBuzzTextSequenceEntry.FaceAndMemory->GetFace(), 
 																							   FT_KERNING_DEFAULT,
 																							   InFontInfo.Size, 
diff --git a/Engine/Source/Runtime/SlateCore/Private/Fonts/SlateTextShaper.h b/Engine/Source/Runtime/SlateCore/Private/Fonts/SlateTextShaper.h
--- a/Engine/Source/Runtime/SlateCore/Private/Fonts/SlateTextShaper.h
+++ b/Engine/Source/Runtime/SlateCore/Private/Fonts/SlateTextShaper.h
@@ -6,6 +6,7 @@
 #include "Fonts/FontCache.h"
 #include "Fonts/FontCacheFreeType.h"
 #include "Fonts/FontCacheHarfBuzz.h"
+#include "Fonts/FontRasterizationMode.h"
 
 class FCompositeFontCache;
 class FSlateFontRenderer;
@@ -17,20 +18,20 @@ class IBreakIterator;
 class FShapedGlyphFaceData
 {
 public:
-	FShapedGlyphFaceData(TWeakPtr<FFreeTypeFace> InFontFace, const uint32 InGlyphFlags, const float InFontSize, const float InFontScale, const float InFontSkew)
+	FShapedGlyphFaceData(TWeakPtr<FFreeTypeFace> InFontFace, const uint32 InGlyphFlags, const float InFontSize, const float InFontScale, const float InFontSkew, const EFontRasterizationMode InRasterizationMode, const int16 InSdfPpem)
 		: FontFace(MoveTemp(InFontFace))
 		, GlyphFlags(InGlyphFlags)
 		, FontSize(InFontSize)
 		, FontScale(InFontScale)
 		, BitmapRenderScale(1.0f)
 		, FontSkew(InFontSkew)
-		, bSupportsSdf(false)
+		, RasterizationMode(InRasterizationMode)
+		, SdfPpem(InSdfPpem)
 	{
 #if WITH_FREETYPE
 		if (TSharedPtr<FFreeTypeFace> FontFacePin = FontFace.Pin())
 		{
 			BitmapRenderScale = FontFacePin->GetBitmapRenderScale();
-			bSupportsSdf = FontFacePin->SupportsSdf();
 		}
 #endif	// WITH_FREETYPE
 	}
@@ -47,8 +48,11 @@ public:
 	float BitmapRenderScale;
 	/** The skew transform amount for the rendered font */
 	float FontSkew;
-	/** Is the Face eligible for signed distance field rendering */
-	bool bSupportsSdf;
+	/** The glyph rasterization mode (bitmap / distance field) */
+	EFontRasterizationMode RasterizationMode;
+	/** Distance field px/em resolution, not used if RasterizationMode is Bitmap */
+	int16 SdfPpem;
+
 };
 
 
diff --git a/Engine/Source/Runtime/SlateCore/Private/Rendering/ElementBatcher.cpp b/Engine/Source/Runtime/SlateCore/Private/Rendering/ElementBatcher.cpp
--- a/Engine/Source/Runtime/SlateCore/Private/Rendering/ElementBatcher.cpp
+++ b/Engine/Source/Runtime/SlateCore/Private/Rendering/ElementBatcher.cpp
@@ -27,7 +27,6 @@ DEFINE_STAT(STAT_SlateElements_Box);
 DEFINE_STAT(STAT_SlateElements_Border);
 DEFINE_STAT(STAT_SlateElements_Text);
 DEFINE_STAT(STAT_SlateElements_ShapedText);
-DEFINE_STAT(STAT_SlateElements_ShapedTextSdf);
 DEFINE_STAT(STAT_SlateElements_Line);
 DEFINE_STAT(STAT_SlateElements_Other);
 DEFINE_STAT(STAT_SlateInvalidation_RecachedElements);
@@ -325,7 +324,6 @@ void FSlateElementBatcher::AddElements(FSlateWindowElementList& WindowElementLis
 	ElementStat_Borders = 0;
 	ElementStat_Text = 0;
 	ElementStat_ShapedText = 0;
-	ElementStat_ShapedTextSdf = 0;
 	ElementStat_Line = 0;
 	ElementStat_RecachedElements = 0;
 #endif
@@ -369,7 +367,6 @@ void FSlateElementBatcher::AddElements(FSlateWindowElementList& WindowElementLis
 		ElementStat_Borders +
 		ElementStat_Text +
 		ElementStat_ShapedText +
-		ElementStat_ShapedTextSdf +
 		ElementStat_Line +
 		ElementStat_Other;
 
@@ -378,7 +375,6 @@ void FSlateElementBatcher::AddElements(FSlateWindowElementList& WindowElementLis
 	INC_DWORD_STAT_BY(STAT_SlateElements_Border, ElementStat_Borders);
 	INC_DWORD_STAT_BY(STAT_SlateElements_Text, ElementStat_Text);
 	INC_DWORD_STAT_BY(STAT_SlateElements_ShapedText, ElementStat_ShapedText);
-	INC_DWORD_STAT_BY(STAT_SlateElements_ShapedTextSdf, ElementStat_ShapedTextSdf);
 	INC_DWORD_STAT_BY(STAT_SlateElements_Line, ElementStat_Line);
 	INC_DWORD_STAT_BY(STAT_SlateElements_Other, ElementStat_Other);
 	INC_DWORD_STAT_BY(STAT_SlateInvalidation_RecachedElements, ElementStat_RecachedElements);
@@ -432,8 +428,7 @@ void FSlateElementBatcher::AddElementsInternal(const FSlateDrawElementMap& DrawE
 		// ElementStat_ShapedText/Sdf incremented in AddShapedTextElement
 		for (const FSlateShapedTextElement& DrawElement : ShapedTextElements)
 		{
-			bool bSdfFont = DrawElement.GetShapedGlyphSequence() && DrawElement.GetShapedGlyphSequence()->IsSdfFont();
-			DrawElement.IsPixelSnapped() && !bSdfFont ? AddShapedTextElement<ESlateVertexRounding::Enabled>(DrawElement) : AddShapedTextElement<ESlateVertexRounding::Disabled>(DrawElement);
+			DrawElement.IsPixelSnapped() ? AddShapedTextElement<ESlateVertexRounding::Enabled>(DrawElement) : AddShapedTextElement<ESlateVertexRounding::Disabled>(DrawElement);
 		}
 	}
 
@@ -1582,7 +1577,6 @@ void FSlateElementBatcher::AddShapedTextElement( const FSlateShapedTextElement&
 	checkSlow(ShapedGlyphSequence);
 
 	const FFontOutlineSettings& OutlineSettings = ShapedGlyphSequence->GetFontOutlineSettings();
-	const bool bSdfFont = ShapedGlyphSequence->IsSdfFont();
 
 	ensure(ShapedGlyphSequence->GetGlyphsToRender().Num() > 0);
 
@@ -1681,7 +1675,7 @@ void FSlateElementBatcher::AddShapedTextElement( const FSlateShapedTextElement&
 		return BuildShapedTextSequence<Rounding>(BuildContext);
 	};
 
-	STAT((bSdfFont ? ElementStat_ShapedTextSdf : ElementStat_ShapedText)++);
+	STAT(ElementStat_ShapedText++);
 
 	const UObject* BaseFontMaterial = ShapedGlyphSequence->GetFontMaterial();
 	const bool bOutlineFont = OutlineSettings.OutlineSize > 0;
@@ -3267,8 +3261,8 @@ int32 FSlateElementBatcher::BuildShapedTextSequence(const FShapedTextBuildContex
 	float EllipsisLineY = 0;
 	bool bNeedEllipsis = false;
 	bool bNeedSpaceForEllipsis = false;
-	bool bIsSdfFont = GlyphSequenceToRender->IsSdfFont();
-	bool bRequiresManualSkewing = bIsSdfFont && !FMath::IsNearlyEqual(GlyphSequenceToRender->GetFontSkew(), 0.f);
+	const bool bNonZeroSkew = !FMath::IsNearlyEqual(GlyphSequenceToRender->GetFontSkew(), 0.f);
+	bool bPrevSdfGlyph = false;
 	float SdfPixelSpread = 0;
 	float SdfBias = 0;
 	// For left to right overflow direction - Sum of total whitespace we're currently advancing through. Once a non-whitespace glyph is detected this will return to 0
@@ -3289,6 +3283,7 @@ int32 FSlateElementBatcher::BuildShapedTextSequence(const FShapedTextBuildContex
 
 		const float BitmapRenderScale = GlyphToRender.GetBitmapRenderScale();
 		const float InvBitmapRenderScale = 1.0f / BitmapRenderScale;
+		const bool bIsSdfGlyph = GlyphToRender.FontFaceData->RasterizationMode != EFontRasterizationMode::Bitmap;
 
 		float X = 0;
 		float SizeX = 0;
@@ -3313,12 +3308,10 @@ int32 FSlateElementBatcher::BuildShapedTextSequence(const FShapedTextBuildContex
 			int8 NextAtlasDataTextureIndex = -1;
 			float NextSdfPixelSpread = 0;
 			float NextSdfBias = 0;
-			const bool bIsSdfGlyph = bIsSdfFont && GlyphToRender.FontFaceData && GlyphToRender.FontFaceData->bSupportsSdf;
 
 			if (bIsSdfGlyph)
 			{
-				const FFontSdfSettings& FontSdfSettings = GlyphSequenceToRender->GetFontSdfSettings();
-				const FSdfGlyphFontAtlasData SdfGlyphAtlasData = Context.FontCache->GetSdfGlyphFontAtlasData(GlyphToRender, *Context.OutlineSettings, GlyphSequenceToRender->GetRasterizationMode(), FontSdfSettings);
+				const FSdfGlyphFontAtlasData SdfGlyphAtlasData = Context.FontCache->GetSdfGlyphFontAtlasData(GlyphToRender, *Context.OutlineSettings);
 				bCanRenderGlyph = (SdfGlyphAtlasData.Valid && SdfGlyphAtlasData.bSupportsSdf);
 				if (bCanRenderGlyph)
 				{
@@ -3337,7 +3330,7 @@ int32 FSlateElementBatcher::BuildShapedTextSequence(const FShapedTextBuildContex
 					// SdfGlyphAtlasData.Metrics values assume this operation, so QuadMeshSize already accounts for this
 					SpriteSize = FVector2f(SdfGlyphAtlasData.USize-1, SdfGlyphAtlasData.VSize-1);
 					SpriteOffset = FVector2f(SdfGlyphAtlasData.StartU+.5f, SdfGlyphAtlasData.StartV+.5f);
-					NextSdfPixelSpread = (SdfGlyphAtlasData.EmInnerSpread+SdfGlyphAtlasData.EmOuterSpread)*static_cast<float>(FontSdfSettings.GetClampedPpem());
+					NextSdfPixelSpread = (SdfGlyphAtlasData.EmInnerSpread+SdfGlyphAtlasData.EmOuterSpread)*float(GlyphToRender.FontFaceData->SdfPpem);
 					// Value representing zero distance
 					NextSdfBias = (SdfGlyphAtlasData.EmOuterSpread-EmOutlineSize)/(SdfGlyphAtlasData.EmInnerSpread+SdfGlyphAtlasData.EmOuterSpread);
 				}
@@ -3375,8 +3368,9 @@ int32 FSlateElementBatcher::BuildShapedTextSequence(const FShapedTextBuildContex
 					}
 				}
 
+				// Initiate new render batch for first glyph and whenever texture / shader / shader parameters change between previous and current glyph
 				check(NextAtlasDataTextureIndex >= 0);
-				if (FontAtlasTexture == nullptr || NextAtlasDataTextureIndex != FontTextureIndex || (bIsSdfGlyph && (NextSdfPixelSpread != SdfPixelSpread || NextSdfBias != SdfBias)))
+				if (FontAtlasTexture == nullptr || NextAtlasDataTextureIndex != FontTextureIndex || bIsSdfGlyph != bPrevSdfGlyph || (bIsSdfGlyph && (NextSdfPixelSpread != SdfPixelSpread || NextSdfBias != SdfBias)))
 				{
 					// Font has a new texture for this glyph or shader parameters changed. Refresh the batch we use and the index we are currently using
 					FontTextureIndex = NextAtlasDataTextureIndex;
@@ -3431,7 +3425,7 @@ int32 FSlateElementBatcher::BuildShapedTextSequence(const FShapedTextBuildContex
 							// Signed distance sample bias, the color value (between 0 to 1) representing zero distance
 							SdfBias,
 							// The last parameter needs to be 0 for alpha texture single-channel SDF, 1 for BGRA/RGBA MTSDF
-							1.f
+							ContentType == ESlateFontAtlasContentType::Msdf ? 1.f : 0.f
 						));
 					}
 
@@ -3479,6 +3473,7 @@ int32 FSlateElementBatcher::BuildShapedTextSequence(const FShapedTextBuildContex
 					{
 						RenderBatch->ReserveIndices(GlyphsLeft * 6);
 					}
+					bPrevSdfGlyph = bIsSdfGlyph;
 				}
 				U = SpriteOffset.X * InvTextureSizeX;
 				V = SpriteOffset.Y * InvTextureSizeY;
@@ -3604,7 +3599,7 @@ int32 FSlateElementBatcher::BuildShapedTextSequence(const FShapedTextBuildContex
 				VtMax = FMath::Lerp(0.0f, 1.0f, LowerLeft.Y / Context.MaxHeight);
 			}
 
-			if (bRequiresManualSkewing)
+			if (bNonZeroSkew && bIsSdfGlyph)
 			{
 				FTransform2f ShearTransform(FShear2f(GlyphSequenceToRender->GetFontSkew(), 0.f));
 				FVector2f DeltaTopLeft(0, QuadMeshOffsets.Y);
@@ -3615,10 +3610,21 @@ int32 FSlateElementBatcher::BuildShapedTextSequence(const FShapedTextBuildContex
 				LowerLeft.X += DeltaBottomLeft.X;	LowerRight.X += DeltaBottomLeft.X;
 			}
 			// Add four vertices to the list of verts to be added to the vertex buffer
-			RenderBatch->AddVertex(FSlateVertex::Make<Rounding>(RenderTransform, FVector2f(UpperLeft), FVector4f(U, V, Ut, Vt), FVector2f(0.0f, 0.0f), Tint));
-			RenderBatch->AddVertex(FSlateVertex::Make<Rounding>(RenderTransform, FVector2f(UpperRight), FVector4f(U + SizeU, V, UtMax, Vt), FVector2f(1.0f, 0.0f), Tint));
-			RenderBatch->AddVertex(FSlateVertex::Make<Rounding>(RenderTransform, FVector2f(LowerLeft), FVector4f(U, V + SizeV, Ut, VtMax), FVector2f(0.0f, 1.0f), Tint));
-			RenderBatch->AddVertex(FSlateVertex::Make<Rounding>(RenderTransform, FVector2f(LowerRight), FVector4f(U + SizeU, V + SizeV, UtMax, VtMax), FVector2f(1.0f, 1.0f), Tint));
+			if (bIsSdfGlyph)
+			{
+				// Vertex rounding has no benefit for distance field rendering but makes positioning worse, so disable it
+				RenderBatch->AddVertex(FSlateVertex::Make<ESlateVertexRounding::Disabled>(RenderTransform, FVector2f(UpperLeft), FVector4f(U, V, Ut, Vt), FVector2f(0.0f, 0.0f), Tint));
+				RenderBatch->AddVertex(FSlateVertex::Make<ESlateVertexRounding::Disabled>(RenderTransform, FVector2f(UpperRight), FVector4f(U + SizeU, V, UtMax, Vt), FVector2f(1.0f, 0.0f), Tint));
+				RenderBatch->AddVertex(FSlateVertex::Make<ESlateVertexRounding::Disabled>(RenderTransform, FVector2f(LowerLeft), FVector4f(U, V + SizeV, Ut, VtMax), FVector2f(0.0f, 1.0f), Tint));
+				RenderBatch->AddVertex(FSlateVertex::Make<ESlateVertexRounding::Disabled>(RenderTransform, FVector2f(LowerRight), FVector4f(U + SizeU, V + SizeV, UtMax, VtMax), FVector2f(1.0f, 1.0f), Tint));
+			}
+			else
+			{
+				RenderBatch->AddVertex(FSlateVertex::Make<Rounding>(RenderTransform, FVector2f(UpperLeft), FVector4f(U, V, Ut, Vt), FVector2f(0.0f, 0.0f), Tint));
+				RenderBatch->AddVertex(FSlateVertex::Make<Rounding>(RenderTransform, FVector2f(UpperRight), FVector4f(U + SizeU, V, UtMax, Vt), FVector2f(1.0f, 0.0f), Tint));
+				RenderBatch->AddVertex(FSlateVertex::Make<Rounding>(RenderTransform, FVector2f(LowerLeft), FVector4f(U, V + SizeV, Ut, VtMax), FVector2f(0.0f, 1.0f), Tint));
+				RenderBatch->AddVertex(FSlateVertex::Make<Rounding>(RenderTransform, FVector2f(LowerRight), FVector4f(U + SizeU, V + SizeV, UtMax, VtMax), FVector2f(1.0f, 1.0f), Tint));
+			}
 
 			if (bUseStaticIndicies)
 			{
diff --git a/Engine/Source/Runtime/SlateCore/Public/Fonts/CompositeFont.h b/Engine/Source/Runtime/SlateCore/Public/Fonts/CompositeFont.h
--- a/Engine/Source/Runtime/SlateCore/Public/Fonts/CompositeFont.h
+++ b/Engine/Source/Runtime/SlateCore/Public/Fonts/CompositeFont.h
@@ -3,6 +3,7 @@
 #pragma once
 
 #include "CoreMinimal.h"
+#include "Fonts/FontRasterizationMode.h"
 #include "UObject/ObjectMacros.h"
 #include "UObject/GCObject.h"
 #include "UObject/Class.h"
@@ -131,6 +132,15 @@ private:
 	TArray<uint8> Data;
 };
 
+/** Font face rasterization settings */
+struct FFontRasterizationSettings
+{
+	/** Rasterization mode (bitmap / distance field) */
+	EFontRasterizationMode Mode = EFontRasterizationMode::Bitmap;
+	/** Distance field px/em resolution, not used if Mode is Bitmap */
+	int32 DistanceFieldPpem = 0;
+};
+
 /** Payload data describing an individual font in a typeface. Keep this lean as it's also used as a key! */
 USTRUCT()
 struct FFontData
@@ -188,6 +198,9 @@ struct FFontData
 	/** Get the font face asset used by this data (if any). */
 	SLATECORE_API const UObject* GetFontFaceAsset() const;
 
+	/** Get the font face rasterization settings. */
+	SLATECORE_API FFontRasterizationSettings GetFontRasterizationSettings() const;
+
 #if WITH_EDITORONLY_DATA
 	/** True if this object contains any legacy data that needs to be upgraded PostLoad by calling the functions below (in order). */
 	SLATECORE_API bool HasLegacyData() const;
diff --git a/Engine/Source/Runtime/SlateCore/Public/Fonts/FontCache.h b/Engine/Source/Runtime/SlateCore/Public/Fonts/FontCache.h
--- a/Engine/Source/Runtime/SlateCore/Public/Fonts/FontCache.h
+++ b/Engine/Source/Runtime/SlateCore/Public/Fonts/FontCache.h
@@ -354,7 +354,6 @@ public:
 		, SequenceWidth(0)
 		, GlyphFontFaces()
 		, SourceIndicesToGlyphData(FSourceTextRange(0, 0))
-		, SdfSettings()
 		, CachedFontSkew(0.f)
 	{ }
 
@@ -363,8 +362,6 @@ public:
 						 const uint16 InMaxTextHeight,
 						 const UObject* InFontMaterial,
 						 const FFontOutlineSettings& InOutlineSettings,
-						 const EFontRasterizationMode InRasterizationMode,
-						 const FFontSdfSettings& InSdfSettings,
 						 const FSourceTextRange& InSourceTextRange);
 	SLATECORE_API ~FShapedGlyphSequence();
 
@@ -401,18 +398,6 @@ public:
 		return OutlineSettings;
 	}
 
-	/** Returns true if the rasterization mode is signed distance field-based (and the feature is enabled) */
-	bool IsSdfFont() const;
-
-	/** Get the font rasterization mode to be used when rendering these glyphs */
-	EFontRasterizationMode GetRasterizationMode() const;
-
-	/** Get the signed distance field settings to be used when rendering these glyphs in distance field mode */
-	const FFontSdfSettings& GetFontSdfSettings() const
-	{
-		return SdfSettings;
-	}
-
 	/** Check to see whether this glyph sequence is dirty (ie, contains glyphs with invalid font pointers) */
 	SLATECORE_API bool IsDirty() const;
 
@@ -604,10 +589,6 @@ private:
 	TArray<TWeakPtr<FFreeTypeFace>> GlyphFontFaces;
 	/** A map of source indices to their shaped glyph data indices - used to perform efficient reverse look-up */
 	FSourceIndicesToGlyphData SourceIndicesToGlyphData;
-	/** Rasterization mode to use when rendering these glyphs */
-	EFontRasterizationMode RasterizationMode;
-	/** Sdf settings to use when rendering these glyphs */
-	FFontSdfSettings SdfSettings;
 	/** The Font Skew parameter of the FontInfo applied to the this Shaped glype sequence */
 	float CachedFontSkew;
 
@@ -836,7 +817,7 @@ public:
 	/**
 	 * Get the atlas information and the scaled metrics of a given shaped sdf glyph. This information will be cached if required.
 	 */
-	SLATECORE_API FSdfGlyphFontAtlasData GetSdfGlyphFontAtlasData(const FShapedGlyphEntry& InShapedGlyph, const FFontOutlineSettings& InOutlineSettings, EFontRasterizationMode InRasterizationMode, const FFontSdfSettings& InSdfSettings);
+	SLATECORE_API FSdfGlyphFontAtlasData GetSdfGlyphFontAtlasData(const FShapedGlyphEntry& InShapedGlyph, const FFontOutlineSettings& InOutlineSettings);
 
 	/**
 	 * Gets the overflow glyph sequence for a given font. The overflow sequence is used to replace characters that are clipped
diff --git a/Engine/Source/Runtime/SlateCore/Public/Fonts/FontFaceInterface.h b/Engine/Source/Runtime/SlateCore/Public/Fonts/FontFaceInterface.h
--- a/Engine/Source/Runtime/SlateCore/Public/Fonts/FontFaceInterface.h
+++ b/Engine/Source/Runtime/SlateCore/Public/Fonts/FontFaceInterface.h
@@ -58,4 +58,7 @@ class IFontFaceInterface
 
 	/** Get the data buffer containing the data for the current font face. */
 	virtual FFontFaceDataConstRef GetFontFaceData() const = 0;
+
+	/** Returns the font face's rasterization settings. */
+	virtual FFontRasterizationSettings GetRasterizationSettings() const = 0;
 };
diff --git a/Engine/Source/Runtime/SlateCore/Public/Fonts/FontProviderInterface.h b/Engine/Source/Runtime/SlateCore/Public/Fonts/FontProviderInterface.h
--- a/Engine/Source/Runtime/SlateCore/Public/Fonts/FontProviderInterface.h
+++ b/Engine/Source/Runtime/SlateCore/Public/Fonts/FontProviderInterface.h
@@ -26,10 +26,4 @@ class IFontProviderInterface
 
 	virtual const FCompositeFont* GetCompositeFont() const = 0;
 
-	virtual bool IsSdfFont() const = 0;
-
-	virtual EFontRasterizationMode GetFontRasterizationMode() const = 0;
-
-	virtual const FFontSdfSettings& GetSdfSettings() const = 0;
-
 };
diff --git a/Engine/Source/Runtime/SlateCore/Public/Rendering/ElementBatcher.h b/Engine/Source/Runtime/SlateCore/Public/Rendering/ElementBatcher.h
--- a/Engine/Source/Runtime/SlateCore/Public/Rendering/ElementBatcher.h
+++ b/Engine/Source/Runtime/SlateCore/Public/Rendering/ElementBatcher.h
@@ -28,7 +28,6 @@ DECLARE_DWORD_COUNTER_STAT_EXTERN(TEXT("Elements (Box)"), STAT_SlateElements_Box
 DECLARE_DWORD_COUNTER_STAT_EXTERN(TEXT("Elements (Border)"), STAT_SlateElements_Border, STATGROUP_Slate, SLATECORE_API);
 DECLARE_DWORD_COUNTER_STAT_EXTERN(TEXT("Elements (Text)"), STAT_SlateElements_Text, STATGROUP_Slate, SLATECORE_API);
 DECLARE_DWORD_COUNTER_STAT_EXTERN(TEXT("Elements (ShapedText)"), STAT_SlateElements_ShapedText, STATGROUP_Slate, SLATECORE_API);
-DECLARE_DWORD_COUNTER_STAT_EXTERN(TEXT("Elements (Sdf)"), STAT_SlateElements_ShapedTextSdf, STATGROUP_Slate, SLATECORE_API);
 DECLARE_DWORD_COUNTER_STAT_EXTERN(TEXT("Elements (Line)"), STAT_SlateElements_Line, STATGROUP_Slate, SLATECORE_API);
 DECLARE_DWORD_COUNTER_STAT_EXTERN(TEXT("Elements (Other)"), STAT_SlateElements_Other, STATGROUP_Slate, SLATECORE_API);
 
@@ -471,9 +470,6 @@ private:
 	/** Track the number of drawn shaped text from the previous frame to report to stats. */
 	int32 ElementStat_ShapedText;
 
-	/** Track the number of drawn shaped text using signed distance field from the previous frame to report to stats. */
-	int32 ElementStat_ShapedTextSdf;
-
 	/** Track the number of drawn lines from the previous frame to report to stats. */
 	int32 ElementStat_Line;
 
