Slate: Display glyph info (Font, size, outline size, glyphindex, fontskew, ...) when hovering a glyph in the Font Atlas Visualizer.

[RN] Slate Font Atlas Visualizer: when hovering a glyph, display info about it (Font, size,...).
#tests Opened the Font atlas viewer and Texture atlas viewer, hovered glyphs and textures, was able to display expected info for all but one (issue in the FontCache system with SDF to investigate). No regression on texture tooltip display. Flushed the cache while actively hovering glyph without stability issue.

diff --git a/Engine/Source/Developer/SlateReflector/Private/Widgets/SAtlasVisualizer.cpp b/Engine/Source/Developer/SlateReflector/Private/Widgets/SAtlasVisualizer.cpp
--- a/Engine/Source/Developer/SlateReflector/Private/Widgets/SAtlasVisualizer.cpp
+++ b/Engine/Source/Developer/SlateReflector/Private/Widgets/SAtlasVisualizer.cpp
@@ -461,10 +461,13 @@ void SAtlasVisualizer::RebuildToolTip(const FAtlasSlotInfo& Info)
 
 		TArray<FName> Resources = FSlateStyleRegistry::GetSylesUsingBrush(Info.TextureName);
 
-		Builder.AppendLine(LOCTEXT("AtlasDebuggingToolTipTitle", "\nUsed by:"));
-		for (FName Name : Resources)
+		if (!Resources.IsEmpty())
 		{
-			Builder.AppendLine(Name);
+			Builder.AppendLine(LOCTEXT("AtlasDebuggingToolTipTitle", "\nUsed by:"));
+			for (FName Name : Resources)
+			{
+				Builder.AppendLine(Name);
+			}
 		}
 
 		SetToolTipText(Builder.ToText());
diff --git a/Engine/Source/Runtime/SlateCore/Private/Fonts/FontCache.cpp b/Engine/Source/Runtime/SlateCore/Private/Fonts/FontCache.cpp
--- a/Engine/Source/Runtime/SlateCore/Private/Fonts/FontCache.cpp
+++ b/Engine/Source/Runtime/SlateCore/Private/Fonts/FontCache.cpp
@@ -3,6 +3,7 @@
 #include "Fonts/FontCache.h"
 #include "Misc/ScopeLock.h"
 #include "HAL/IConsoleManager.h"
+#include "Algo/Find.h"
 #include "Application/SlateApplicationBase.h"
 #include "Fonts/FontCacheFreeType.h"
 #include "Fonts/FontCacheHarfBuzz.h"
@@ -933,6 +934,121 @@ bool FSlateFontCache::IsAtlasPageResourceAlphaOnly(const int32 InIndex) const
 	return AllFontTextures[GetAllFontTexturesIndex(InIndex)]->GetContentType() == ESlateFontAtlasContentType::Alpha;
 }
 
+#if WITH_ATLAS_DEBUGGING
+FAtlasSlotInfo FSlateFontCache::GetAtlasSlotInfoAtPosition(FIntPoint InPosition, int32 AtlasIndex) const
+{
+	if (!AllFontTextures.IsValidIndex(AtlasIndex)) //The currently selected texture page can become empty for a short amount of time when flushing cache, so prevent crash when flushing with Font Atlas Visualizer open!
+	{
+		return FAtlasSlotInfo();
+	}
+
+	const int32 TextureIndex = GetAllFontTexturesIndex(AtlasIndex);
+	const FSlateTextureAtlas& Atlas = static_cast<FSlateFontAtlas&>(AllFontTextures[TextureIndex].Get());
+
+	const FAtlasedTextureSlot* Slot = Atlas.GetSlotAtPosition(InPosition);
+	if (Slot)
+	{
+		TStringBuilder<512> Builder;
+		auto FeedFaceInfo = [&Builder](const TWeakPtr<FFreeTypeFace>& FontFace)
+			{
+#if WITH_FREETYPE
+				FFreeTypeFace* FreeTypeFace = FontFace.Pin().Get();
+				if (FreeTypeFace)
+				{
+					FT_Face FT_face = FreeTypeFace->GetFace();
+					Builder += TEXT("Font: ");
+					Builder += FT_face->family_name;
+					Builder += TEXT(" ");
+					Builder += FT_face->style_name;
+				}
+#endif
+			};
+
+		//Try to get info from bitmap map.
+		const TPair<FShapedGlyphEntryKey, TSharedRef<FShapedGlyphFontAtlasData>>* Glyph = Algo::FindByPredicate(ShapedGlyphToAtlasData, [&Slot, &TextureIndex](const TPair<FShapedGlyphEntryKey, TSharedRef<FShapedGlyphFontAtlasData>>& Entry)
+			{
+				const FShapedGlyphFontAtlasData& AtlasData = Entry.Value.Get();
+				return	AtlasData.StartU == Slot->X && 
+						AtlasData.StartV == Slot->Y && 
+						TextureIndex == AtlasData.TextureIndex;
+			});
+
+		//Display info from bitmap and return if found.
+		if (Glyph != nullptr)
+		{
+			const FShapedGlyphEntryKey& EntryKey = Glyph->Key;
+
+			FeedFaceInfo(EntryKey.FontFace);
+			Builder += TEXT("\nSize: ");
+			Builder += FString::FromInt(EntryKey.FontRenderSize);
+			Builder += TEXT("\nOutlineSize: ");
+			Builder += LexToString(EntryKey.OutlineRenderSize);
+			Builder += TEXT("\nGlyphIndex: ");
+			Builder += FString::FromInt(EntryKey.GlyphIndex);
+			Builder += TEXT("\nFontSkew: ");
+			Builder += LexToString(EntryKey.FontSkew);
+
+			FAtlasSlotInfo NewInfo;
+			NewInfo.AtlasSlotRect = FSlateRect(FVector2f(static_cast<float>(Slot->X), static_cast<float>(Slot->Y)), FVector2f(static_cast<float>(Slot->X + Slot->Width), static_cast<float>(Slot->Y + Slot->Height)));
+			NewInfo.TextureName = Builder.ToString();
+			return NewInfo;
+		}
+
+		//Try to get info from SDF map
+		const TPair<FSdfGlyphEntryKey, TSharedRef<FSdfGlyphFontAtlasData>>* SdfGlyph = Algo::FindByPredicate(SdfGlyphToAtlasData, [&Slot, &TextureIndex](const TPair<FSdfGlyphEntryKey, TSharedRef<FSdfGlyphFontAtlasData>>& Entry)
+			{
+				const FSdfGlyphFontAtlasData& AtlasData = Entry.Value.Get();
+				return	AtlasData.StartU == Slot->X + Slot->Padding && 
+						AtlasData.StartV == Slot->Y + Slot->Padding && 
+						AtlasData.TextureIndex == TextureIndex;
+			});
+
+
+		const FSdfGlyphEntryKey* EntryKey = nullptr;
+		if (SdfGlyph != nullptr)
+		{
+			EntryKey = &(SdfGlyph->Key);
+		}
+		else
+		{
+			//Last chance, try to get info from SDF tasks (rasterization in progress)
+			const TPair<FSdfGlyphTaskKey, TSharedRef<FSdfGlyphFontAtlasData>>* SdfTaskGlyph = Algo::FindByPredicate(SdfTaskToAtlasData, [&Slot, &TextureIndex](const TPair<FSdfGlyphTaskKey, TSharedRef<FSdfGlyphFontAtlasData>>& Entry)
+				{
+					const FSdfGlyphFontAtlasData& AtlasData = Entry.Value.Get();
+					return	AtlasData.StartU == Slot->X + Slot->Padding &&
+							AtlasData.StartV == Slot->Y + Slot->Padding &&
+							AtlasData.TextureIndex == TextureIndex;
+				});
+
+			if (SdfTaskGlyph != nullptr)
+			{
+				EntryKey = &(SdfTaskGlyph->Key.SdfGlyphEntryKey);
+				Builder += TEXT("***Temporary glyph***\n");
+			}
+		}
+
+		//Display info from any SDF if found.
+		if (EntryKey != nullptr)
+		{
+			FeedFaceInfo(EntryKey->FontFace);
+			Builder += TEXT("\nSize: ");
+			Builder += FString::FromInt(EntryKey->Ppem);
+			Builder += TEXT("\nGlyphIndex: ");
+			Builder += FString::FromInt(EntryKey->GlyphIndex);
+			Builder += TEXT("\nSpreadCategory: ");
+			Builder += LexToString(EntryKey->SpreadCategory);
+
+			FAtlasSlotInfo NewInfo;
+			NewInfo.AtlasSlotRect = FSlateRect(FVector2f(static_cast<float>(Slot->X), static_cast<float>(Slot->Y)), FVector2f(static_cast<float>(Slot->X + Slot->Width), static_cast<float>(Slot->Y + Slot->Height)));
+			NewInfo.TextureName = Builder.ToString();
+			return NewInfo;
+		}
+	}
+
+	return FAtlasSlotInfo();
+}
+#endif
+
 bool FSlateFontCache::AddNewEntry(const FShapedGlyphEntry& InShapedGlyph, const FFontOutlineSettings& InOutlineSettings, FShapedGlyphFontAtlasData& OutAtlasData)
 {
 	// Render the glyph
@@ -1881,7 +1997,6 @@ int32 FSlateFontCache::GetAllFontTexturesIndex(const int32 InIndex) const
 	{
 		return MsdfFontAtlasIndices[Index];
 	}
-	Index -= MsdfFontAtlasIndices.Num();
 	checkNoEntry();
 	return 0;
 }
diff --git a/Engine/Source/Runtime/SlateCore/Public/Fonts/FontCache.h b/Engine/Source/Runtime/SlateCore/Public/Fonts/FontCache.h
--- a/Engine/Source/Runtime/SlateCore/Public/Fonts/FontCache.h
+++ b/Engine/Source/Runtime/SlateCore/Public/Fonts/FontCache.h
@@ -130,7 +130,7 @@ struct FSdfGlyphFontAtlasData
 	FMetrics Metrics;
 	/** Index to a specific texture in the font cache. */
 	uint8 TextureIndex = 0;
-	/** True if the glyph is available in the face but sdf generation was not possible/successfull */
+	/** True if the glyph is available in the face but sdf generation was not possible/successful */
 	bool bSupportsSdf = false;
 	/** True if the SDF is a placeholder and the task for its final version hasn't been spawned yet (due to too many running tasks) */
 	bool bPendingRespawn = false;
@@ -761,7 +761,7 @@ public:
 	SLATECORE_API virtual FSlateShaderResource* GetAtlasPageResource(const int32 InIndex) const override;
 	SLATECORE_API virtual bool IsAtlasPageResourceAlphaOnly(const int32 InIndex) const override;
 #if WITH_ATLAS_DEBUGGING
-	virtual FAtlasSlotInfo GetAtlasSlotInfoAtPosition(FIntPoint InPosition, int32 AtlasIndex) const override { return FAtlasSlotInfo(); }
+	virtual FAtlasSlotInfo GetAtlasSlotInfoAtPosition(FIntPoint InPosition, int32 AtlasIndex) const override;
 #endif
 	/** 
 	 * Performs text shaping on the given string using the given font info. Returns you the shaped text sequence to use for text rendering via FSlateDrawElement::MakeShapedText.
