Enable AutoRTFM for this module

diff --git a/Engine/Source/Runtime/SlateCore/SlateCore.Build.cs b/Engine/Source/Runtime/SlateCore/SlateCore.Build.cs
--- a/Engine/Source/Runtime/SlateCore/SlateCore.Build.cs
+++ b/Engine/Source/Runtime/SlateCore/SlateCore.Build.cs
@@ -53,7 +53,5 @@ public class SlateCore : ModuleRules
 		{
 			AddEngineThirdPartyPrivateStaticDependencies(Target, "XInput");
 		}
-
-		bDisableAutoRTFMInstrumentation = true;
 	}
 }
