[Slate/SlateCore] Adding new OverflowPolicy MiddleEllipsis
This new OverflowPolicy will cut the names in the middle keeping the initial and last part visible.

#virtualized

diff --git a/Engine/Source/Runtime/Slate/Private/Framework/Text/SlateTextLayout.cpp b/Engine/Source/Runtime/Slate/Private/Framework/Text/SlateTextLayout.cpp
--- a/Engine/Source/Runtime/Slate/Private/Framework/Text/SlateTextLayout.cpp
+++ b/Engine/Source/Runtime/Slate/Private/Framework/Text/SlateTextLayout.cpp
@@ -9,6 +9,7 @@
 #include "Framework/Text/ISlateRunRenderer.h"
 #include "Framework/Text/ISlateLineHighlighter.h"
 #include "Framework/Text/SlateTextRun.h"
+#include "Framework/Text/SlateTextUtils.h"
 #include "Framework/Text/SlatePasswordRun.h"
 #include "Trace/SlateMemoryTags.h"
 
@@ -147,15 +148,19 @@ int32 FSlateTextLayout::OnPaint( const FPaintArgs& Args, const FGeometry& Allott
 
 		bool bIsLastVisibleLine = false;
 		bool bIsNextLineClipped = false;
-		if (LineOverflowPolicy == ETextOverflowPolicy::Ellipsis || LineOverflowPolicy == ETextOverflowPolicy::MultilineEllipsis)
+		if (SlateTextUtils::IsEllipsisPolicy(LineOverflowPolicy))
 		{
+			// Middle Ellipsis currently do not support multiline
 			if (bIsMultiline)
 			{
-				bIsLastVisibleLine = LineIndex == LastLineIndexToDisplay;
-				bIsNextLineClipped = LineViews.IsValidIndex(LineIndex + 1) ? bIsLastVisibleLine : false;
+				if (LineOverflowPolicy != ETextOverflowPolicy::MiddleEllipsis)
+				{
+					bIsLastVisibleLine = LineIndex == LastLineIndexToDisplay;
+					bIsNextLineClipped = LineViews.IsValidIndex(LineIndex + 1) ? bIsLastVisibleLine : false;
 
-				//When wrapping/multiline text, we have to use the reading direction of the text, not the justification.
-				LineOverflowDirection = LineView.TextBaseDirection == TextBiDi::ETextDirection::LeftToRight ? ETextOverflowDirection::LeftToRight : ETextOverflowDirection::RightToLeft;
+					//When wrapping/multiline text, we have to use the reading direction of the text, not the justification.
+					LineOverflowDirection = LineView.TextBaseDirection == TextBiDi::ETextDirection::LeftToRight ? ETextOverflowDirection::LeftToRight : ETextOverflowDirection::RightToLeft;
+				}
 			}
 			else
 			{
diff --git a/Engine/Source/Runtime/Slate/Private/Framework/Text/SlateTextRun.cpp b/Engine/Source/Runtime/Slate/Private/Framework/Text/SlateTextRun.cpp
--- a/Engine/Source/Runtime/Slate/Private/Framework/Text/SlateTextRun.cpp
+++ b/Engine/Source/Runtime/Slate/Private/Framework/Text/SlateTextRun.cpp
@@ -6,6 +6,7 @@
 #include "Framework/Application/SlateApplication.h"
 #include "Framework/Text/DefaultLayoutBlock.h"
 #include "Framework/Text/ShapedTextCache.h"
+#include "Framework/Text/SlateTextUtils.h"
 #include "Framework/Text/RunUtils.h"
 #include "Fonts/ShapedTextFwd.h"
 
@@ -119,10 +120,11 @@ int32 FSlateTextRun::OnPaint(const FPaintArgs& PaintArgs, const FTextArgs& TextA
 	);
 
 	FTextOverflowArgs OverflowArgs;
-	if ((TextArgs.OverflowPolicy == ETextOverflowPolicy::Ellipsis || TextArgs.OverflowPolicy == ETextOverflowPolicy::MultilineEllipsis) && TextArgs.OverflowDirection != ETextOverflowDirection::NoOverflow)
+	if (SlateTextUtils::IsEllipsisPolicy(TextArgs.OverflowPolicy) && TextArgs.OverflowDirection != ETextOverflowDirection::NoOverflow)
 	{
 		OverflowArgs.OverflowTextPtr = BlockTextContext.ShapedTextCache->FindOrAddOverflowEllipsisText(AllottedGeometry.GetAccumulatedLayoutTransform().GetScale(), BlockTextContext, Style.Font);
 		OverflowArgs.OverflowDirection = TextArgs.OverflowDirection;
+		OverflowArgs.OverflowPolicy = TextArgs.OverflowPolicy;
 		OverflowArgs.bIsLastVisibleBlock = TextArgs.bIsLastVisibleBlock;
 		OverflowArgs.bIsNextBlockClipped = TextArgs.bIsNextBlockClipped;
 	}
diff --git a/Engine/Source/Runtime/Slate/Private/Framework/Text/SlateTextUtils.cpp b/Engine/Source/Runtime/Slate/Private/Framework/Text/SlateTextUtils.cpp
new file mode 100644
--- /dev/null
+++ b/Engine/Source/Runtime/Slate/Private/Framework/Text/SlateTextUtils.cpp
@@ -0,0 +1,11 @@
+// Copyright Epic Games, Inc. All Rights Reserved.
+
+#include "Framework/Text/SlateTextUtils.h"
+#include "Styling/SlateTypes.h"
+
+bool SlateTextUtils::IsEllipsisPolicy(const ETextOverflowPolicy& InOverflowPolicy)
+{
+	return	InOverflowPolicy == ETextOverflowPolicy::Ellipsis ||
+			InOverflowPolicy == ETextOverflowPolicy::MultilineEllipsis ||
+			InOverflowPolicy == ETextOverflowPolicy::MiddleEllipsis;
+}
diff --git a/Engine/Source/Runtime/Slate/Public/Framework/Text/SlateTextUtils.h b/Engine/Source/Runtime/Slate/Public/Framework/Text/SlateTextUtils.h
new file mode 100644
--- /dev/null
+++ b/Engine/Source/Runtime/Slate/Public/Framework/Text/SlateTextUtils.h
@@ -0,0 +1,12 @@
+// Copyright Epic Games, Inc. All Rights Reserved.
+
+#pragma once
+#include "HAL/Platform.h"
+
+enum class ETextOverflowPolicy : uint8;
+
+namespace SlateTextUtils
+{
+	// Whether the given InOverflowPolicy is a valid Ellipsis policy
+	SLATE_API bool IsEllipsisPolicy(const ETextOverflowPolicy& InOverflowPolicy);
+}
diff --git a/Engine/Source/Runtime/SlateCore/Private/Rendering/ElementBatcher.cpp b/Engine/Source/Runtime/SlateCore/Private/Rendering/ElementBatcher.cpp
--- a/Engine/Source/Runtime/SlateCore/Private/Rendering/ElementBatcher.cpp
+++ b/Engine/Source/Runtime/SlateCore/Private/Rendering/ElementBatcher.cpp
@@ -1636,6 +1636,7 @@ void FSlateElementBatcher::AddShapedTextElement( const FSlateShapedTextElement&
 		BuildContext.bEnableCulling = false;
 		BuildContext.bForceEllipsis = DrawElement.OverflowArgs.bIsLastVisibleBlock && DrawElement.OverflowArgs.bIsNextBlockClipped;
 		BuildContext.OverflowDirection = DrawElement.OverflowArgs.OverflowDirection;
+		BuildContext.OverflowPolicy = DrawElement.OverflowArgs.OverflowPolicy;
 		BuildContext.MaxGlyphCountToRender = MaxGlyphCountToRender;
 
 		if (ShapedGlyphSequence->GetGlyphsToRender().Num() > 200 || (OverflowGlyphSequence && BuildContext.OverflowDirection != ETextOverflowDirection::NoOverflow))
@@ -3232,6 +3233,123 @@ const FSlateClippingState* FSlateElementBatcher::ResolveClippingState(const FSla
 	return nullptr;
 }
 
+void FSlateElementBatcher::CalculateMiddleEllipsisSkipIndexAndOffset(const FShapedTextBuildContext& InContext, ETextOverflowDirection InOverflowDirection, FMiddleEllipsisOverflowData& OutMiddleEllipsisData)
+{
+	const FShapedGlyphSequence* GlyphSequenceToRender = InContext.ShapedGlyphSequence;
+	const TArray<FShapedGlyphEntry>& GlyphsToRender = GlyphSequenceToRender->GetGlyphsToRender();
+	const int32 NumGlyphs = GlyphsToRender.Num();
+
+	float OverflowGlyphSequenceWidth = InContext.OverflowGlyphSequence ? InContext.OverflowGlyphSequence->GetMeasuredWidth() : 0.f;
+	float OverflowWidth = (GlyphSequenceToRender->GetMeasuredWidth() + OverflowGlyphSequenceWidth) - (InContext.LocalClipBoundingBoxRight - InContext.LocalClipBoundingBoxLeft);
+	int32 MiddleUp = GlyphsToRender.Num() >> 1;
+	int32 MiddleDown = MiddleUp;
+
+	while (OverflowWidth > 0.f && (MiddleDown >= 0 || MiddleUp < GlyphsToRender.Num()))
+	{
+		if (MiddleUp == MiddleDown)
+		{
+			OverflowWidth -= GlyphsToRender[MiddleUp].XAdvance;
+			OutMiddleEllipsisData.SkipIndexStart = OutMiddleEllipsisData.SkipIndexEnd = MiddleUp;
+			MiddleUp++;
+			MiddleDown--;
+		}
+		else
+		{
+			if (GlyphsToRender.IsValidIndex(MiddleDown))
+			{
+				OverflowWidth -= GlyphsToRender[MiddleDown].XAdvance;
+				OutMiddleEllipsisData.SkipIndexStart = MiddleDown;
+			}
+
+			if (OverflowWidth <= 0)
+			{
+				break;
+			}
+
+			if (GlyphsToRender.IsValidIndex(MiddleUp))
+			{
+				OverflowWidth -= GlyphsToRender[MiddleUp].XAdvance;
+				OutMiddleEllipsisData.SkipIndexEnd = MiddleUp;
+			}
+			MiddleUp++;
+			MiddleDown--;
+		}
+	}
+
+	// if the OverflowDirection is RightToLeft we need to calculate the offset of all clipped Character/Whitespaces to add to the LineX of the visible characters to align them correctly
+	if (InOverflowDirection == ETextOverflowDirection::RightToLeft && OutMiddleEllipsisData.SkipIndexStart != INDEX_NONE && OutMiddleEllipsisData.SkipIndexEnd != INDEX_NONE)
+	{
+		int32 ClippedIndexWidth = 0;
+		for (int32 GlyphToSkipIndex = OutMiddleEllipsisData.SkipIndexStart; GlyphToSkipIndex <= OutMiddleEllipsisData.SkipIndexEnd; ++GlyphToSkipIndex)
+		{
+			const FShapedGlyphEntry& GlyphToSkip = GlyphsToRender[GlyphToSkipIndex];
+			ClippedIndexWidth += GlyphToSkip.XAdvance;
+		}
+
+		// Get the previous and the next index of the skipped ones
+		int32 PreviousSkippedIndex = OutMiddleEllipsisData.SkipIndexStart - 1;
+		int32 NextSkippedIndex = OutMiddleEllipsisData.SkipIndexEnd + 1;
+
+		// True if the first clipped index is not less than 0 and the previous index are whitespaces/Non-visible Glyphs
+		bool bHaveWhiteSpaceBeforeAndIsValidIndex = GlyphsToRender.IsValidIndex(PreviousSkippedIndex);
+		// True if the last clipped index is not the last index of the GlyphsToRender and the next index(s) are whitespaces/Non-visible Glyphs
+		bool bHaveWhiteSpaceAfterAndIsValidIndex = GlyphsToRender.IsValidIndex(NextSkippedIndex);
+
+		// Calculate the X offset to add to each character to align them correctly at the right
+		while (bHaveWhiteSpaceBeforeAndIsValidIndex || bHaveWhiteSpaceAfterAndIsValidIndex)
+		{
+			if (bHaveWhiteSpaceBeforeAndIsValidIndex)
+			{
+				if (!GlyphsToRender[PreviousSkippedIndex].bIsVisible)
+				{
+					ClippedIndexWidth += GlyphsToRender[PreviousSkippedIndex].XAdvance;
+				}
+				else
+				{
+					bHaveWhiteSpaceBeforeAndIsValidIndex = false;
+				}
+				PreviousSkippedIndex--;
+				bHaveWhiteSpaceBeforeAndIsValidIndex &= GlyphsToRender.IsValidIndex(PreviousSkippedIndex);
+			}
+
+			if (bHaveWhiteSpaceAfterAndIsValidIndex)
+			{
+				if (!GlyphsToRender[NextSkippedIndex].bIsVisible)
+				{
+					ClippedIndexWidth += GlyphsToRender[NextSkippedIndex].XAdvance;
+				}
+				else
+				{
+					bHaveWhiteSpaceAfterAndIsValidIndex = false;
+				}
+				NextSkippedIndex++;
+				bHaveWhiteSpaceAfterAndIsValidIndex &= GlyphsToRender.IsValidIndex(NextSkippedIndex);
+			}
+		}
+
+		// BoundingBoxRight Edge - (GlyphSequence Width - ClippedGlyphs + EllipsisWidth)
+		const float FinalRightToLeftClippedOffset = InContext.LocalClipBoundingBoxRight - (GlyphSequenceToRender->GetMeasuredWidth() - ClippedIndexWidth + OverflowGlyphSequenceWidth);
+
+		// Start at the Offset needed to align the text to the right for the MiddleEllipsis
+		OutMiddleEllipsisData.LineX += FinalRightToLeftClippedOffset;
+	}
+
+	int32 SkippedIndexNum = OutMiddleEllipsisData.SkipIndexEnd - OutMiddleEllipsisData.SkipIndexStart + 1;
+
+	// If we are skipping all the Glyphs or all but 1, assign already the ellipsis X and Y since it may happen that the left side is all clipped, and they are not calculated correctly
+	if (SkippedIndexNum == NumGlyphs || SkippedIndexNum == NumGlyphs - 1)
+	{
+		OutMiddleEllipsisData.EllipsisLineX = OutMiddleEllipsisData.LineX;
+		OutMiddleEllipsisData.EllipsisLineY = OutMiddleEllipsisData.LineY;
+	}
+
+	// If the only visible Glyph is on the right of the ellipsis add the Ellipsis Width to the LineX to not paint the Glyph on top of the ellipsis
+	if (OutMiddleEllipsisData.SkipIndexStart == 0 && OutMiddleEllipsisData.SkipIndexEnd + 1 == NumGlyphs - 1)
+	{
+		OutMiddleEllipsisData.LineX += OverflowGlyphSequenceWidth;
+	}
+}
+
 template<ESlateVertexRounding Rounding>
 int32 FSlateElementBatcher::BuildShapedTextSequence(const FShapedTextBuildContext& Context)
 {
@@ -3256,6 +3374,7 @@ int32 FSlateElementBatcher::BuildShapedTextSequence(const FShapedTextBuildContex
 	FColor Tint = FColor::White;
 
 	ETextOverflowDirection OverflowDirection = Context.OverflowDirection;
+	ETextOverflowPolicy OverflowPolicy = Context.OverflowPolicy;
 
 	float EllipsisLineX = 0;
 	float EllipsisLineY = 0;
@@ -3277,8 +3396,41 @@ int32 FSlateElementBatcher::BuildShapedTextSequence(const FShapedTextBuildContex
 	const int32 NumGlyphs = GlyphSequenceToRender->GetGlyphsToRender().Num();
 	const TArray<FShapedGlyphEntry>& GlyphsToRender = GlyphSequenceToRender->GetGlyphsToRender();
 	int32 GlyphsRendered = 0;
+
+	// Middle ellipsis Section
+	const bool bWillBeClipped = GlyphSequenceToRender->GetMeasuredWidth() > (Context.LocalClipBoundingBoxRight - Context.LocalClipBoundingBoxLeft);
+	int32 SkipIndexStart = INDEX_NONE;
+	int32 SkipIndexEnd = INDEX_NONE;
+
+	// True by default, will be checked only after the middle skip and will become false as soon as we find a Visible character after the ellipsis, used to skip whitespaces after the ellipsis
+	bool bIsWhitespaceAfterMiddleSkip = true;
+
+	// Calculate the Glyph Indexes to skip to fit the ellipsis and add eventually the offset to the LineX to align the text correctly for the RightToLeft direction
+	if (OverflowPolicy == ETextOverflowPolicy::MiddleEllipsis && bWillBeClipped)
+	{
+		FMiddleEllipsisOverflowData OutMiddleEllipsisData;
+		OutMiddleEllipsisData.LineX = LineX;
+		OutMiddleEllipsisData.LineY = LineY;
+
+		CalculateMiddleEllipsisSkipIndexAndOffset(Context, OverflowDirection, OutMiddleEllipsisData);
+
+		// Assign the new data
+		SkipIndexStart = OutMiddleEllipsisData.SkipIndexStart;
+		SkipIndexEnd = OutMiddleEllipsisData.SkipIndexEnd;
+		LineX = OutMiddleEllipsisData.LineX;
+		LineY = OutMiddleEllipsisData.LineY;
+		EllipsisLineX = OutMiddleEllipsisData.EllipsisLineX;
+		EllipsisLineY = OutMiddleEllipsisData.EllipsisLineY;
+	}
+
 	for (int32 GlyphIndex = 0; GlyphIndex < NumGlyphs; ++GlyphIndex)
 	{
+		// Skip the index if its in between the Glyph to skip
+		if (GlyphIndex >= SkipIndexStart && GlyphIndex <= SkipIndexEnd)
+		{
+			continue;
+		}
+
 		const FShapedGlyphEntry& GlyphToRender = GlyphsToRender[GlyphIndex];
 
 		const float BitmapRenderScale = GlyphToRender.GetBitmapRenderScale();
@@ -3350,6 +3502,12 @@ int32 FSlateElementBatcher::BuildShapedTextSequence(const FShapedTextBuildContex
 
 			if (bCanRenderGlyph)
 			{
+				// If we can render a glyph after the middle skip it means that there are no more whitespaces to skip
+				if (GlyphIndex > SkipIndexEnd)
+				{
+					bIsWhitespaceAfterMiddleSkip = false;
+				}
+
 				// Note PosX,PosY is the upper left corner of the bounding box representing the string.  This computes the Y position of the baseline where text will sit
 				X = LineX + QuadMeshOffsets.X + (float)GlyphToRender.XOffset;
 				Y = LineY - QuadMeshOffsets.Y + (float)GlyphToRender.YOffset + ((Context.MaxHeight + Context.TextBaseline) * InvBitmapRenderScale);
@@ -3492,7 +3650,15 @@ int32 FSlateElementBatcher::BuildShapedTextSequence(const FShapedTextBuildContex
 		// Overflow Detection
 		// First figure out the size of the glyph. If the glyph contains multiple characters we have to measure all of them and if clipped, omit them all. This is common in complex languages with lots of diacritics
 		float OverflowTestWidth = SizeX;
-		if (OverflowDirection != ETextOverflowDirection::NoOverflow && (GlyphToRender.NumGraphemeClustersInGlyph > 1 || GlyphToRender.NumCharactersInGlyph > 1))
+		bool bMiddleEllipsisIsNextGlyphFirstSkipped = false;
+		int32 NextGlyphIndex = GlyphIndex + 1;
+		if (OverflowPolicy == ETextOverflowPolicy::MiddleEllipsis && bWillBeClipped && (NextGlyphIndex >= SkipIndexStart && NextGlyphIndex <= SkipIndexEnd))
+		{
+			bMiddleEllipsisIsNextGlyphFirstSkipped = true;
+			EllipsisLineX = LineX + GlyphToRender.XAdvance;
+			EllipsisLineY = LineY;
+		}
+		else if (OverflowDirection != ETextOverflowDirection::NoOverflow && (GlyphToRender.NumGraphemeClustersInGlyph > 1 || GlyphToRender.NumCharactersInGlyph > 1))
 		{
 			const int32 StartIndex = GlyphIndex;
 			int32 EndIndex = GlyphIndex;
@@ -3512,70 +3678,78 @@ int32 FSlateElementBatcher::BuildShapedTextSequence(const FShapedTextBuildContex
 		// Left to right overflow - If the current pen position + the ellipsis cannot fit, we have reached the end of the possible area for drawing this text. 
 		if (OverflowDirection == ETextOverflowDirection::LeftToRight)
 		{
-			// If we are on the last glyph don't bother checking if the ellipsis can fit. If the last glyph can fit there is no need for ellipsis
-			float OverflowSequenceNeededSize = GlyphIndex < NumGlyphs - 1 ? Context.OverflowGlyphSequence->GetMeasuredWidth() : 0;
-			if(X + OverflowTestWidth + OverflowSequenceNeededSize >= Context.LocalClipBoundingBoxRight || (Context.MaxGlyphCountToRender >= 0 && GlyphIndex >= Context.MaxGlyphCountToRender))
+			// MiddleEllipsis is handled in a different way above
+			if (OverflowPolicy != ETextOverflowPolicy::MiddleEllipsis)
 			{
-				bNeedEllipsis = true;
-				// We subtract out any whitespace advance. This avoids the ellipsis from ever floating out in the middle of a block of whitespace.
-				// e.g without this something like "The quick brown		fox jumps over the lazy dog" could be clipped to "The quick brown	..." but we want it to be "The quick brown..."
-				EllipsisLineX = LineX - PreviousWhitespaceAdvance;
-				EllipsisLineY = LineY;
-				// No characters to render after the ellipsis on the right side
-				break;
+				// If we are on the last glyph don't bother checking if the ellipsis can fit. If the last glyph can fit there is no need for ellipsis
+				float OverflowSequenceNeededSize = GlyphIndex < NumGlyphs - 1 ? Context.OverflowGlyphSequence->GetMeasuredWidth() : 0;
+				if(X + OverflowTestWidth + OverflowSequenceNeededSize >= Context.LocalClipBoundingBoxRight || (Context.MaxGlyphCountToRender >= 0 && GlyphIndex >= Context.MaxGlyphCountToRender))
+				{
+					bNeedEllipsis = true;
+					// We subtract out any whitespace advance. This avoids the ellipsis from ever floating out in the middle of a block of whitespace.
+					// e.g without this something like "The quick brown		fox jumps over the lazy dog" could be clipped to "The quick brown	..." but we want it to be "The quick brown..."
+					EllipsisLineX = LineX - PreviousWhitespaceAdvance;
+					EllipsisLineY = LineY;
+					// No characters to render after the ellipsis on the right side
+					break;
+				}
+				GlyphsRendered++;
 			}
-			GlyphsRendered++;
 		}
-		else if(OverflowDirection == ETextOverflowDirection::RightToLeft)
+		else if (OverflowDirection == ETextOverflowDirection::RightToLeft)
 		{
-			bool bClipped = false;
-			// Right to left overflow
-			if (X < Context.LocalClipBoundingBoxLeft || (Context.MaxGlyphCountToRender >= 0 && GlyphIndex < (NumGlyphs - Context.MaxGlyphCountToRender)))
+			// MiddleEllipsis is handled in a different way above
+			if (OverflowPolicy != ETextOverflowPolicy::MiddleEllipsis)
 			{
-				// This glyph is in the clipped region or is not visible so just advance. It cannot be shown
-				bClipped = true;
-				bNeedSpaceForEllipsis = true;
-			}
-			if (bNeedSpaceForEllipsis || !GlyphToRender.bIsVisible)
-			{
-				bClipped = true;
+				bool bClipped = false;
+				// Right to left overflow
+				if (X < Context.LocalClipBoundingBoxLeft || (Context.MaxGlyphCountToRender >= 0 && GlyphIndex < (NumGlyphs - Context.MaxGlyphCountToRender)))
+				{
+					// This glyph is in the clipped region or is not visible so just advance. It cannot be shown
+					bClipped = true;
+					bNeedSpaceForEllipsis = true;
+				}
+				if (bNeedSpaceForEllipsis || !GlyphToRender.bIsVisible)
+				{
+					bClipped = true;
 
-				// Can the ellipsis fit in the free spot by skipping the previous glyph(s)
-				const float EllipsisWidth = Context.OverflowGlyphSequence->GetMeasuredWidth();
-				const float AvailableX = X + GlyphToRender.XAdvance - Context.LocalClipBoundingBoxLeft;
-				if (AvailableX >= EllipsisWidth)
+					// Can the ellipsis fit in the free spot by skipping the previous glyph(s)
+					const float EllipsisWidth = Context.OverflowGlyphSequence->GetMeasuredWidth();
+					const float AvailableX = X + GlyphToRender.XAdvance - Context.LocalClipBoundingBoxLeft;
+					if (AvailableX >= EllipsisWidth)
+					{
+						// The available area can fit the ellipsis. Mark that we need an ellipsis and stop checking for overflow. The rest of the text can be built normally
+						bNeedSpaceForEllipsis = false;
+					}
+
+					//Always try to put the ellipsis, whether it fits or not: it's better to have an ellipsis a bit clipped than no feedback at all.
+					bNeedEllipsis = true;
+					EllipsisLineX = (LineX + GlyphToRender.XAdvance - EllipsisWidth);
+					EllipsisLineY = LineY;
+				}
+				else
 				{
-					// The available area can fit the ellipsis. Mark that we need an ellipsis and stop checking for overflow. The rest of the text can be built normally
-					bNeedSpaceForEllipsis = false;
+					OverflowDirection = ETextOverflowDirection::NoOverflow;
+					GlyphsRendered++;
 				}
 
-				//Always try to put the ellipsis, whether it fits or not: it's better to have an ellipsis a bit clipped than no feedback at all.
-				bNeedEllipsis = true;
-				EllipsisLineX = (LineX + GlyphToRender.XAdvance - EllipsisWidth);
-				EllipsisLineY = LineY;
-			}
-			else
-			{
-				OverflowDirection = ETextOverflowDirection::NoOverflow;
-				GlyphsRendered++;
-			}
+				// If we just clipped a glyph omit all characters in said glyph. Otherwise floating diacritics would be visible above the ellipsis. This is common in complex languages.
+				if (bClipped && GlyphToRender.NumCharactersInGlyph > 1)
+				{
+					GlyphIndex += GlyphToRender.NumCharactersInGlyph - 1;
+					LineX += GlyphToRender.XAdvance;
+					continue;
+				}
 
-			// If we just clipped a glyph omit all characters in said glyph. Otherwise floating diacritics would be visible above the ellipsis. This is common in complex languages.
-			if (bClipped && GlyphToRender.NumCharactersInGlyph > 1)
-			{
-				GlyphIndex += GlyphToRender.NumCharactersInGlyph - 1;
-				LineX += GlyphToRender.XAdvance;
-				continue;
+				bCanRenderGlyph = !bClipped;
 			}
-
-			bCanRenderGlyph = !bClipped;
 		}
 		else
 		{
 			GlyphsRendered++;
 		}
 
-		if(bCanRenderGlyph && RenderBatch)
+		if (bCanRenderGlyph && RenderBatch)
 		{
 			FVector2f UpperLeft(X, Y);
 			FVector2f UpperRight(X + SizeX, Y);
@@ -3647,9 +3821,30 @@ int32 FSlateElementBatcher::BuildShapedTextSequence(const FShapedTextBuildContex
 		{
 			// How much whitespace we are currently walking through
 			PreviousWhitespaceAdvance += GlyphToRender.XAdvance;
+
+			if (OverflowPolicy == ETextOverflowPolicy::MiddleEllipsis && bMiddleEllipsisIsNextGlyphFirstSkipped)
+			{
+				EllipsisLineX -= PreviousWhitespaceAdvance;
+				LineX -= PreviousWhitespaceAdvance;
+			}
+		}
+
+		if (OverflowPolicy == ETextOverflowPolicy::MiddleEllipsis)
+		{
+			// If the next Glyph is clipped add to the LineX the width of the Ellipsis
+			if (bMiddleEllipsisIsNextGlyphFirstSkipped)
+			{
+				LineX += Context.OverflowGlyphSequence->GetMeasuredWidth();
+			}
+
+			// If after the MiddleEllipsis there are still non-visible Glyphs, do not add their XAdvance to the LineX
+			LineX += (GlyphIndex > SkipIndexEnd && bIsWhitespaceAfterMiddleSkip) ? 0 : GlyphToRender.XAdvance;
+		}
+		else
+		{
+			LineX += GlyphToRender.XAdvance;
 		}
 
-		LineX += GlyphToRender.XAdvance;
 		LineY += GlyphToRender.YAdvance;
 	}
 
@@ -3659,6 +3854,10 @@ int32 FSlateElementBatcher::BuildShapedTextSequence(const FShapedTextBuildContex
 		EllipsisLineX = LineX - PreviousWhitespaceAdvance; //It's ok to always substract whitespace advance (LTR or RTL overflow direction), as the value will always be 0 for RTL overflow direction.
 		EllipsisLineY = LineY;
 	}
+	else if (OverflowPolicy == ETextOverflowPolicy::MiddleEllipsis)
+	{
+		bNeedEllipsis = bWillBeClipped;
+	}
 
 	if (bNeedEllipsis)
 	{
@@ -3669,6 +3868,7 @@ int32 FSlateElementBatcher::BuildShapedTextSequence(const FShapedTextBuildContex
 		EllipsisContext.OverflowGlyphSequence = nullptr;
 		EllipsisContext.bEnableCulling = false;
 		EllipsisContext.OverflowDirection = ETextOverflowDirection::NoOverflow;
+		EllipsisContext.OverflowPolicy = ETextOverflowPolicy::Clip;
 		EllipsisContext.StartLineX = EllipsisLineX;
 		EllipsisContext.StartLineY = EllipsisLineY;
 
diff --git a/Engine/Source/Runtime/SlateCore/Public/Rendering/DrawElementTextOverflowArgs.h b/Engine/Source/Runtime/SlateCore/Public/Rendering/DrawElementTextOverflowArgs.h
--- a/Engine/Source/Runtime/SlateCore/Public/Rendering/DrawElementTextOverflowArgs.h
+++ b/Engine/Source/Runtime/SlateCore/Public/Rendering/DrawElementTextOverflowArgs.h
@@ -3,6 +3,7 @@
 #pragma once
 
 #include "Fonts/ShapedTextFwd.h"
+#include "Styling/SlateTypes.h"
 
 enum class ETextOverflowDirection : uint8
 {
@@ -16,16 +17,25 @@ enum class ETextOverflowDirection : uint8
 
 struct FTextOverflowArgs
 {
+	FTextOverflowArgs(FShapedGlyphSequencePtr& InOverflowText, ETextOverflowDirection InOverflowDirection, ETextOverflowPolicy InOverflowPolicy)
+		: OverflowTextPtr(InOverflowText)
+		, OverflowDirection(InOverflowDirection)
+		, OverflowPolicy(InOverflowPolicy)
+		, bIsLastVisibleBlock(false)
+		, bIsNextBlockClipped(false)
+	{}
+
 	FTextOverflowArgs(FShapedGlyphSequencePtr& InOverflowText, ETextOverflowDirection InOverflowDirection)
 		: OverflowTextPtr(InOverflowText)
 		, OverflowDirection(InOverflowDirection)
+		, OverflowPolicy(ETextOverflowPolicy::Clip)
 		, bIsLastVisibleBlock(false)
 		, bIsNextBlockClipped(false)
-		
 	{}
 
 	FTextOverflowArgs()
 		: OverflowDirection(ETextOverflowDirection::NoOverflow)
+		, OverflowPolicy(ETextOverflowPolicy::Clip)
 		, bIsLastVisibleBlock(false)
 		, bIsNextBlockClipped(false)
 	{}
@@ -33,6 +43,7 @@ struct FTextOverflowArgs
 	/** Sequence that represents the ellipsis glyph */
 	FShapedGlyphSequencePtr OverflowTextPtr;
 	ETextOverflowDirection OverflowDirection;
+	ETextOverflowPolicy OverflowPolicy;
 	bool bIsLastVisibleBlock;
 	bool bIsNextBlockClipped;
 };
diff --git a/Engine/Source/Runtime/SlateCore/Public/Rendering/ElementBatcher.h b/Engine/Source/Runtime/SlateCore/Public/Rendering/ElementBatcher.h
--- a/Engine/Source/Runtime/SlateCore/Public/Rendering/ElementBatcher.h
+++ b/Engine/Source/Runtime/SlateCore/Public/Rendering/ElementBatcher.h
@@ -8,6 +8,7 @@
 #include "Layout/Clipping.h"
 #include "Stats/Stats.h"
 #include "SlateGlobals.h"
+#include "Styling/SlateTypes.h"
 #include "Containers/StaticArray.h"
 
 class FSlateBatchData;
@@ -436,6 +437,7 @@ private:
 		int32 LayerId;
 		FColor FontTint;
 		ETextOverflowDirection OverflowDirection;
+		ETextOverflowPolicy OverflowPolicy;
 		bool bEnableCulling : 1;
 		bool bForceEllipsis : 1;
 		
@@ -444,6 +446,29 @@ private:
 	template<ESlateVertexRounding Rounding>
 	int32 BuildShapedTextSequence(const FShapedTextBuildContext& Context);
 private:
+	/** struct containing info to use for the middle ellipsis policy */
+	struct FMiddleEllipsisOverflowData
+	{
+		FMiddleEllipsisOverflowData()
+			: SkipIndexStart(INDEX_NONE)
+			, SkipIndexEnd(INDEX_NONE)
+			, LineX(0.f)
+			, LineY(0.f)
+			, EllipsisLineX(0.f)
+			, EllipsisLineY(0.f)
+		{}
+
+		int32 SkipIndexStart;
+		int32 SkipIndexEnd;
+		float LineX;
+		float LineY;
+		float EllipsisLineX;
+		float EllipsisLineY;
+	};
+
+	/** Retrieve information for the middle ellipsis, like offset to use and Glyphs to skip */
+	void CalculateMiddleEllipsisSkipIndexAndOffset(const FShapedTextBuildContext& InContext, ETextOverflowDirection InOverflowDirection, FMiddleEllipsisOverflowData& OutMiddleEllipsisData);
+
 	/** Uncached Batch data currently being filled in */
 	FSlateBatchData* BatchData;
 	FSlateBatchData* BatchDataHDR;
diff --git a/Engine/Source/Runtime/SlateCore/Public/Styling/SlateTypes.h b/Engine/Source/Runtime/SlateCore/Public/Styling/SlateTypes.h
--- a/Engine/Source/Runtime/SlateCore/Public/Styling/SlateTypes.h
+++ b/Engine/Source/Runtime/SlateCore/Public/Styling/SlateTypes.h
@@ -85,6 +85,15 @@ enum class ETextOverflowPolicy : uint8
 
 	/** Overflowing text will be replaced with an ellipsis. A partially clipped line on the vertical axis will be totally clipped, and ellipsis displayed on previous line */
 	MultilineEllipsis,
+
+	/** Overflowing text will be replaced with an ellipsis starting from the center\n
+	 * Current Limits:
+	 * - Multiline is not supported
+	 * - RichText is not fully supported
+	 * - Highlight text is not supported
+	 * - Arabic mix with western character is not supported
+	 */
+	MiddleEllipsis
 };
 
 
