[SlateCore] Add check for the OverflowDirection for the MiddleEllipsis when checking if the space is enough.

diff --git a/Engine/Source/Runtime/SlateCore/Private/Rendering/ElementBatcher.cpp b/Engine/Source/Runtime/SlateCore/Private/Rendering/ElementBatcher.cpp
--- a/Engine/Source/Runtime/SlateCore/Private/Rendering/ElementBatcher.cpp
+++ b/Engine/Source/Runtime/SlateCore/Private/Rendering/ElementBatcher.cpp
@@ -1613,7 +1613,8 @@ void FSlateElementBatcher::AddShapedTextElement( const FSlateShapedTextElement&
 				// In checks below, ignore floating-point differences caused by transforming and untransforming the clip rect
 				const bool NeedLeftEllipsis = FMath::FloorToInt(BuildContext.LocalClipBoundingBoxLeft) > 0 && BuildContext.OverflowDirection == ETextOverflowDirection::RightToLeft;
 				const bool NeedRightEllipsis = ShapedGlyphSequence->GetMeasuredWidth() > FMath::CeilToInt(BuildContext.LocalClipBoundingBoxRight) && BuildContext.OverflowDirection == ETextOverflowDirection::LeftToRight;
-				if (!NeedLeftEllipsis && !NeedRightEllipsis && !DrawElement.OverflowArgs.bIsNextBlockClipped)
+				const bool NeedMiddleEllipsis = BuildContext.OverflowPolicy == ETextOverflowPolicy::MiddleEllipsis && (FMath::FloorToInt(BuildContext.LocalClipBoundingBoxLeft) > 0 || ShapedGlyphSequence->GetMeasuredWidth() > FMath::CeilToInt(BuildContext.LocalClipBoundingBoxRight));
+				if (!NeedLeftEllipsis && !NeedRightEllipsis && !NeedMiddleEllipsis && !DrawElement.OverflowArgs.bIsNextBlockClipped)
 				{
 					BuildContext.OverflowDirection = ETextOverflowDirection::NoOverflow;
 				}
@@ -3310,7 +3311,7 @@ int32 FSlateElementBatcher::BuildShapedTextSequence(const FShapedTextBuildContex
 	int32 GlyphsRendered = 0;
 
 	// Middle ellipsis Section
-	const bool bWillBeClipped = GlyphSequenceToRender->GetMeasuredWidth() > (Context.LocalClipBoundingBoxRight - Context.LocalClipBoundingBoxLeft);
+	const bool bWillBeClipped = GlyphSequenceToRender->GetMeasuredWidth() > (Context.LocalClipBoundingBoxRight - Context.LocalClipBoundingBoxLeft) && OverflowDirection != ETextOverflowDirection::NoOverflow;
 	int32 SkipIndexStart = INDEX_NONE;
 	int32 SkipIndexEnd = INDEX_NONE;
 
