[Slate]
* Implement default dtor (and cctor in some types) in cpp file to work around modular build linker issues on old console platform

diff --git a/Engine/Source/Runtime/Slate/Private/Widgets/Colors/SColorBlock.cpp b/Engine/Source/Runtime/Slate/Private/Widgets/Colors/SColorBlock.cpp
--- a/Engine/Source/Runtime/Slate/Private/Widgets/Colors/SColorBlock.cpp
+++ b/Engine/Source/Runtime/Slate/Private/Widgets/Colors/SColorBlock.cpp
@@ -31,6 +31,8 @@ SColorBlock::SColorBlock()
 {
 }
 
+SColorBlock::~SColorBlock() = default;
+
 /**
  * Construct this widget
  *
diff --git a/Engine/Source/Runtime/Slate/Private/Widgets/Colors/SColorGradingWheel.cpp b/Engine/Source/Runtime/Slate/Private/Widgets/Colors/SColorGradingWheel.cpp
--- a/Engine/Source/Runtime/Slate/Private/Widgets/Colors/SColorGradingWheel.cpp
+++ b/Engine/Source/Runtime/Slate/Private/Widgets/Colors/SColorGradingWheel.cpp
@@ -20,6 +20,8 @@ SColorGradingWheel::SColorGradingWheel()
 {
 }
 
+SColorGradingWheel::~SColorGradingWheel() = default;
+
 /* SColorGradingWheel methods
  *****************************************************************************/
 
diff --git a/Engine/Source/Runtime/Slate/Private/Widgets/Colors/SColorSpectrum.cpp b/Engine/Source/Runtime/Slate/Private/Widgets/Colors/SColorSpectrum.cpp
--- a/Engine/Source/Runtime/Slate/Private/Widgets/Colors/SColorSpectrum.cpp
+++ b/Engine/Source/Runtime/Slate/Private/Widgets/Colors/SColorSpectrum.cpp
@@ -10,6 +10,8 @@ SColorSpectrum::SColorSpectrum()
 {
 }
 
+SColorSpectrum::~SColorSpectrum() = default;
+
 /* SColorSpectrum methods
  *****************************************************************************/
 
diff --git a/Engine/Source/Runtime/Slate/Private/Widgets/Colors/SColorWheel.cpp b/Engine/Source/Runtime/Slate/Private/Widgets/Colors/SColorWheel.cpp
--- a/Engine/Source/Runtime/Slate/Private/Widgets/Colors/SColorWheel.cpp
+++ b/Engine/Source/Runtime/Slate/Private/Widgets/Colors/SColorWheel.cpp
@@ -8,6 +8,8 @@ SColorWheel::SColorWheel()
 	: SelectedColor(*this, FLinearColor(ForceInit))
 {}
 
+SColorWheel::~SColorWheel() = default;
+
 /* SColorWheel methods
  *****************************************************************************/
 
diff --git a/Engine/Source/Runtime/Slate/Private/Widgets/Colors/SComplexGradient.cpp b/Engine/Source/Runtime/Slate/Private/Widgets/Colors/SComplexGradient.cpp
--- a/Engine/Source/Runtime/Slate/Private/Widgets/Colors/SComplexGradient.cpp
+++ b/Engine/Source/Runtime/Slate/Private/Widgets/Colors/SComplexGradient.cpp
@@ -14,6 +14,8 @@ SComplexGradient::SComplexGradient()
 	SetCanTick(false);
 }
 
+SComplexGradient::~SComplexGradient() = default;
+
 /* SComplexGradient interface
  *****************************************************************************/
 
diff --git a/Engine/Source/Runtime/Slate/Private/Widgets/Images/SThrobber.cpp b/Engine/Source/Runtime/Slate/Private/Widgets/Images/SThrobber.cpp
--- a/Engine/Source/Runtime/Slate/Private/Widgets/Images/SThrobber.cpp
+++ b/Engine/Source/Runtime/Slate/Private/Widgets/Images/SThrobber.cpp
@@ -136,6 +136,8 @@ SCircularThrobber::SCircularThrobber()
 	: ColorAndOpacity(*this)
 {}
 
+SCircularThrobber::~SCircularThrobber() = default;
+
 void SCircularThrobber::Construct(const FArguments& InArgs)
 {
 	PieceImage = InArgs._PieceImage;
diff --git a/Engine/Source/Runtime/Slate/Private/Widgets/Input/SButton.cpp b/Engine/Source/Runtime/Slate/Private/Widgets/Input/SButton.cpp
--- a/Engine/Source/Runtime/Slate/Private/Widgets/Input/SButton.cpp
+++ b/Engine/Source/Runtime/Slate/Private/Widgets/Input/SButton.cpp
@@ -62,6 +62,8 @@ SButton::SButton()
 #endif
 }
 
+SButton::~SButton() = default;
+
 /**
  * Construct this widget
  *
diff --git a/Engine/Source/Runtime/Slate/Private/Widgets/Input/SCheckBox.cpp b/Engine/Source/Runtime/Slate/Private/Widgets/Input/SCheckBox.cpp
--- a/Engine/Source/Runtime/Slate/Private/Widgets/Input/SCheckBox.cpp
+++ b/Engine/Source/Runtime/Slate/Private/Widgets/Input/SCheckBox.cpp
@@ -20,6 +20,8 @@ SCheckBox::SCheckBox()
 #endif
 }
 
+SCheckBox::~SCheckBox() = default;
+
 /**
  * Construct this widget
  *
diff --git a/Engine/Source/Runtime/Slate/Private/Widgets/Input/SEditableTextBox.cpp b/Engine/Source/Runtime/Slate/Private/Widgets/Input/SEditableTextBox.cpp
--- a/Engine/Source/Runtime/Slate/Private/Widgets/Input/SEditableTextBox.cpp
+++ b/Engine/Source/Runtime/Slate/Private/Widgets/Input/SEditableTextBox.cpp
@@ -17,6 +17,8 @@ SEditableTextBox::SEditableTextBox()
 #endif
 }
 
+SEditableTextBox::~SEditableTextBox() = default;
+
 /**
  * Construct this widget
  *
diff --git a/Engine/Source/Runtime/Slate/Private/Widgets/Input/SSlider.cpp b/Engine/Source/Runtime/Slate/Private/Widgets/Input/SSlider.cpp
--- a/Engine/Source/Runtime/Slate/Private/Widgets/Input/SSlider.cpp
+++ b/Engine/Source/Runtime/Slate/Private/Widgets/Input/SSlider.cpp
@@ -35,6 +35,8 @@ SSlider::SSlider()
 #endif
 }
 
+SSlider::~SSlider() = default;
+
 void SSlider::Construct( const SSlider::FArguments& InDeclaration )
 {
 	check(InDeclaration._Style);
diff --git a/Engine/Source/Runtime/Slate/Private/Widgets/Layout/SBorder.cpp b/Engine/Source/Runtime/Slate/Private/Widgets/Layout/SBorder.cpp
--- a/Engine/Source/Runtime/Slate/Private/Widgets/Layout/SBorder.cpp
+++ b/Engine/Source/Runtime/Slate/Private/Widgets/Layout/SBorder.cpp
@@ -23,6 +23,8 @@ SBorder::SBorder()
 {
 }
 
+SBorder::~SBorder() = default;
+
 void SBorder::Construct( const SBorder::FArguments& InArgs )
 {
 	// Only do this if we're exactly an SBorder
diff --git a/Engine/Source/Runtime/Slate/Private/Widgets/Layout/SBox.cpp b/Engine/Source/Runtime/Slate/Private/Widgets/Layout/SBox.cpp
--- a/Engine/Source/Runtime/Slate/Private/Widgets/Layout/SBox.cpp
+++ b/Engine/Source/Runtime/Slate/Private/Widgets/Layout/SBox.cpp
@@ -34,6 +34,8 @@ SBox::SBox()
 	bCanSupportFocus = false;
 }
 
+SBox::~SBox() = default;
+
 void SBox::Construct( const FArguments& InArgs )
 {
 	SetWidthOverride(InArgs._WidthOverride);
diff --git a/Engine/Source/Runtime/Slate/Private/Widgets/Layout/SDPIScaler.cpp b/Engine/Source/Runtime/Slate/Private/Widgets/Layout/SDPIScaler.cpp
--- a/Engine/Source/Runtime/Slate/Private/Widgets/Layout/SDPIScaler.cpp
+++ b/Engine/Source/Runtime/Slate/Private/Widgets/Layout/SDPIScaler.cpp
@@ -19,6 +19,8 @@ SDPIScaler::SDPIScaler()
 	bHasRelativeLayoutScale = true;
 }
 
+SDPIScaler::~SDPIScaler() = default;
+
 void SDPIScaler::Construct( const FArguments& InArgs )
 {
 	SetDPIScale(InArgs._DPIScale);
diff --git a/Engine/Source/Runtime/Slate/Private/Widgets/Layout/SGridPanel.cpp b/Engine/Source/Runtime/Slate/Private/Widgets/Layout/SGridPanel.cpp
--- a/Engine/Source/Runtime/Slate/Private/Widgets/Layout/SGridPanel.cpp
+++ b/Engine/Source/Runtime/Slate/Private/Widgets/Layout/SGridPanel.cpp
@@ -49,6 +49,8 @@ SGridPanel::SGridPanel()
 	SetCanTick(false);
 }
 
+SGridPanel::~SGridPanel() = default;
+
 SGridPanel::FSlot::FSlotArguments SGridPanel::Slot(int32 Column, int32 Row, SGridPanel::Layer InLayer)
 {
 	return FSlot::FSlotArguments(MakeUnique<FSlot>(Column, Row, InLayer.TheLayer));
diff --git a/Engine/Source/Runtime/Slate/Private/Widgets/Layout/SResponsiveGridPanel.cpp b/Engine/Source/Runtime/Slate/Private/Widgets/Layout/SResponsiveGridPanel.cpp
--- a/Engine/Source/Runtime/Slate/Private/Widgets/Layout/SResponsiveGridPanel.cpp
+++ b/Engine/Source/Runtime/Slate/Private/Widgets/Layout/SResponsiveGridPanel.cpp
@@ -20,6 +20,8 @@ SResponsiveGridPanel::SResponsiveGridPanel()
 : Slots(this, GET_MEMBER_NAME_CHECKED(SResponsiveGridPanel, Slots))
 {}
 
+SResponsiveGridPanel::~SResponsiveGridPanel() = default;
+
 SResponsiveGridPanel::FScopedWidgetSlotArguments SResponsiveGridPanel::AddSlot(int32 Row)
 {
 	TWeakPtr<SResponsiveGridPanel> WeakPanel = SharedThis(this);
diff --git a/Engine/Source/Runtime/Slate/Private/Widgets/Layout/SScrollBox.cpp b/Engine/Source/Runtime/Slate/Private/Widgets/Layout/SScrollBox.cpp
--- a/Engine/Source/Runtime/Slate/Private/Widgets/Layout/SScrollBox.cpp
+++ b/Engine/Source/Runtime/Slate/Private/Widgets/Layout/SScrollBox.cpp
@@ -98,6 +98,8 @@ SScrollBox::SScrollBox()
 	bClippingProxy = true;
 }
 
+SScrollBox::~SScrollBox() = default;
+
 void SScrollBox::Construct( const FArguments& InArgs )
 {
 	check(InArgs._Style);
diff --git a/Engine/Source/Runtime/Slate/Private/Widgets/Layout/SUniformGridPanel.cpp b/Engine/Source/Runtime/Slate/Private/Widgets/Layout/SUniformGridPanel.cpp
--- a/Engine/Source/Runtime/Slate/Private/Widgets/Layout/SUniformGridPanel.cpp
+++ b/Engine/Source/Runtime/Slate/Private/Widgets/Layout/SUniformGridPanel.cpp
@@ -11,6 +11,8 @@ SUniformGridPanel::SUniformGridPanel()
 {
 }
 
+SUniformGridPanel::~SUniformGridPanel() = default;
+
 SUniformGridPanel::FSlot::FSlotArguments SUniformGridPanel::Slot(int32 Column, int32 Row)
 {
 	return FSlot::FSlotArguments(MakeUnique<FSlot>(Column, Row));
diff --git a/Engine/Source/Runtime/Slate/Private/Widgets/Layout/SUniformWrapPanel.cpp b/Engine/Source/Runtime/Slate/Private/Widgets/Layout/SUniformWrapPanel.cpp
--- a/Engine/Source/Runtime/Slate/Private/Widgets/Layout/SUniformWrapPanel.cpp
+++ b/Engine/Source/Runtime/Slate/Private/Widgets/Layout/SUniformWrapPanel.cpp
@@ -29,6 +29,8 @@ SUniformWrapPanel::SUniformWrapPanel()
 {
 }
 
+SUniformWrapPanel::~SUniformWrapPanel() = default;
+
 void SUniformWrapPanel::Construct( const FArguments& InArgs )
 {
 	SlotPadding.Assign(*this, InArgs._SlotPadding);
diff --git a/Engine/Source/Runtime/Slate/Private/Widgets/Layout/SWidgetSwitcher.cpp b/Engine/Source/Runtime/Slate/Private/Widgets/Layout/SWidgetSwitcher.cpp
--- a/Engine/Source/Runtime/Slate/Private/Widgets/Layout/SWidgetSwitcher.cpp
+++ b/Engine/Source/Runtime/Slate/Private/Widgets/Layout/SWidgetSwitcher.cpp
@@ -17,6 +17,8 @@ SWidgetSwitcher::SWidgetSwitcher()
 	SetCanTick(false);
 }
 
+SWidgetSwitcher::~SWidgetSwitcher() = default;
+
 SWidgetSwitcher::FSlot::FSlotArguments SWidgetSwitcher::Slot()
 {
 	return FSlot::FSlotArguments(MakeUnique<FSlot>());
diff --git a/Engine/Source/Runtime/Slate/Private/Widgets/Layout/SWindowTitleBarArea.cpp b/Engine/Source/Runtime/Slate/Private/Widgets/Layout/SWindowTitleBarArea.cpp
--- a/Engine/Source/Runtime/Slate/Private/Widgets/Layout/SWindowTitleBarArea.cpp
+++ b/Engine/Source/Runtime/Slate/Private/Widgets/Layout/SWindowTitleBarArea.cpp
@@ -67,6 +67,8 @@ SWindowTitleBarArea::SWindowTitleBarArea()
 	}
 }
 
+SWindowTitleBarArea::~SWindowTitleBarArea() = default;
+
 void SWindowTitleBarArea::Construct( const FArguments& InArgs )
 {
 	MinimizeButton = SNew(SButton)
diff --git a/Engine/Source/Runtime/Slate/Private/Widgets/Layout/SWrapBox.cpp b/Engine/Source/Runtime/Slate/Private/Widgets/Layout/SWrapBox.cpp
--- a/Engine/Source/Runtime/Slate/Private/Widgets/Layout/SWrapBox.cpp
+++ b/Engine/Source/Runtime/Slate/Private/Widgets/Layout/SWrapBox.cpp
@@ -20,6 +20,8 @@ SWrapBox::SWrapBox()
 {
 }
 
+SWrapBox::~SWrapBox() = default;
+
 SWrapBox::FScopedWidgetSlotArguments SWrapBox::AddSlot()
 {
 	return FScopedWidgetSlotArguments{ MakeUnique<FSlot>(), Slots, INDEX_NONE };
diff --git a/Engine/Source/Runtime/Slate/Private/Widgets/SCanvas.cpp b/Engine/Source/Runtime/Slate/Private/Widgets/SCanvas.cpp
--- a/Engine/Source/Runtime/Slate/Private/Widgets/SCanvas.cpp
+++ b/Engine/Source/Runtime/Slate/Private/Widgets/SCanvas.cpp
@@ -43,6 +43,8 @@ SCanvas::SCanvas()
 	bCanSupportFocus = false;
 }
 
+SCanvas::~SCanvas() = default;
+
 void SCanvas::Construct( const SCanvas::FArguments& InArgs )
 {
 	Children.AddSlots(MoveTemp(const_cast<TArray<FSlot::FSlotArguments>&>(InArgs._Slots)));
diff --git a/Engine/Source/Runtime/Slate/Private/Widgets/SViewport.cpp b/Engine/Source/Runtime/Slate/Private/Widgets/SViewport.cpp
--- a/Engine/Source/Runtime/Slate/Private/Widgets/SViewport.cpp
+++ b/Engine/Source/Runtime/Slate/Private/Widgets/SViewport.cpp
@@ -32,6 +32,7 @@ SViewport::SViewport()
 	, bEnableStereoRendering(false)
 { }
 
+SViewport::~SViewport() = default;
 
 /* SViewport interface
  *****************************************************************************/
diff --git a/Engine/Source/Runtime/Slate/Private/Widgets/SWeakWidget.cpp b/Engine/Source/Runtime/Slate/Private/Widgets/SWeakWidget.cpp
--- a/Engine/Source/Runtime/Slate/Private/Widgets/SWeakWidget.cpp
+++ b/Engine/Source/Runtime/Slate/Private/Widgets/SWeakWidget.cpp
@@ -10,6 +10,8 @@ SWeakWidget::SWeakWidget()
 {
 }
 
+SWeakWidget::~SWeakWidget() = default;
+
 void SWeakWidget::Construct(const FArguments& InArgs)
 {
 	SetCanTick(false);
diff --git a/Engine/Source/Runtime/Slate/Private/Widgets/Views/SExpanderArrow.cpp b/Engine/Source/Runtime/Slate/Private/Widgets/Views/SExpanderArrow.cpp
--- a/Engine/Source/Runtime/Slate/Private/Widgets/Views/SExpanderArrow.cpp
+++ b/Engine/Source/Runtime/Slate/Private/Widgets/Views/SExpanderArrow.cpp
@@ -18,6 +18,8 @@ SExpanderArrow::SExpanderArrow()
 
 }
 
+SExpanderArrow::~SExpanderArrow() = default;
+
 void SExpanderArrow::Construct( const FArguments& InArgs, const TSharedPtr<class ITableRow>& TableRow  )
 {
 	OwnerRowPtr = TableRow;
diff --git a/Engine/Source/Runtime/Slate/Private/Widgets/Views/STableViewBase.cpp b/Engine/Source/Runtime/Slate/Private/Widgets/Views/STableViewBase.cpp
--- a/Engine/Source/Runtime/Slate/Private/Widgets/Views/STableViewBase.cpp
+++ b/Engine/Source/Runtime/Slate/Private/Widgets/Views/STableViewBase.cpp
@@ -819,6 +819,8 @@ STableViewBase::STableViewBase( ETableViewMode::Type InTableViewMode )
 	FixedLineScrollOffset = 0.25f;
 }
 
+STableViewBase::~STableViewBase() = default;
+
 double STableViewBase::GetTargetScrollOffset() const
 {
 	if (FixedLineScrollOffset.IsSet() && !IsRightClickScrolling() && InertialScrollManager.GetScrollVelocity() == 0.f)
diff --git a/Engine/Source/Runtime/Slate/Public/Framework/MultiBox/MultiBoxBuilder.h b/Engine/Source/Runtime/Slate/Public/Framework/MultiBox/MultiBoxBuilder.h
--- a/Engine/Source/Runtime/Slate/Public/Framework/MultiBox/MultiBoxBuilder.h
+++ b/Engine/Source/Runtime/Slate/Public/Framework/MultiBox/MultiBoxBuilder.h
@@ -45,7 +45,7 @@ public:
 	 */
 	SLATE_API FMultiBoxBuilder( const EMultiBoxType InType, FMultiBoxCustomization InCustomization, const bool bInShouldCloseWindowAfterMenuSelection, const TSharedPtr< const FUICommandList >& InCommandList, TSharedPtr<FExtender> InExtender = TSharedPtr<FExtender>(), FName InTutorialHighlightName = NAME_None, FName InMenuName = NAME_None );
 
-	virtual ~FMultiBoxBuilder() {}
+	SLATE_API virtual ~FMultiBoxBuilder() {}
 
 	/**
 	 * Adds an editable text entry
diff --git a/Engine/Source/Runtime/Slate/Public/Widgets/Colors/SColorBlock.h b/Engine/Source/Runtime/Slate/Public/Widgets/Colors/SColorBlock.h
--- a/Engine/Source/Runtime/Slate/Public/Widgets/Colors/SColorBlock.h
+++ b/Engine/Source/Runtime/Slate/Public/Widgets/Colors/SColorBlock.h
@@ -71,6 +71,7 @@ public:
 
 public:
 	SLATE_API SColorBlock();
+	SLATE_API ~SColorBlock();
 
 	/**
 	 * Construct this widget
diff --git a/Engine/Source/Runtime/Slate/Public/Widgets/Colors/SColorGradingWheel.h b/Engine/Source/Runtime/Slate/Public/Widgets/Colors/SColorGradingWheel.h
--- a/Engine/Source/Runtime/Slate/Public/Widgets/Colors/SColorGradingWheel.h
+++ b/Engine/Source/Runtime/Slate/Public/Widgets/Colors/SColorGradingWheel.h
@@ -51,6 +51,7 @@ public:
 	
 public:
 	SLATE_API SColorGradingWheel();
+	SLATE_API virtual ~SColorGradingWheel();
 
 	/**
 	 * Construct this widget.
diff --git a/Engine/Source/Runtime/Slate/Public/Widgets/Colors/SColorSpectrum.h b/Engine/Source/Runtime/Slate/Public/Widgets/Colors/SColorSpectrum.h
--- a/Engine/Source/Runtime/Slate/Public/Widgets/Colors/SColorSpectrum.h
+++ b/Engine/Source/Runtime/Slate/Public/Widgets/Colors/SColorSpectrum.h
@@ -44,6 +44,7 @@ public:
 	
 public:
 	SLATE_API SColorSpectrum();
+	SLATE_API virtual ~SColorSpectrum();
 
 	/**
 	 * Construct this widget.
diff --git a/Engine/Source/Runtime/Slate/Public/Widgets/Colors/SColorWheel.h b/Engine/Source/Runtime/Slate/Public/Widgets/Colors/SColorWheel.h
--- a/Engine/Source/Runtime/Slate/Public/Widgets/Colors/SColorWheel.h
+++ b/Engine/Source/Runtime/Slate/Public/Widgets/Colors/SColorWheel.h
@@ -44,6 +44,7 @@ public:
 	
 public:
 	SLATE_API SColorWheel();
+	SLATE_API virtual ~SColorWheel();
 
 	/**
 	 * Construct this widget.
diff --git a/Engine/Source/Runtime/Slate/Public/Widgets/Colors/SComplexGradient.h b/Engine/Source/Runtime/Slate/Public/Widgets/Colors/SComplexGradient.h
--- a/Engine/Source/Runtime/Slate/Public/Widgets/Colors/SComplexGradient.h
+++ b/Engine/Source/Runtime/Slate/Public/Widgets/Colors/SComplexGradient.h
@@ -41,6 +41,7 @@ public:
 	SLATE_END_ARGS()
 	
 	SLATE_API SComplexGradient();
+	SLATE_API virtual ~SComplexGradient();
 
 	/**
 	 * Construct this widget
diff --git a/Engine/Source/Runtime/Slate/Public/Widgets/Images/SThrobber.h b/Engine/Source/Runtime/Slate/Public/Widgets/Images/SThrobber.h
--- a/Engine/Source/Runtime/Slate/Public/Widgets/Images/SThrobber.h
+++ b/Engine/Source/Runtime/Slate/Public/Widgets/Images/SThrobber.h
@@ -132,6 +132,7 @@ public:
 	SLATE_END_ARGS()
 
 	SLATE_API SCircularThrobber();
+	SLATE_API virtual ~SCircularThrobber();
 
 	/** Constructs the widget */
 	SLATE_API void Construct(const FArguments& InArgs);
diff --git a/Engine/Source/Runtime/Slate/Public/Widgets/Input/SButton.h b/Engine/Source/Runtime/Slate/Public/Widgets/Input/SButton.h
--- a/Engine/Source/Runtime/Slate/Public/Widgets/Input/SButton.h
+++ b/Engine/Source/Runtime/Slate/Public/Widgets/Input/SButton.h
@@ -120,6 +120,8 @@ public:
 
 	SLATE_END_ARGS()
 
+		SLATE_API virtual ~SButton();
+
 protected:
 		SLATE_API SButton();
 
diff --git a/Engine/Source/Runtime/Slate/Public/Widgets/Input/SCheckBox.h b/Engine/Source/Runtime/Slate/Public/Widgets/Input/SCheckBox.h
--- a/Engine/Source/Runtime/Slate/Public/Widgets/Input/SCheckBox.h
+++ b/Engine/Source/Runtime/Slate/Public/Widgets/Input/SCheckBox.h
@@ -154,6 +154,7 @@ public:
 	SLATE_END_ARGS()
 
 	SLATE_API SCheckBox();
+	SLATE_API virtual ~SCheckBox();
 
 	/**
 	 * Construct this widget
diff --git a/Engine/Source/Runtime/Slate/Public/Widgets/Input/SEditableTextBox.h b/Engine/Source/Runtime/Slate/Public/Widgets/Input/SEditableTextBox.h
--- a/Engine/Source/Runtime/Slate/Public/Widgets/Input/SEditableTextBox.h
+++ b/Engine/Source/Runtime/Slate/Public/Widgets/Input/SEditableTextBox.h
@@ -170,6 +170,7 @@ public:
 	SLATE_END_ARGS()
 
 	SLATE_API SEditableTextBox();
+	SLATE_API virtual ~SEditableTextBox();
 	
 	/**
 	 * Construct this widget
diff --git a/Engine/Source/Runtime/Slate/Public/Widgets/Input/SSlider.h b/Engine/Source/Runtime/Slate/Public/Widgets/Input/SSlider.h
--- a/Engine/Source/Runtime/Slate/Public/Widgets/Input/SSlider.h
+++ b/Engine/Source/Runtime/Slate/Public/Widgets/Input/SSlider.h
@@ -101,6 +101,7 @@ public:
 	SLATE_END_ARGS()
 
 	SLATE_API SSlider();
+	SLATE_API virtual ~SSlider();
 
 	/**
 	 * Construct the widget.
diff --git a/Engine/Source/Runtime/Slate/Public/Widgets/Layout/SBorder.h b/Engine/Source/Runtime/Slate/Public/Widgets/Layout/SBorder.h
--- a/Engine/Source/Runtime/Slate/Public/Widgets/Layout/SBorder.h
+++ b/Engine/Source/Runtime/Slate/Public/Widgets/Layout/SBorder.h
@@ -99,6 +99,7 @@ public:
 	 * Default constructor.
 	 */
 	SLATE_API SBorder();
+	SLATE_API virtual ~SBorder();
 
 	/**
 	 * Construct this widget
diff --git a/Engine/Source/Runtime/Slate/Public/Widgets/Layout/SBox.h b/Engine/Source/Runtime/Slate/Public/Widgets/Layout/SBox.h
--- a/Engine/Source/Runtime/Slate/Public/Widgets/Layout/SBox.h
+++ b/Engine/Source/Runtime/Slate/Public/Widgets/Layout/SBox.h
@@ -94,6 +94,7 @@ public:
 	SLATE_END_ARGS()
 
 	SLATE_API SBox();
+	SLATE_API virtual ~SBox();
 
 	SLATE_API void Construct(const FArguments& InArgs);
 
diff --git a/Engine/Source/Runtime/Slate/Public/Widgets/Layout/SDPIScaler.h b/Engine/Source/Runtime/Slate/Public/Widgets/Layout/SDPIScaler.h
--- a/Engine/Source/Runtime/Slate/Public/Widgets/Layout/SDPIScaler.h
+++ b/Engine/Source/Runtime/Slate/Public/Widgets/Layout/SDPIScaler.h
@@ -39,6 +39,7 @@ public:
 	SLATE_END_ARGS()
 
 	SLATE_API SDPIScaler();
+	SLATE_API virtual ~SDPIScaler();
 
 	SLATE_API void Construct( const FArguments& InArgs );
 
diff --git a/Engine/Source/Runtime/Slate/Public/Widgets/Layout/SGridPanel.h b/Engine/Source/Runtime/Slate/Public/Widgets/Layout/SGridPanel.h
--- a/Engine/Source/Runtime/Slate/Public/Widgets/Layout/SGridPanel.h
+++ b/Engine/Source/Runtime/Slate/Public/Widgets/Layout/SGridPanel.h
@@ -242,6 +242,7 @@ public:
 	SLATE_END_ARGS()
 
 	SLATE_API SGridPanel();
+	SLATE_API virtual ~SGridPanel();
 
 	/** Removes all slots from the panel */
 	SLATE_API void ClearChildren();
diff --git a/Engine/Source/Runtime/Slate/Public/Widgets/Layout/SResponsiveGridPanel.h b/Engine/Source/Runtime/Slate/Public/Widgets/Layout/SResponsiveGridPanel.h
--- a/Engine/Source/Runtime/Slate/Public/Widgets/Layout/SResponsiveGridPanel.h
+++ b/Engine/Source/Runtime/Slate/Public/Widgets/Layout/SResponsiveGridPanel.h
@@ -144,6 +144,7 @@ public:
 	SLATE_END_ARGS()
 
 	SLATE_API SResponsiveGridPanel();
+	SLATE_API virtual ~SResponsiveGridPanel();
 
 	/** Removes all rows from the panel */
 	SLATE_API void ClearChildren();
diff --git a/Engine/Source/Runtime/Slate/Public/Widgets/Layout/SScrollBox.h b/Engine/Source/Runtime/Slate/Public/Widgets/Layout/SScrollBox.h
--- a/Engine/Source/Runtime/Slate/Public/Widgets/Layout/SScrollBox.h
+++ b/Engine/Source/Runtime/Slate/Public/Widgets/Layout/SScrollBox.h
@@ -319,6 +319,7 @@ public:
 	SLATE_END_ARGS()
 
 	SLATE_API SScrollBox();
+	SLATE_API virtual ~SScrollBox();
 
 	/** @return a new slot. Slots contain children for SScrollBox */
 	static SLATE_API FSlot::FSlotArguments Slot();
diff --git a/Engine/Source/Runtime/Slate/Public/Widgets/Layout/SUniformGridPanel.h b/Engine/Source/Runtime/Slate/Public/Widgets/Layout/SUniformGridPanel.h
--- a/Engine/Source/Runtime/Slate/Public/Widgets/Layout/SUniformGridPanel.h
+++ b/Engine/Source/Runtime/Slate/Public/Widgets/Layout/SUniformGridPanel.h
@@ -72,6 +72,7 @@ public:
 	};
 
 	SLATE_API SUniformGridPanel();
+	SLATE_API virtual ~SUniformGridPanel();
 
 	/**
 	 * Used by declarative syntax to create a Slot in the specified Column, Row.
diff --git a/Engine/Source/Runtime/Slate/Public/Widgets/Layout/SUniformWrapPanel.h b/Engine/Source/Runtime/Slate/Public/Widgets/Layout/SUniformWrapPanel.h
--- a/Engine/Source/Runtime/Slate/Public/Widgets/Layout/SUniformWrapPanel.h
+++ b/Engine/Source/Runtime/Slate/Public/Widgets/Layout/SUniformWrapPanel.h
@@ -39,6 +39,7 @@ public:
 	};
 
 	SLATE_API SUniformWrapPanel();
+	SLATE_API virtual ~SUniformWrapPanel();
 
 	/**
 	 * Used by declarative syntax to create a Slot.
diff --git a/Engine/Source/Runtime/Slate/Public/Widgets/Layout/SWidgetSwitcher.h b/Engine/Source/Runtime/Slate/Public/Widgets/Layout/SWidgetSwitcher.h
--- a/Engine/Source/Runtime/Slate/Public/Widgets/Layout/SWidgetSwitcher.h
+++ b/Engine/Source/Runtime/Slate/Public/Widgets/Layout/SWidgetSwitcher.h
@@ -39,6 +39,7 @@ public:
 	SLATE_END_ARGS()
 
 	SLATE_API SWidgetSwitcher();
+	SLATE_API virtual ~SWidgetSwitcher();
 
 public:
 
diff --git a/Engine/Source/Runtime/Slate/Public/Widgets/Layout/SWindowTitleBarArea.h b/Engine/Source/Runtime/Slate/Public/Widgets/Layout/SWindowTitleBarArea.h
--- a/Engine/Source/Runtime/Slate/Public/Widgets/Layout/SWindowTitleBarArea.h
+++ b/Engine/Source/Runtime/Slate/Public/Widgets/Layout/SWindowTitleBarArea.h
@@ -52,6 +52,7 @@ public:
 	SLATE_END_ARGS()
 
 	SLATE_API SWindowTitleBarArea();
+	SLATE_API virtual ~SWindowTitleBarArea();
 
 	SLATE_API void Construct( const FArguments& InArgs );
 
diff --git a/Engine/Source/Runtime/Slate/Public/Widgets/Layout/SWrapBox.h b/Engine/Source/Runtime/Slate/Public/Widgets/Layout/SWrapBox.h
--- a/Engine/Source/Runtime/Slate/Public/Widgets/Layout/SWrapBox.h
+++ b/Engine/Source/Runtime/Slate/Public/Widgets/Layout/SWrapBox.h
@@ -162,6 +162,7 @@ public:
 	SLATE_END_ARGS()
 
 	SLATE_API SWrapBox();
+	SLATE_API virtual ~SWrapBox();
 
 	static FSlot::FSlotArguments Slot()
 	{
diff --git a/Engine/Source/Runtime/Slate/Public/Widgets/SCanvas.h b/Engine/Source/Runtime/Slate/Public/Widgets/SCanvas.h
--- a/Engine/Source/Runtime/Slate/Public/Widgets/SCanvas.h
+++ b/Engine/Source/Runtime/Slate/Public/Widgets/SCanvas.h
@@ -106,6 +106,7 @@ public:
 	SLATE_END_ARGS()
 
 	SLATE_API SCanvas();
+	SLATE_API virtual ~SCanvas();
 
 	/**
 	 * Construct this widget
diff --git a/Engine/Source/Runtime/Slate/Public/Widgets/SViewport.h b/Engine/Source/Runtime/Slate/Public/Widgets/SViewport.h
--- a/Engine/Source/Runtime/Slate/Public/Widgets/SViewport.h
+++ b/Engine/Source/Runtime/Slate/Public/Widgets/SViewport.h
@@ -84,6 +84,7 @@ public:
 
 	/** Default constructor. */
 	SLATE_API SViewport();
+	SLATE_API virtual ~SViewport();
 
 	/**
 	 * Construct the widget.
diff --git a/Engine/Source/Runtime/Slate/Public/Widgets/SWeakWidget.h b/Engine/Source/Runtime/Slate/Public/Widgets/SWeakWidget.h
--- a/Engine/Source/Runtime/Slate/Public/Widgets/SWeakWidget.h
+++ b/Engine/Source/Runtime/Slate/Public/Widgets/SWeakWidget.h
@@ -34,6 +34,7 @@ public:
 public:
 
 	SLATE_API SWeakWidget();
+	SLATE_API virtual ~SWeakWidget();
 
 	SLATE_API void Construct(const FArguments& InArgs);
 
diff --git a/Engine/Source/Runtime/Slate/Public/Widgets/Views/SExpanderArrow.h b/Engine/Source/Runtime/Slate/Public/Widgets/Views/SExpanderArrow.h
--- a/Engine/Source/Runtime/Slate/Public/Widgets/Views/SExpanderArrow.h
+++ b/Engine/Source/Runtime/Slate/Public/Widgets/Views/SExpanderArrow.h
@@ -39,6 +39,7 @@ public:
 	SLATE_END_ARGS()
 
 	SLATE_API SExpanderArrow();
+	SLATE_API ~SExpanderArrow();
 	SLATE_API void Construct( const FArguments& InArgs, const TSharedPtr<class ITableRow>& TableRow );
 
 protected:
diff --git a/Engine/Source/Runtime/Slate/Public/Widgets/Views/STableViewBase.h b/Engine/Source/Runtime/Slate/Public/Widgets/Views/STableViewBase.h
--- a/Engine/Source/Runtime/Slate/Public/Widgets/Views/STableViewBase.h
+++ b/Engine/Source/Runtime/Slate/Public/Widgets/Views/STableViewBase.h
@@ -255,6 +255,7 @@ public:
 protected:
 
 	SLATE_API STableViewBase( ETableViewMode::Type InTableViewMode );
+	SLATE_API virtual ~STableViewBase();
 
 	/** Returns the "true" scroll offset where the list will ultimately settle (and may already be). */
 	SLATE_API double GetTargetScrollOffset() const;
diff --git a/Engine/Source/Runtime/SlateCore/Private/SlateCoreClasses.cpp b/Engine/Source/Runtime/SlateCore/Private/SlateCoreClasses.cpp
--- a/Engine/Source/Runtime/SlateCore/Private/SlateCoreClasses.cpp
+++ b/Engine/Source/Runtime/SlateCore/Private/SlateCoreClasses.cpp
@@ -96,5 +96,6 @@ FAutoConsoleVariableRef CVarHitTestGridDebugging(
 
 #endif
 
-FSlateWidgetStyle::FSlateWidgetStyle()
-{ }
+FSlateWidgetStyle::FSlateWidgetStyle() = default;
+FSlateWidgetStyle::~FSlateWidgetStyle() = default;
+
diff --git a/Engine/Source/Runtime/SlateCore/Private/Styling/SlateTypes.cpp b/Engine/Source/Runtime/SlateCore/Private/Styling/SlateTypes.cpp
--- a/Engine/Source/Runtime/SlateCore/Private/Styling/SlateTypes.cpp
+++ b/Engine/Source/Runtime/SlateCore/Private/Styling/SlateTypes.cpp
@@ -36,6 +36,9 @@ FCheckBoxStyle::FCheckBoxStyle()
 {
 }
 
+FCheckBoxStyle::FCheckBoxStyle(const FCheckBoxStyle&) = default;
+FCheckBoxStyle::~FCheckBoxStyle() = default;
+
 const FName FCheckBoxStyle::TypeName( TEXT("FCheckBoxStyle") );
 
 const FCheckBoxStyle& FCheckBoxStyle::GetDefault()
@@ -85,6 +88,10 @@ FTextBlockStyle::FTextBlockStyle()
 {
 }
 
+FTextBlockStyle::FTextBlockStyle(const FTextBlockStyle&) = default;
+FTextBlockStyle::~FTextBlockStyle() = default; // To get around linker errors on one old platform
+
+
 const FName FTextBlockStyle::TypeName( TEXT("FTextBlockStyle") );
 
 void FTextBlockStyle::GetResources(TArray< const FSlateBrush* >& OutBrushes) const
@@ -119,6 +126,9 @@ FButtonStyle::FButtonStyle()
 	Disabled = FSlateNoResource();
 }
 
+FButtonStyle::FButtonStyle(const FButtonStyle&) = default;
+FButtonStyle::~FButtonStyle() = default;
+
 const FName FButtonStyle::TypeName( TEXT("FButtonStyle") );
 
 void FButtonStyle::GetResources( TArray< const FSlateBrush* >& OutBrushes ) const
@@ -157,6 +167,9 @@ FComboButtonStyle::FComboButtonStyle()
 {
 }
 
+FComboButtonStyle::FComboButtonStyle(const FComboButtonStyle&) = default;
+FComboButtonStyle::~FComboButtonStyle() = default;
+
 const FName FComboButtonStyle::TypeName( TEXT("FComboButtonStyle") );
 
 void FComboButtonStyle::GetResources( TArray< const FSlateBrush* >& OutBrushes ) const
@@ -179,6 +192,9 @@ FComboBoxStyle::FComboBoxStyle()
 	ComboButtonStyle.SetMenuBorderPadding(FMargin(1.0));
 }
 
+FComboBoxStyle::FComboBoxStyle(const FComboBoxStyle&) = default;
+FComboBoxStyle::~FComboBoxStyle() = default;
+
 const FName FComboBoxStyle::TypeName( TEXT("FComboBoxStyle") );
 
 void FComboBoxStyle::GetResources( TArray< const FSlateBrush* >& OutBrushes ) const
@@ -211,6 +227,9 @@ FHyperlinkStyle::FHyperlinkStyle()
 {
 }
 
+FHyperlinkStyle::FHyperlinkStyle(const FHyperlinkStyle&) = default;
+FHyperlinkStyle::~FHyperlinkStyle() = default;
+
 const FName FHyperlinkStyle::TypeName( TEXT("FHyperlinkStyle") );
 
 void FHyperlinkStyle::GetResources( TArray< const FSlateBrush* >& OutBrushes ) const
@@ -233,6 +252,9 @@ FEditableTextStyle::FEditableTextStyle()
 {
 }
 
+FEditableTextStyle::FEditableTextStyle(const FEditableTextStyle&) = default;
+FEditableTextStyle::~FEditableTextStyle() = default;
+
 void FEditableTextStyle::GetResources( TArray< const FSlateBrush* >& OutBrushes ) const
 {
 	OutBrushes.Add( &BackgroundImageSelected );
@@ -268,6 +290,12 @@ FEditableTextBoxStyle::FEditableTextBoxStyle()
 {
 }
 
+PRAGMA_DISABLE_DEPRECATION_WARNINGS
+FEditableTextBoxStyle::FEditableTextBoxStyle(const FEditableTextBoxStyle&) = default;
+FEditableTextBoxStyle& FEditableTextBoxStyle::operator=(const FEditableTextBoxStyle&) = default;
+FEditableTextBoxStyle::~FEditableTextBoxStyle() = default;
+PRAGMA_ENABLE_DEPRECATION_WARNINGS
+
 void FEditableTextBoxStyle::GetResources( TArray< const FSlateBrush* >& OutBrushes ) const
 {
 	OutBrushes.Add( &BackgroundImageNormal );
@@ -291,6 +319,8 @@ FInlineEditableTextBlockStyle::FInlineEditableTextBlockStyle()
 {
 }
 
+FInlineEditableTextBlockStyle::~FInlineEditableTextBlockStyle() = default;
+
 void FInlineEditableTextBlockStyle::GetResources( TArray< const FSlateBrush* >& OutBrushes ) const
 {
 	EditableTextBoxStyle.GetResources(OutBrushes);
@@ -311,6 +341,8 @@ FProgressBarStyle::FProgressBarStyle()
 {
 }
 
+FProgressBarStyle::~FProgressBarStyle() = default;
+
 void FProgressBarStyle::GetResources( TArray< const FSlateBrush* >& OutBrushes ) const
 {
 	OutBrushes.Add( &BackgroundImage );
@@ -341,6 +373,9 @@ FScrollBarStyle::FScrollBarStyle()
 {
 }
 
+FScrollBarStyle::FScrollBarStyle(const FScrollBarStyle&) = default;
+FScrollBarStyle::~FScrollBarStyle() = default;
+
 void FScrollBarStyle::GetResources( TArray< const FSlateBrush* >& OutBrushes ) const
 {
 	OutBrushes.Add( &HorizontalBackgroundImage );
@@ -368,6 +403,8 @@ FExpandableAreaStyle::FExpandableAreaStyle()
 {
 }
 
+FExpandableAreaStyle::~FExpandableAreaStyle() = default;
+
 void FExpandableAreaStyle::GetResources( TArray< const FSlateBrush* >& OutBrushes ) const
 {
 	OutBrushes.Add( &CollapsedImage );
@@ -390,6 +427,8 @@ FSearchBoxStyle::FSearchBoxStyle()
 {
 }
 
+FSearchBoxStyle::~FSearchBoxStyle() = default;
+
 FSearchBoxStyle& FSearchBoxStyle::SetTextBoxStyle( const FEditableTextBoxStyle& InTextBoxStyle )
 { 
 	TextBoxStyle = InTextBoxStyle;
@@ -432,6 +471,8 @@ FSliderStyle::FSliderStyle()
 {
 }
 
+FSliderStyle::~FSliderStyle() = default;
+
 void FSliderStyle::GetResources( TArray< const FSlateBrush* >& OutBrushes ) const
 {
 	OutBrushes.Add( &NormalBarImage );
@@ -455,6 +496,8 @@ FVolumeControlStyle::FVolumeControlStyle()
 {
 }
 
+FVolumeControlStyle::~FVolumeControlStyle() = default;
+
 void FVolumeControlStyle::GetResources( TArray< const FSlateBrush* >& OutBrushes ) const
 {
 	SliderStyle.GetResources(OutBrushes);
@@ -479,6 +522,9 @@ FInlineTextImageStyle::FInlineTextImageStyle()
 {
 }
 
+FInlineTextImageStyle::FInlineTextImageStyle(const FInlineTextImageStyle&) = default;
+FInlineTextImageStyle::~FInlineTextImageStyle() = default;
+
 void FInlineTextImageStyle::GetResources( TArray< const FSlateBrush* >& OutBrushes ) const
 {
 	OutBrushes.Add( &Image );
@@ -501,6 +547,8 @@ FSpinBoxStyle::FSpinBoxStyle()
 {
 }
 
+FSpinBoxStyle::~FSpinBoxStyle() = default;
+
 void FSpinBoxStyle::GetResources( TArray< const FSlateBrush* >& OutBrushes ) const
 {
 	OutBrushes.Add(&BackgroundBrush);
@@ -525,6 +573,8 @@ FSplitterStyle::FSplitterStyle()
 {
 }
 
+FSplitterStyle::~FSplitterStyle() = default;
+
 void FSplitterStyle::GetResources( TArray< const FSlateBrush* >& OutBrushes ) const
 {
 	OutBrushes.Add( &HandleNormalBrush );
@@ -544,6 +594,8 @@ FTableViewStyle::FTableViewStyle()
 {
 }
 
+FTableViewStyle::~FTableViewStyle() = default;
+
 void FTableViewStyle::GetResources( TArray< const FSlateBrush* >& OutBrushes ) const
 {
 	OutBrushes.Add( &BackgroundBrush );	
@@ -564,6 +616,8 @@ FTableRowStyle::FTableRowStyle()
 {
 }
 
+FTableRowStyle::~FTableRowStyle() = default;
+
 void FTableRowStyle::GetResources( TArray< const FSlateBrush* >& OutBrushes ) const
 {
 	OutBrushes.Add( &SelectorFocusedBrush );
@@ -595,6 +649,8 @@ FTableColumnHeaderStyle::FTableColumnHeaderStyle()
 {
 }
 
+FTableColumnHeaderStyle::~FTableColumnHeaderStyle() = default;
+
 void FTableColumnHeaderStyle::GetResources( TArray< const FSlateBrush* >& OutBrushes ) const
 {
 	OutBrushes.Add( &SortPrimaryAscendingImage );
@@ -624,6 +680,8 @@ FHeaderRowStyle::FHeaderRowStyle()
 {
 }
 
+FHeaderRowStyle::~FHeaderRowStyle() = default;
+
 void FHeaderRowStyle::GetResources( TArray< const FSlateBrush* >& OutBrushes ) const
 {
 	ColumnStyle.GetResources(OutBrushes);
@@ -654,6 +712,8 @@ FDockTabStyle::FDockTabStyle()
 {
 }
 
+FDockTabStyle::~FDockTabStyle() = default;
+
 void FDockTabStyle::GetResources( TArray< const FSlateBrush* >& OutBrushes ) const
 {
 	CloseButtonStyle.GetResources(OutBrushes);
@@ -680,6 +740,8 @@ FScrollBoxStyle::FScrollBoxStyle()
 {
 }
 
+FScrollBoxStyle::~FScrollBoxStyle() = default;
+
 void FScrollBoxStyle::GetResources( TArray< const FSlateBrush* >& OutBrushes ) const
 {
 	OutBrushes.Add( &TopShadowBrush );
@@ -701,6 +763,8 @@ FScrollBorderStyle::FScrollBorderStyle()
 {
 }
 
+FScrollBorderStyle::~FScrollBorderStyle() = default;
+
 void FScrollBorderStyle::GetResources( TArray< const FSlateBrush* >& OutBrushes ) const
 {
 	OutBrushes.Add( &TopShadowBrush );
@@ -725,6 +789,9 @@ FWindowStyle::FWindowStyle()
 {
 }
 
+FWindowStyle::FWindowStyle(const FWindowStyle&) = default;
+FWindowStyle::~FWindowStyle() = default;
+
 void FWindowStyle::GetResources(TArray< const FSlateBrush* >& OutBrushes) const
 {
 	MinimizeButtonStyle.GetResources(OutBrushes);
diff --git a/Engine/Source/Runtime/SlateCore/Private/Styling/ToolBarStyle.cpp b/Engine/Source/Runtime/SlateCore/Private/Styling/ToolBarStyle.cpp
--- a/Engine/Source/Runtime/SlateCore/Private/Styling/ToolBarStyle.cpp
+++ b/Engine/Source/Runtime/SlateCore/Private/Styling/ToolBarStyle.cpp
@@ -32,6 +32,10 @@ FToolBarStyle::FToolBarStyle()
 	, bShowLabels(true)
 {}
 
+FToolBarStyle::FToolBarStyle(const FToolBarStyle&) = default;
+
+FToolBarStyle::~FToolBarStyle() = default;
+
 const FToolBarStyle& FToolBarStyle::GetDefault()
 {
 	static FToolBarStyle Default;
diff --git a/Engine/Source/Runtime/SlateCore/Private/Widgets/Images/SImage.cpp b/Engine/Source/Runtime/SlateCore/Private/Widgets/Images/SImage.cpp
--- a/Engine/Source/Runtime/SlateCore/Private/Widgets/Images/SImage.cpp
+++ b/Engine/Source/Runtime/SlateCore/Private/Widgets/Images/SImage.cpp
@@ -25,6 +25,8 @@ SImage::SImage()
 	bCanSupportFocus = false;
 }
 
+SImage::~SImage() = default;
+
 void SImage::Construct( const FArguments& InArgs )
 {
 	ImageAttribute.Assign(*this, InArgs._Image);
diff --git a/Engine/Source/Runtime/SlateCore/Private/Widgets/SCompoundWidget.cpp b/Engine/Source/Runtime/SlateCore/Private/Widgets/SCompoundWidget.cpp
--- a/Engine/Source/Runtime/SlateCore/Private/Widgets/SCompoundWidget.cpp
+++ b/Engine/Source/Runtime/SlateCore/Private/Widgets/SCompoundWidget.cpp
@@ -86,6 +86,8 @@ SCompoundWidget::SCompoundWidget()
 {
 }
 
+SCompoundWidget::~SCompoundWidget() = default;
+
 void SCompoundWidget::SetVisibility( TAttribute<EVisibility> InVisibility )
 {
 	SWidget::SetVisibility(MoveTemp(InVisibility));
diff --git a/Engine/Source/Runtime/SlateCore/Private/Widgets/SOverlay.cpp b/Engine/Source/Runtime/SlateCore/Private/Widgets/SOverlay.cpp
--- a/Engine/Source/Runtime/SlateCore/Private/Widgets/SOverlay.cpp
+++ b/Engine/Source/Runtime/SlateCore/Private/Widgets/SOverlay.cpp
@@ -115,6 +115,8 @@ SOverlay::SOverlay()
 	bCanSupportFocus = false;
 }
 
+SOverlay::~SOverlay() = default;
+
 SOverlay::FOverlaySlot::FSlotArguments SOverlay::Slot( int32 ZOrder )
 {
 	return FOverlaySlot::FSlotArguments(MakeUnique<FOverlaySlot>(ZOrder));
diff --git a/Engine/Source/Runtime/SlateCore/Private/Widgets/SPanel.cpp b/Engine/Source/Runtime/SlateCore/Private/Widgets/SPanel.cpp
--- a/Engine/Source/Runtime/SlateCore/Private/Widgets/SPanel.cpp
+++ b/Engine/Source/Runtime/SlateCore/Private/Widgets/SPanel.cpp
@@ -4,6 +4,9 @@
 #include "Types/PaintArgs.h"
 #include "Layout/ArrangedChildren.h"
 
+SPanel::SPanel() = default;
+SPanel::~SPanel() = default;
+
 int32 SPanel::OnPaint( const FPaintArgs& Args, const FGeometry& AllottedGeometry, const FSlateRect& MyCullingRect, FSlateWindowElementList& OutDrawElements, int32 LayerId, const FWidgetStyle& InWidgetStyle, bool bParentEnabled ) const
 {
 	FArrangedChildren ArrangedChildren(EVisibility::Visible);
diff --git a/Engine/Source/Runtime/SlateCore/Public/Styling/SlateTypes.h b/Engine/Source/Runtime/SlateCore/Public/Styling/SlateTypes.h
--- a/Engine/Source/Runtime/SlateCore/Public/Styling/SlateTypes.h
+++ b/Engine/Source/Runtime/SlateCore/Public/Styling/SlateTypes.h
@@ -106,8 +106,8 @@ struct FCheckBoxStyle : public FSlateWidgetStyle
 	GENERATED_USTRUCT_BODY()
 
 	SLATECORE_API FCheckBoxStyle();
-
-	virtual ~FCheckBoxStyle() {}
+	SLATECORE_API FCheckBoxStyle(const FCheckBoxStyle&);
+	SLATECORE_API virtual ~FCheckBoxStyle();
 
 	SLATECORE_API virtual void GetResources( TArray< const FSlateBrush* > & OutBrushes ) const override;
 
@@ -327,8 +327,8 @@ struct FTextBlockStyle : public FSlateWidgetStyle
 	GENERATED_USTRUCT_BODY()
 
 	SLATECORE_API FTextBlockStyle();
-
-	virtual ~FTextBlockStyle() {}
+	SLATECORE_API FTextBlockStyle(const FTextBlockStyle&);
+	SLATECORE_API virtual ~FTextBlockStyle();
 
 	SLATECORE_API virtual void GetResources( TArray< const FSlateBrush* >& OutBrushes ) const override;
 
@@ -500,8 +500,9 @@ struct FButtonStyle : public FSlateWidgetStyle
 	GENERATED_USTRUCT_BODY()
 
 	SLATECORE_API FButtonStyle();
+	SLATECORE_API FButtonStyle(const FButtonStyle&);
 
-	virtual ~FButtonStyle() {}
+	SLATECORE_API virtual ~FButtonStyle();
 
 	SLATECORE_API virtual void GetResources( TArray< const FSlateBrush* >& OutBrushes ) const override;
 
@@ -628,8 +629,8 @@ struct FComboButtonStyle : public FSlateWidgetStyle
 	GENERATED_USTRUCT_BODY()
 
 	SLATECORE_API FComboButtonStyle();
-
-	virtual ~FComboButtonStyle() {}
+	SLATECORE_API FComboButtonStyle(const FComboButtonStyle&);
+	SLATECORE_API virtual ~FComboButtonStyle();
 
 	SLATECORE_API virtual void GetResources( TArray< const FSlateBrush* >& OutBrushes ) const override;
 
@@ -725,8 +726,8 @@ struct FComboBoxStyle : public FSlateWidgetStyle
 	GENERATED_USTRUCT_BODY()
 
 	SLATECORE_API FComboBoxStyle();
-
-	virtual ~FComboBoxStyle() {}
+	SLATECORE_API FComboBoxStyle(const FComboBoxStyle&);
+	SLATECORE_API virtual ~FComboBoxStyle();
 
 	SLATECORE_API virtual void GetResources( TArray< const FSlateBrush* >& OutBrushes ) const override;
 
@@ -815,8 +816,9 @@ struct FHyperlinkStyle : public FSlateWidgetStyle
 	GENERATED_USTRUCT_BODY()
 
 	SLATECORE_API FHyperlinkStyle();
+	SLATECORE_API FHyperlinkStyle(const FHyperlinkStyle&);
 
-	virtual ~FHyperlinkStyle() {}
+	SLATECORE_API virtual ~FHyperlinkStyle();
 
 	SLATECORE_API virtual void GetResources( TArray< const FSlateBrush* >& OutBrushes ) const override;
 
@@ -850,8 +852,8 @@ struct FEditableTextStyle : public FSlateWidgetStyle
 	GENERATED_USTRUCT_BODY()
 
 	SLATECORE_API FEditableTextStyle();
-
-	virtual ~FEditableTextStyle() {}
+	SLATECORE_API FEditableTextStyle(const FEditableTextStyle&);
+	SLATECORE_API virtual ~FEditableTextStyle();
 
 	SLATECORE_API virtual void GetResources( TArray< const FSlateBrush* >& OutBrushes ) const override;
 
@@ -911,8 +913,8 @@ struct FScrollBarStyle : public FSlateWidgetStyle
 	GENERATED_USTRUCT_BODY()
 
 	SLATECORE_API FScrollBarStyle();
-
-	virtual ~FScrollBarStyle() {}
+	SLATECORE_API FScrollBarStyle(const FScrollBarStyle&);
+	SLATECORE_API virtual ~FScrollBarStyle();
 
 	SLATECORE_API virtual void GetResources( TArray< const FSlateBrush* >& OutBrushes ) const override;
 
@@ -997,14 +999,10 @@ struct FEditableTextBoxStyle : public FSlateWidgetStyle
 {
 	GENERATED_USTRUCT_BODY()
 
-	PRAGMA_DISABLE_DEPRECATION_WARNINGS
 	SLATECORE_API FEditableTextBoxStyle();
-	FEditableTextBoxStyle(const FEditableTextBoxStyle&) = default;
-
-	FEditableTextBoxStyle& operator=(const FEditableTextBoxStyle&) = default;
-
-	virtual ~FEditableTextBoxStyle() = default;
-	PRAGMA_ENABLE_DEPRECATION_WARNINGS
+	SLATECORE_API FEditableTextBoxStyle(const FEditableTextBoxStyle&);
+	SLATECORE_API FEditableTextBoxStyle& operator=(const FEditableTextBoxStyle&);
+	SLATECORE_API virtual ~FEditableTextBoxStyle();
 
 	SLATECORE_API virtual void GetResources( TArray< const FSlateBrush* >& OutBrushes ) const override;
 
@@ -1116,7 +1114,7 @@ struct FInlineEditableTextBlockStyle : public FSlateWidgetStyle
 
 	SLATECORE_API FInlineEditableTextBlockStyle();
 
-	virtual ~FInlineEditableTextBlockStyle() {}
+	SLATECORE_API virtual ~FInlineEditableTextBlockStyle();
 
 	SLATECORE_API virtual void GetResources( TArray< const FSlateBrush* >& OutBrushes ) const override;
 
@@ -1147,7 +1145,7 @@ struct FProgressBarStyle : public FSlateWidgetStyle
 
 	SLATECORE_API FProgressBarStyle();
 
-	virtual ~FProgressBarStyle() {}
+	SLATECORE_API virtual ~FProgressBarStyle();
 
 	SLATECORE_API virtual void GetResources( TArray< const FSlateBrush* >& OutBrushes ) const override;
 
@@ -1200,7 +1198,7 @@ struct FExpandableAreaStyle : public FSlateWidgetStyle
 
 	SLATECORE_API FExpandableAreaStyle();
 
-	virtual ~FExpandableAreaStyle() {}
+	SLATECORE_API virtual ~FExpandableAreaStyle();
 
 	SLATECORE_API virtual void GetResources( TArray< const FSlateBrush* >& OutBrushes ) const override;
 
@@ -1246,7 +1244,7 @@ struct FSearchBoxStyle : public FSlateWidgetStyle
 
 	SLATECORE_API FSearchBoxStyle();
 
-	virtual ~FSearchBoxStyle() {}
+	SLATECORE_API virtual ~FSearchBoxStyle();
 
 	SLATECORE_API virtual void GetResources( TArray< const FSlateBrush* >& OutBrushes ) const override;
 
@@ -1322,7 +1320,7 @@ struct FSliderStyle : public FSlateWidgetStyle
 
 	SLATECORE_API FSliderStyle();
 
-	virtual ~FSliderStyle() {}
+	SLATECORE_API virtual ~FSliderStyle();
 
 	SLATECORE_API virtual void GetResources( TArray< const FSlateBrush* >& OutBrushes ) const override;
 
@@ -1391,7 +1389,7 @@ struct FVolumeControlStyle : public FSlateWidgetStyle
 
 	SLATECORE_API FVolumeControlStyle();
 
-	virtual ~FVolumeControlStyle() {}
+	SLATECORE_API virtual ~FVolumeControlStyle();
 
 	SLATECORE_API virtual void GetResources( TArray< const FSlateBrush* >& OutBrushes ) const override;
 
@@ -1440,8 +1438,8 @@ struct FInlineTextImageStyle : public FSlateWidgetStyle
 	GENERATED_USTRUCT_BODY()
 
 	SLATECORE_API FInlineTextImageStyle();
-
-	virtual ~FInlineTextImageStyle() {}
+	SLATECORE_API FInlineTextImageStyle(const FInlineTextImageStyle&);
+	SLATECORE_API virtual ~FInlineTextImageStyle();
 
 	SLATECORE_API virtual void GetResources( TArray< const FSlateBrush* >& OutBrushes ) const override;
 
@@ -1471,7 +1469,7 @@ struct FSpinBoxStyle : public FSlateWidgetStyle
 
 	SLATECORE_API FSpinBoxStyle();
 
-	virtual ~FSpinBoxStyle() {}
+	SLATECORE_API virtual ~FSpinBoxStyle();
 
 	SLATECORE_API virtual void GetResources( TArray< const FSlateBrush* >& OutBrushes ) const override;
 
@@ -1557,7 +1555,7 @@ struct FSplitterStyle : public FSlateWidgetStyle
 
 	SLATECORE_API FSplitterStyle();
 
-	virtual ~FSplitterStyle() {}
+	SLATECORE_API virtual ~FSplitterStyle();
 
 	SLATECORE_API virtual void GetResources( TArray< const FSlateBrush* >& OutBrushes ) const override;
 
@@ -1588,7 +1586,7 @@ struct FTableViewStyle : public FSlateWidgetStyle
 
 	SLATECORE_API FTableViewStyle();
 
-	virtual ~FTableViewStyle() {}
+	SLATECORE_API virtual ~FTableViewStyle();
 
 	SLATECORE_API virtual void GetResources( TArray< const FSlateBrush* >& OutBrushes ) const override;
 
@@ -1623,7 +1621,7 @@ struct FTableRowStyle : public FSlateWidgetStyle
 
 	SLATECORE_API FTableRowStyle();
 
-	virtual ~FTableRowStyle() {}
+	SLATECORE_API virtual ~FTableRowStyle();
 
 	SLATECORE_API virtual void GetResources( TArray< const FSlateBrush* >& OutBrushes ) const override;
 
@@ -1765,7 +1763,7 @@ struct FTableColumnHeaderStyle : public FSlateWidgetStyle
 
 	SLATECORE_API FTableColumnHeaderStyle();
 
-	virtual ~FTableColumnHeaderStyle() {}
+	SLATECORE_API virtual ~FTableColumnHeaderStyle();
 
 	SLATECORE_API virtual void GetResources( TArray< const FSlateBrush* >& OutBrushes ) const override;
 
@@ -1831,7 +1829,7 @@ struct FHeaderRowStyle : public FSlateWidgetStyle
 
 	SLATECORE_API FHeaderRowStyle();
 
-	virtual ~FHeaderRowStyle() {}
+	SLATECORE_API virtual ~FHeaderRowStyle();
 
 	SLATECORE_API virtual void GetResources( TArray< const FSlateBrush* >& OutBrushes ) const override;
 
@@ -1893,7 +1891,7 @@ struct FDockTabStyle : public FSlateWidgetStyle
 
 	SLATECORE_API FDockTabStyle();
 
-	virtual ~FDockTabStyle() {}
+	SLATECORE_API virtual ~FDockTabStyle();
 
 	SLATECORE_API virtual void GetResources( TArray< const FSlateBrush* >& OutBrushes ) const override;
 
@@ -2005,7 +2003,7 @@ struct FScrollBoxStyle : public FSlateWidgetStyle
 
 	SLATECORE_API FScrollBoxStyle();
 
-	virtual ~FScrollBoxStyle() {}
+	SLATECORE_API virtual ~FScrollBoxStyle();
 
 	SLATECORE_API virtual void GetResources( TArray< const FSlateBrush* >& OutBrushes ) const override;
 
@@ -2088,7 +2086,7 @@ struct FScrollBorderStyle : public FSlateWidgetStyle
 
 	SLATECORE_API FScrollBorderStyle();
 
-	virtual ~FScrollBorderStyle() {}
+	SLATECORE_API virtual ~FScrollBorderStyle();
 
 	SLATECORE_API virtual void GetResources( TArray< const FSlateBrush* >& OutBrushes ) const override;
 
@@ -2118,8 +2116,8 @@ struct FWindowStyle : public FSlateWidgetStyle
 	GENERATED_USTRUCT_BODY()
 
 	SLATECORE_API FWindowStyle();
-
-	virtual ~FWindowStyle() {}
+	SLATECORE_API FWindowStyle(const FWindowStyle&);
+	SLATECORE_API virtual ~FWindowStyle();
 
 	SLATECORE_API virtual void GetResources( TArray< const FSlateBrush* >& OutBrushes ) const override;
 
diff --git a/Engine/Source/Runtime/SlateCore/Public/Styling/SlateWidgetStyle.h b/Engine/Source/Runtime/SlateCore/Public/Styling/SlateWidgetStyle.h
--- a/Engine/Source/Runtime/SlateCore/Public/Styling/SlateWidgetStyle.h
+++ b/Engine/Source/Runtime/SlateCore/Public/Styling/SlateWidgetStyle.h
@@ -22,7 +22,7 @@ public:
 	SLATECORE_API FSlateWidgetStyle( );
 
 	/** Virtual Destructor. */
-	virtual ~FSlateWidgetStyle() { }
+	SLATECORE_API virtual ~FSlateWidgetStyle();
 
 public:
 
diff --git a/Engine/Source/Runtime/SlateCore/Public/Styling/ToolBarStyle.h b/Engine/Source/Runtime/SlateCore/Public/Styling/ToolBarStyle.h
--- a/Engine/Source/Runtime/SlateCore/Public/Styling/ToolBarStyle.h
+++ b/Engine/Source/Runtime/SlateCore/Public/Styling/ToolBarStyle.h
@@ -17,8 +17,8 @@ struct FToolBarStyle : public FSlateWidgetStyle
 	GENERATED_BODY()
 
 	SLATECORE_API FToolBarStyle();
-
-	virtual ~FToolBarStyle() {}
+	SLATECORE_API FToolBarStyle(const FToolBarStyle&);
+	SLATECORE_API virtual ~FToolBarStyle();
 
 	SLATECORE_API virtual void GetResources(TArray<const FSlateBrush*>& OutBrushes) const override;
 
diff --git a/Engine/Source/Runtime/SlateCore/Public/Widgets/Images/SImage.h b/Engine/Source/Runtime/SlateCore/Public/Widgets/Images/SImage.h
--- a/Engine/Source/Runtime/SlateCore/Public/Widgets/Images/SImage.h
+++ b/Engine/Source/Runtime/SlateCore/Public/Widgets/Images/SImage.h
@@ -47,6 +47,7 @@ public:
 
 	/** Constructor */
 	SLATECORE_API SImage();
+	SLATECORE_API virtual ~SImage();
 
 	/**
 	 * Construct this widget
diff --git a/Engine/Source/Runtime/SlateCore/Public/Widgets/SCompoundWidget.h b/Engine/Source/Runtime/SlateCore/Public/Widgets/SCompoundWidget.h
--- a/Engine/Source/Runtime/SlateCore/Public/Widgets/SCompoundWidget.h
+++ b/Engine/Source/Runtime/SlateCore/Public/Widgets/SCompoundWidget.h
@@ -101,6 +101,7 @@ protected:
 
 	/** Disallow public construction */
 	SLATECORE_API SCompoundWidget();
+	SLATECORE_API virtual ~SCompoundWidget();
 
 	struct FCompoundWidgetOneChildSlot : ::TSingleWidgetChildrenWithBasicLayoutSlot<EInvalidateWidgetReason::None>
 	{
diff --git a/Engine/Source/Runtime/SlateCore/Public/Widgets/SOverlay.h b/Engine/Source/Runtime/SlateCore/Public/Widgets/SOverlay.h
--- a/Engine/Source/Runtime/SlateCore/Public/Widgets/SOverlay.h
+++ b/Engine/Source/Runtime/SlateCore/Public/Widgets/SOverlay.h
@@ -74,6 +74,7 @@ public:
 	SLATE_END_ARGS()
 
 	SLATECORE_API SOverlay();
+	SLATECORE_API ~SOverlay();
 
 	/**
 	 * Construct this widget.
diff --git a/Engine/Source/Runtime/SlateCore/Public/Widgets/SPanel.h b/Engine/Source/Runtime/SlateCore/Public/Widgets/SPanel.h
--- a/Engine/Source/Runtime/SlateCore/Public/Widgets/SPanel.h
+++ b/Engine/Source/Runtime/SlateCore/Public/Widgets/SPanel.h
@@ -74,7 +74,8 @@ protected:
 protected:
 
 	/** Hidden default constructor. */
-	SPanel( ) { }
+	SLATECORE_API SPanel( );
+	SLATECORE_API virtual ~SPanel( );
 
 public:
 	SLATECORE_API virtual void SetVisibility( TAttribute<EVisibility> InVisibility ) override final;
