[SlateCore] Added MinWidth/Height to various Slot types (in addition to the existing MaxWidth/Height attributes)

#rnx

diff --git a/Engine/Source/Runtime/Slate/Private/Widgets/Layout/SScrollBox.cpp b/Engine/Source/Runtime/Slate/Private/Widgets/Layout/SScrollBox.cpp
--- a/Engine/Source/Runtime/Slate/Private/Widgets/Layout/SScrollBox.cpp
+++ b/Engine/Source/Runtime/Slate/Private/Widgets/Layout/SScrollBox.cpp
@@ -12,6 +12,10 @@
 void SScrollBox::FSlot::Construct(const FChildren& SlotOwner, FSlotArguments&& InArgs)
 {
 	TBasicLayoutWidgetSlot<FSlot>::Construct(SlotOwner, MoveTemp(InArgs));
+	if (InArgs._MinSize.IsSet())
+	{
+		SetMinSize(MoveTemp(InArgs._MinSize));
+	}
 	if (InArgs._MaxSize.IsSet())
 	{
 		SetMaxSize(MoveTemp(InArgs._MaxSize));
@@ -25,10 +29,13 @@ void SScrollBox::FSlot::Construct(const FChildren& SlotOwner, FSlotArguments&& I
 void SScrollBox::FSlot::RegisterAttributes(FSlateWidgetSlotAttributeInitializer& AttributeInitializer)
 {
 	TBasicLayoutWidgetSlot<FSlot>::RegisterAttributes(AttributeInitializer);
+	SLATE_ADD_SLOT_ATTRIBUTE_DEFINITION_WITH_NAME(FSlot, AttributeInitializer, "Slot.MinSize", MinSize, EInvalidateWidgetReason::Layout);
 	SLATE_ADD_SLOT_ATTRIBUTE_DEFINITION_WITH_NAME(FSlot, AttributeInitializer, "Slot.MaxSize", MaxSize, EInvalidateWidgetReason::Layout);
 	SLATE_ADD_SLOT_ATTRIBUTE_DEFINITION_WITH_NAME(FSlot, AttributeInitializer, "Slot.SizeValue", SizeValue, EInvalidateWidgetReason::Layout)
+		.UpdatePrerequisite("Slot.MinSize")
 		.UpdatePrerequisite("Slot.MaxSize");
 	SLATE_ADD_SLOT_ATTRIBUTE_DEFINITION_WITH_NAME(FSlot, AttributeInitializer, "Slot.ShrinkSizeValue", ShrinkSizeValue, EInvalidateWidgetReason::Layout)
+		.UpdatePrerequisite("Slot.MinSize")
 		.UpdatePrerequisite("Slot.MaxSize");
 }
 
diff --git a/Engine/Source/Runtime/Slate/Public/Widgets/Layout/SScrollBox.h b/Engine/Source/Runtime/Slate/Public/Widgets/Layout/SScrollBox.h
--- a/Engine/Source/Runtime/Slate/Public/Widgets/Layout/SScrollBox.h
+++ b/Engine/Source/Runtime/Slate/Public/Widgets/Layout/SScrollBox.h
@@ -77,6 +77,7 @@ public:
 	public:
 		SLATE_SLOT_BEGIN_ARGS(FSlot, TBasicLayoutWidgetSlot<FSlot>)
 			SLATE_ARGUMENT(TOptional<FSizeParam>, SizeParam)
+			TAttribute<float> _MinSize;
 			TAttribute<float> _MaxSize;
 
 			/** The widget's DesiredSize will be used as the space required. */
@@ -85,6 +86,7 @@ public:
 				_SizeParam = FAuto();
 				return Me();
 			}
+
 			/** The available space will be distributed proportionately. */
 			FSlotArguments& FillSize(TAttribute<float> InStretchCoefficient)
 			{
@@ -104,6 +106,13 @@ public:
 				return Me();
 			}
 
+			/** Set the min size in SlateUnit this slot can be. */
+			FSlotArguments& MinSize(TAttribute<float> InMinHeight)
+			{
+				_MinSize = MoveTemp(InMinHeight);
+				return Me();
+			}
+
 			/** Set the max size in SlateUnit this slot can be. */
 			FSlotArguments& MaxSize(TAttribute<float> InMaxHeight)
 			{
@@ -118,6 +127,7 @@ public:
 			, SizeRule(FSizeParam::SizeRule_Auto)
 			, SizeValue(*this, 1.f)
 			, ShrinkSizeValue(*this, 1.f)
+			, MinSize(*this, 0.0f)
 			, MaxSize(*this, 0.0f)
 		{ }
 
@@ -144,6 +154,12 @@ public:
 		{
 			return ShrinkSizeValue.Get();
 		}
+
+		/** Get the min size the slot can be.*/
+		float GetMinSize() const
+		{
+			return MinSize.Get();
+		}
 		
 		/** Get the max size the slot can be.*/
 		float GetMaxSize() const
@@ -200,6 +216,12 @@ public:
 			SetSizeParam(FStretchContent(MoveTemp(InStretchCoefficient), MoveTemp(InShrinkStretchCoefficient)));
 		}
 
+		/** Set the min size in SlateUnit this slot can be. */
+		void SetMinSize(TAttribute<float> InMinSize)
+		{
+			MinSize.Assign(*this, MoveTemp(InMinSize));
+		}
+
 		/** Set the max size in SlateUnit this slot can be. */
 		void SetMaxSize(TAttribute<float> InMaxSize)
 		{
@@ -227,6 +249,9 @@ public:
 		/** The actual value this size parameter stores, used for shrinking (negative if not defined, use SizeValue). */
 		typename TBasicLayoutWidgetSlot<FSlot>::template TSlateSlotAttribute<float> ShrinkSizeValue;
 
+		/** The min size that this slot can be */
+		typename TBasicLayoutWidgetSlot<FSlot>::template TSlateSlotAttribute<float> MinSize;
+
 		/** The max size that this slot can be (0 if no max) */
 		typename TBasicLayoutWidgetSlot<FSlot>::template TSlateSlotAttribute<float> MaxSize;
 
diff --git a/Engine/Source/Runtime/SlateCore/Private/Widgets/SBoxPanel.cpp b/Engine/Source/Runtime/SlateCore/Private/Widgets/SBoxPanel.cpp
--- a/Engine/Source/Runtime/SlateCore/Private/Widgets/SBoxPanel.cpp
+++ b/Engine/Source/Runtime/SlateCore/Private/Widgets/SBoxPanel.cpp
@@ -150,6 +150,21 @@ static FVector2D ComputeDesiredSizeForBox( const TPanelChildren<SlotType>& Child
 		{
 			FVector2f CurChildDesiredSize = CurChild.GetWidget()->GetDesiredSize();
 
+			auto ClampMinMax = [&CurChild](float& DesiredSizeToClamp)
+			{
+				if (const float MinSize = CurChild.GetMinSize();
+					MinSize > 0)
+				{
+					DesiredSizeToClamp = FMath::Max(MinSize, DesiredSizeToClamp);
+				}
+
+				if (const float MaxSize = CurChild.GetMaxSize();
+					MaxSize > 0)
+				{
+					DesiredSizeToClamp = FMath::Min(MaxSize, DesiredSizeToClamp);
+				}
+			};
+
 			if (Orientation == Orient_Vertical)
 			{
 				// For a vertical panel, we want to find the maximum desired width (including margin).
@@ -158,11 +173,7 @@ static FVector2D ComputeDesiredSizeForBox( const TPanelChildren<SlotType>& Child
 
 				// Clamp to the max size if it was specified
 				float FinalChildDesiredSize = CurChildDesiredSize.Y;
-				float MaxSize = CurChild.GetMaxSize();
-				if( MaxSize > 0 )
-				{
-					FinalChildDesiredSize = FMath::Min( MaxSize, FinalChildDesiredSize );
-				}
+				ClampMinMax(FinalChildDesiredSize);
 
 				MyDesiredSize.Y += FinalChildDesiredSize + CurChild.GetPadding().template GetTotalSpaceAlong<Orient_Vertical>();
 			}
@@ -174,11 +185,7 @@ static FVector2D ComputeDesiredSizeForBox( const TPanelChildren<SlotType>& Child
 
 				// Clamp to the max size if it was specified
 				float FinalChildDesiredSize = CurChildDesiredSize.X;
-				float MaxSize = CurChild.GetMaxSize();
-				if( MaxSize > 0 )
-				{
-					FinalChildDesiredSize = FMath::Min( MaxSize, FinalChildDesiredSize );
-				}
+				ClampMinMax(FinalChildDesiredSize);
 
 				MyDesiredSize.X += FinalChildDesiredSize + CurChild.GetPadding().template GetTotalSpaceAlong<Orient_Horizontal>();
 			}
diff --git a/Engine/Source/Runtime/SlateCore/Public/Layout/LayoutUtils.h b/Engine/Source/Runtime/SlateCore/Public/Layout/LayoutUtils.h
--- a/Engine/Source/Runtime/SlateCore/Public/Layout/LayoutUtils.h
+++ b/Engine/Source/Runtime/SlateCore/Public/Layout/LayoutUtils.h
@@ -255,17 +255,20 @@ static void ArrangeChildrenInStack(EFlowDirection InLayoutFlow, const TPanelChil
 	{
 		return;
 	}
-	
+
 	// Allotted space will be given to fixed-size children first.
 	// Remaining space will be proportionately divided between stretch children (SizeRule_Stretch and SizeRule_StretchContent)
 	// based on their stretch coefficient.
 
 	// Helper function to clamp to max size, if the constraint is set.
-	auto ClampToMaxSize = [](const float Size, const float MaxSize)
+	auto ClampSize = [](const float Size, const float MinSize, const float MaxSize)
 	{
-		return MaxSize > 0.0f ? FMath::Min(MaxSize, Size) : Size; 
+		return FMath::Clamp(
+			Size,
+			MinSize > 0.0f ? MinSize : 0.0f,
+			MaxSize > 0.0f ? MaxSize : std::numeric_limits<float>::max());
 	};
-	
+
 	float GrowStretchCoefficientTotal = 0.0f;
 	float ShrinkStretchCoefficientTotal = 0.0f;
 	float FixedSizeTotal = 0.0f;
@@ -275,8 +278,10 @@ static void ArrangeChildrenInStack(EFlowDirection InLayoutFlow, const TPanelChil
 	{
 		// Size of the item
 		float Size = 0.0f;
-		// Initial size of the item 
+		// Initial size of the item
 		float BasisSize = 0.0f;
+		// Min size constraint of the item.
+		float MinSize = 0.0f;
 		// Max size constraint of the item.
 		float MaxSize = 0.0f;
 		// Stretch coefficient when the items are growing.
@@ -290,11 +295,11 @@ static void ArrangeChildrenInStack(EFlowDirection InLayoutFlow, const TPanelChil
 	};
 	TArray<FStretchItem, TInlineAllocator<16>> StretchItems;
 	StretchItems.Init({}, Children.Num());
-	
+
 	bool bAnyChildVisible = false;
 	bool bAnyStretchContentItems = false;
 	bool bAnyStretchItems = false;
-	
+
 	// Compute the sum of stretch coefficients (SizeRule_Stretch & SizeRule_StretchContent) and space required by fixed-size widgets (SizeRule_Auto),
 	// as well as the total desired size.
 	for (int32 ChildIndex = 0; ChildIndex < Children.Num(); ++ChildIndex)
@@ -304,7 +309,7 @@ static void ArrangeChildrenInStack(EFlowDirection InLayoutFlow, const TPanelChil
 		if (CurChild.GetWidget()->GetVisibility() != EVisibility::Collapsed)
 		{
 			bAnyChildVisible = true;
-			
+
 			// All widgets contribute their margin to the fixed space requirement
 			FixedSizeTotal += CurChild.GetPadding().template GetTotalSpaceAlong<Orientation>();
 
@@ -315,14 +320,16 @@ static void ArrangeChildrenInStack(EFlowDirection InLayoutFlow, const TPanelChil
 				? ChildDesiredSize.Y
 				: ChildDesiredSize.X;
 
+			const float MinSize = CurChild.GetMinSize();
 			const float MaxSize = CurChild.GetMaxSize();
 
-			FStretchItem& Item = StretchItems[ChildIndex]; 
+			FStretchItem& Item = StretchItems[ChildIndex];
+			Item.MinSize = MinSize;
 			Item.MaxSize = MaxSize;
 			Item.SizeRule = CurChild.GetSizeRule();
 
 			// Clamp to the max size if it was specified
-			ChildSize = ClampToMaxSize(ChildSize, MaxSize);
+			ChildSize = ClampSize(ChildSize, MinSize, MaxSize);
 
 			if (CurChild.GetSizeRule() == FSizeParam::SizeRule_Stretch)
 			{
@@ -341,7 +348,7 @@ static void ArrangeChildrenInStack(EFlowDirection InLayoutFlow, const TPanelChil
 			}
 			else if (CurChild.GetSizeRule() == FSizeParam::SizeRule_StretchContent)
 			{
-				// Allow separate values from grow and shrink, as the adjustment is relative to the child size. 
+				// Allow separate values from grow and shrink, as the adjustment is relative to the child size.
 				Item.GrowStretchValue = FMath::Max(0.f, CurChild.GetSizeValue());
 				Item.ShrinkStretchValue = FMath::Max(0.f,CurChild.GetShrinkSizeValue());
 				Item.Size = ChildSize;
@@ -377,7 +384,7 @@ static void ArrangeChildrenInStack(EFlowDirection InLayoutFlow, const TPanelChil
 	const float AllottedSize = Orientation == Orient_Vertical
 		? AllottedGeometry.GetLocalSize().Y
 		: AllottedGeometry.GetLocalSize().X;
-	
+
 	// The space available for SizeRule_Stretch and SizeRule_StretchContent widgets is any space that wasn't taken up by fixed-sized widgets.
 	float AvailableSpace = FMath::Max(MinAvailableSpace, AllottedSize - FixedSizeTotal);
 
@@ -393,8 +400,8 @@ static void ArrangeChildrenInStack(EFlowDirection InLayoutFlow, const TPanelChil
 				// Stretch widgets get a fraction of the space remaining after all the fixed-space requirements are met.
 				// Supporting only one stretch value since otherwise the transition would be discontinuous as (reference) basis size is 0.
 				const float Size = AvailableSpace * Item.GrowStretchValue / GrowStretchCoefficientTotal;
-				
-				Item.Size = ClampToMaxSize(Size, Item.MaxSize);
+
+				Item.Size = ClampSize(Size, Item.MinSize, Item.MaxSize);
 				
 				UsedSpace += Item.Size;
 			}
@@ -492,8 +499,8 @@ static void ArrangeChildrenInStack(EFlowDirection InLayoutFlow, const TPanelChil
 						Item.bFrozen = true;
 						continue;
 					}
-					
-					constexpr float MinSize = 0.0f;
+
+					const float MinSize = Item.MinSize;
 					const float MaxSize = Item.MaxSize;
 					const bool bHasMaxConstraint = MaxSize > 0.0f; 
 
diff --git a/Engine/Source/Runtime/SlateCore/Public/Widgets/SBoxPanel.h b/Engine/Source/Runtime/SlateCore/Public/Widgets/SBoxPanel.h
--- a/Engine/Source/Runtime/SlateCore/Public/Widgets/SBoxPanel.h
+++ b/Engine/Source/Runtime/SlateCore/Public/Widgets/SBoxPanel.h
@@ -38,6 +38,7 @@ protected:
 	public:
 		SLATE_SLOT_BEGIN_ARGS(TSlot, TBasicLayoutWidgetSlot<SlotType>)
 			SLATE_ARGUMENT(TOptional<FSizeParam>, SizeParam)
+			TAttribute<float> _MinSize;
 			TAttribute<float> _MaxSize;
 		SLATE_SLOT_END_ARGS()
 
@@ -48,6 +49,7 @@ protected:
 			, SizeRule(FSizeParam::SizeRule_Stretch)
 			, SizeValue(*this, 1.f)
 			, ShrinkSizeValue(*this, 1.f)
+			, MinSize(*this, 0.0f)
 			, MaxSize(*this, 0.0f)
 		{ }
 
@@ -55,6 +57,10 @@ protected:
 		void Construct(const FChildren& SlotOwner, FSlotArguments&& InArgs)
 		{
 			TBasicLayoutWidgetSlot<SlotType>::Construct(SlotOwner, MoveTemp(InArgs));
+			if (InArgs._MinSize.IsSet())
+			{
+				SetMinSize(MoveTemp(InArgs._MinSize));
+			}
 			if (InArgs._MaxSize.IsSet())
 			{
 				SetMaxSize(MoveTemp(InArgs._MaxSize));
@@ -68,10 +74,13 @@ protected:
 		static void RegisterAttributes(FSlateWidgetSlotAttributeInitializer& AttributeInitializer)
 		{
 			TBasicLayoutWidgetSlot<SlotType>::RegisterAttributes(AttributeInitializer);
+			SLATE_ADD_SLOT_ATTRIBUTE_DEFINITION_WITH_NAME(TSlot<SlotType>, AttributeInitializer, "Slot.MinSize", MinSize, EInvalidateWidgetReason::Layout);
 			SLATE_ADD_SLOT_ATTRIBUTE_DEFINITION_WITH_NAME(TSlot<SlotType>, AttributeInitializer, "Slot.MaxSize", MaxSize, EInvalidateWidgetReason::Layout);
 			SLATE_ADD_SLOT_ATTRIBUTE_DEFINITION_WITH_NAME(TSlot<SlotType>, AttributeInitializer, "Slot.SizeValue", SizeValue, EInvalidateWidgetReason::Layout)
+				.UpdatePrerequisite("Slot.MinSize")
 				.UpdatePrerequisite("Slot.MaxSize");
 			SLATE_ADD_SLOT_ATTRIBUTE_DEFINITION_WITH_NAME(TSlot<SlotType>, AttributeInitializer, "Slot.ShrinkSizeValue", ShrinkSizeValue, EInvalidateWidgetReason::Layout)
+				.UpdatePrerequisite("Slot.MinSize")
 				.UpdatePrerequisite("Slot.MaxSize");
 		}
 
@@ -99,6 +108,12 @@ protected:
 			return ShrinkSizeValue.Get();
 		}
 
+		/** Get the min size the slot can be.*/
+		float GetMinSize() const
+		{
+			return MinSize.Get();
+		}
+
 		/** Get the max size the slot can be.*/
 		float GetMaxSize() const
 		{
@@ -161,6 +176,12 @@ protected:
 			SetSizeParam(FStretchContent(MoveTemp(InStretchCoefficient), MoveTemp(InShrinkStretchCoefficient)));
 		}
 
+		/** Set the min size in SlateUnit this slot can be. */
+		void SetMinSize(TAttribute<float> InMinSize)
+		{
+			MinSize.Assign(*this, MoveTemp(InMinSize));
+		}
+
 		/** Set the max size in SlateUnit this slot can be. */
 		void SetMaxSize(TAttribute<float> InMaxSize)
 		{
@@ -177,6 +198,9 @@ protected:
 		/** The actual value this size parameter stores, used for shrinking (negative if not defined, use SizeValue). */
 		typename TBasicLayoutWidgetSlot<SlotType>::template TSlateSlotAttribute<float> ShrinkSizeValue;
 
+		/** The min size that this slot can be */
+		typename TBasicLayoutWidgetSlot<SlotType>::template TSlateSlotAttribute<float> MinSize;
+
 		/** The max size that this slot can be (0 if no max) */
 		typename TBasicLayoutWidgetSlot<SlotType>::template TSlateSlotAttribute<float> MaxSize;
 	};
@@ -312,7 +336,14 @@ public:
 				_SizeParam = FStretchContent(MoveTemp(InStretchCoefficient), MoveTemp(InShrinkStretchCoefficient));
 				return Me();
 			}
-		
+
+			/** Set the min size in SlateUnit this slot can be. */
+			FSlotArguments& MinWidth(TAttribute<float> InMinWidth)
+			{
+				_MinSize = MoveTemp(InMinWidth);
+				return Me();
+			}
+
 			/** Set the max size in SlateUnit this slot can be. */
 			FSlotArguments& MaxWidth(TAttribute<float> InMaxWidth)
 			{
@@ -350,6 +381,12 @@ public:
 			SetSizeToStretchContent(MoveTemp(InStretchCoefficient), MoveTemp(InShrinkStretchCoefficient));
 		}
 
+		/** Set the min size in SlateUnit this slot can be. */
+		void SetMinWidth(TAttribute<float> InMinWidth)
+		{
+			SetMinSize(MoveTemp(InMinWidth));
+		}
+
 		/** Set the max size in SlateUnit this slot can be. */
 		void SetMaxWidth(TAttribute<float> InMaxWidth)
 		{
@@ -444,6 +481,13 @@ public:
 				_SizeParam = FStretchContent(MoveTemp(InStretchCoefficient), MoveTemp(InShrinkStretchCoefficient));
 				return Me();
 			}
+
+			/** Set the min size in SlateUnit this slot can be. */
+			FSlotArguments& MinHeight(TAttribute<float> InMinHeight)
+			{
+				_MinSize = MoveTemp(InMinHeight);
+				return Me();
+			}
 			
 			/** Set the max size in SlateUnit this slot can be. */
 			FSlotArguments& MaxHeight(TAttribute<float> InMaxHeight)
@@ -482,6 +526,12 @@ public:
 			SetSizeToStretchContent(MoveTemp(InStretchCoefficient), MoveTemp(InShrinkStretchCoefficient));
 		}
 
+		/** Set the min size in SlateUnit this slot can be. */
+		void SetMinHeight(TAttribute<float> InMinHeight)
+		{
+			SetMinSize(MoveTemp(InMinHeight));
+		}
+
 		/** Set the max size in SlateUnit this slot can be. */
 		void SetMaxHeight(TAttribute<float> InMaxHeight)
 		{
@@ -553,7 +603,7 @@ public:
 				_SizeParam = FAuto();
 				return Me();
 			}
-		
+
 			/**
 			 * The available space will be distributed proportionately to each slots stretch coefficient.
 			 * A slot with coefficient of 2 will get assigned twice as much available space as slot with coefficient 1. 
@@ -564,7 +614,7 @@ public:
 				_SizeParam = FStretch(MoveTemp(InStretchCoefficient));
 				return Me();
 			}
-		
+
 			/**
 			 * The widget's content size is adjusted proportionally to fit the available space.
 			 * The slots size starts at DesiredSize, and a slot with coefficient of 2 will get adjusted twice as much as slot with coefficient 1 to fit the available space.
@@ -576,6 +626,14 @@ public:
 				_SizeParam = FStretchContent(MoveTemp(InStretchCoefficient), MoveTemp(InShrinkStretchCoefficient));
 				return Me();
 			}
+
+			/** Set the min size in SlateUnit this slot can be. */
+			FSlotArguments& MinSize(TAttribute<float> InMinHeight)
+			{
+				_MinSize = MoveTemp(InMinHeight);
+				return Me();
+			}
+
 			/** Set the max size in SlateUnit this slot can be. */
 			FSlotArguments& MaxSize(TAttribute<float> InMaxHeight)
 			{
