[SlateCore] Allow FOptionalSize to be constructed from TOptional<float>

#rnx

diff --git a/Engine/Source/Runtime/SlateCore/Public/Types/SlateStructs.h b/Engine/Source/Runtime/SlateCore/Public/Types/SlateStructs.h
--- a/Engine/Source/Runtime/SlateCore/Public/Types/SlateStructs.h
+++ b/Engine/Source/Runtime/SlateCore/Public/Types/SlateStructs.h
@@ -4,6 +4,7 @@
 
 #include "CoreMinimal.h"
 #include "Misc/Attribute.h"
+#include "Misc/Optional.h"
 
 /**
  * Structure for optional floating point sizes.
@@ -26,6 +27,15 @@ struct FOptionalSize
 		: Size(SpecifiedSize)
 	{ }
 
+	/**
+	 * Creates a size with the TOptional value.
+	 *
+	 * @param OptionalSize The optional size to set.
+	 */
+	FOptionalSize( const TOptional<float>& OptionalSize )
+		: Size(OptionalSize.Get(Unspecified))
+	{ }
+
 public:
 
 	/**
