[Slate] Account for custom VAlign in SWidgetBlock

#rnx

diff --git a/Engine/Source/Runtime/Slate/Private/Framework/MultiBox/SWidgetBlock.cpp b/Engine/Source/Runtime/Slate/Private/Framework/MultiBox/SWidgetBlock.cpp
--- a/Engine/Source/Runtime/Slate/Private/Framework/MultiBox/SWidgetBlock.cpp
+++ b/Engine/Source/Runtime/Slate/Private/Framework/MultiBox/SWidgetBlock.cpp
@@ -87,8 +87,12 @@ void SWidgetBlock::BuildMultiBlockWidget(const ISlateStyle* StyleSet, const FNam
 	TSharedPtr< const FMultiBox > MultiBox = OwnerMultiBoxWidgetPinned->GetMultiBox();
 	const TSharedRef<const FWidgetBlock> WidgetBlock = StaticCastSharedRef<const FWidgetBlock>(MultiBlock.ToSharedRef());
 
+	const bool bHasLabel = !WidgetBlock->Label.IsEmpty();
+
+	// Initially use default behavior
+	EVerticalAlignment ContentVerticalAlignment = bHasLabel ? VAlign_Center : VAlign_Fill;
+
 	// Support menus which do not have a defined widget style yet
-	bool bHasLabel = !WidgetBlock->Label.IsEmpty();
 	FMargin Padding;
 	const FTextBlockStyle* LabelStyle = nullptr;
 	if(StyleSet->HasWidgetStyle<FToolBarStyle>(StyleName))
@@ -97,6 +101,11 @@ void SWidgetBlock::BuildMultiBlockWidget(const ISlateStyle* StyleSet, const FNam
 
 		Padding = WidgetBlock->StyleParams.bNoIndent ? ToolBarStyle.BlockPadding : ToolBarStyle.IndentedBlockPadding;
 		LabelStyle = &ToolBarStyle.LabelStyle;
+
+		if (ToolBarStyle.VerticalAlignmentOverride.IsSet())
+		{
+			ContentVerticalAlignment = ToolBarStyle.VerticalAlignmentOverride.GetValue();
+		}
 	}
 	else
 	{
@@ -104,6 +113,11 @@ void SWidgetBlock::BuildMultiBlockWidget(const ISlateStyle* StyleSet, const FNam
 		LabelStyle = &StyleSet->GetWidgetStyle<FTextBlockStyle>(ISlateStyle::Join(StyleName, ".Label"));
 	}
 
+	if (WidgetBlock->StyleParams.VerticalAlignment.IsSet())
+	{
+		ContentVerticalAlignment = WidgetBlock->StyleParams.VerticalAlignment.GetValue();
+	}
+
 	if(OwnerMultiBoxWidgetPinned->GetMultiBox()->GetType() == EMultiBoxType::Menu)
 	{
 		// Account for checkmark used in other menu blocks but not used in for widget rows
@@ -176,7 +190,7 @@ void SWidgetBlock::BuildMultiBlockWidget(const ISlateStyle* StyleSet, const FNam
 			]
 		]
 		+SHorizontalBox::Slot()
-		.VAlign(bHasLabel ? VAlign_Center : VAlign_Fill)
+		.VAlign(ContentVerticalAlignment)
 		.FillWidth(1.f)
 		[
 			WidgetBlock->ContentWidget
diff --git a/Engine/Source/Runtime/SlateCore/Public/Styling/ToolBarStyle.h b/Engine/Source/Runtime/SlateCore/Public/Styling/ToolBarStyle.h
--- a/Engine/Source/Runtime/SlateCore/Public/Styling/ToolBarStyle.h
+++ b/Engine/Source/Runtime/SlateCore/Public/Styling/ToolBarStyle.h
@@ -18,7 +18,7 @@ struct FToolBarStyle : public FSlateWidgetStyle
 
 	SLATECORE_API FToolBarStyle();
 	SLATECORE_API FToolBarStyle(const FToolBarStyle&);
-	SLATECORE_API virtual ~FToolBarStyle();
+	SLATECORE_API virtual ~FToolBarStyle() override;
 
 	SLATECORE_API virtual void GetResources(TArray<const FSlateBrush*>& OutBrushes) const override;
 
@@ -84,11 +84,11 @@ struct FToolBarStyle : public FSlateWidgetStyle
 	UPROPERTY(EditAnywhere, BlueprintReadWrite, Category = Appearance)
 	float UniformBlockHeight;
 	FToolBarStyle& SetUniformBlockHeight(const float InUniformBlockHeight) { UniformBlockHeight = InUniformBlockHeight; return *this; }
-	
+
 	UPROPERTY(EditAnywhere, BlueprintReadWrite, Category = Appearance)
 	int32 NumColumns;
 	FToolBarStyle& SetNumColumns(const int32 InNumColumns) { NumColumns = InNumColumns; return *this; }
-	
+
 	UPROPERTY(EditAnywhere, BlueprintReadWrite, Category = Appearance)
 	FMargin IconPadding;
 	FToolBarStyle& SetIconPadding(const FMargin& InIconPadding) { IconPadding = InIconPadding; return *this; }
@@ -99,11 +99,7 @@ struct FToolBarStyle : public FSlateWidgetStyle
 
 	UPROPERTY(EditAnywhere, BlueprintReadWrite, Category = Appearance)
 	float SeparatorThickness;
-	FToolBarStyle& SetSeparatorThickness(float InSeparatorThickness)
-	{
-		SeparatorThickness = InSeparatorThickness;
-		return *this;
-	}
+	FToolBarStyle& SetSeparatorThickness(float InSeparatorThickness) { SeparatorThickness = InSeparatorThickness; return *this;	}
 
 	UPROPERTY(EditAnywhere, BlueprintReadWrite, Category = Appearance)
 	FMargin ComboButtonPadding;
@@ -136,7 +132,7 @@ struct FToolBarStyle : public FSlateWidgetStyle
 	UPROPERTY(EditAnywhere, BlueprintReadWrite, Category=Appearance)
 	bool bShowLabels;
 	FToolBarStyle& SetShowLabels(bool bInShowLabels) { bShowLabels = bInShowLabels; return *this; }
-	
+
 	UPROPERTY(EditAnywhere, BlueprintReadWrite, Category=Appearance)
 	float ButtonContentMaxWidth = 64.0f;
 	FToolBarStyle& SetButtonContentMaxWidth(float InButtonContentMaxWidth) { ButtonContentMaxWidth = InButtonContentMaxWidth; return *this; }
@@ -144,13 +140,16 @@ struct FToolBarStyle : public FSlateWidgetStyle
 	UPROPERTY(EditAnywhere, BlueprintReadWrite, Category=Appearance)
 	float ButtonContentFillWidth = 1.0f;
 	FToolBarStyle& SetButtonContentFillWidth(float InButtonContentFillWidth) { ButtonContentFillWidth = InButtonContentFillWidth; return *this; }
-	
+
 	UPROPERTY(EditAnywhere, BlueprintReadWrite, Category = Appearance)
 	FMargin IconPaddingWithVisibleLabel;
 	FToolBarStyle& SetIconPaddingWithVisibleLabel(const FMargin& InIconPaddingWithVisibleLabel) { IconPaddingWithVisibleLabel = InIconPaddingWithVisibleLabel; return *this; }
-	
+
 	UPROPERTY(EditAnywhere, BlueprintReadWrite, Category = Appearance)
 	FMargin IconPaddingWithCollapsedLabel;
 	FToolBarStyle& SetIconPaddingWithCollapsedLabel(const FMargin& InIconPaddingWithCollapsedLabel) { IconPaddingWithCollapsedLabel = InIconPaddingWithCollapsedLabel; return *this; }
-	
+
+	UPROPERTY(EditAnywhere, BlueprintReadWrite, Category = Appearance)
+	TOptional<TEnumAsByte<EVerticalAlignment>> VerticalAlignmentOverride;
+	FToolBarStyle& SetVerticalAlignment(const EVerticalAlignment& InVerticalAlignment) { VerticalAlignmentOverride = InVerticalAlignment; return *this; }
 };
