Invalidate cached widgets when the vector graphics cache is flushed.

diff --git a/Engine/Source/Runtime/SlateCore/Private/Rendering/SlateVectorGraphicsCache.cpp b/Engine/Source/Runtime/SlateCore/Private/Rendering/SlateVectorGraphicsCache.cpp
--- a/Engine/Source/Runtime/SlateCore/Private/Rendering/SlateVectorGraphicsCache.cpp
+++ b/Engine/Source/Runtime/SlateCore/Private/Rendering/SlateVectorGraphicsCache.cpp
@@ -7,6 +7,7 @@
 #include "Types/SlateVector2.h"
 #include "ProfilingDebugging/ScopedTimers.h"
 #include "Async/ParallelFor.h"
+#include "Application/SlateApplicationBase.h"
 
 static FAtlasFlushParams SVGAtlasFlushParams;
 FAutoConsoleVariableRef CVarMaxSVGAtlasPagesBeforeFlush(
@@ -195,6 +196,9 @@ void FSlateVectorGraphicsCache::RequestFlushCache(const FString& Reason)
 
 void FSlateVectorGraphicsCache::FlushCache()
 {
+	// Ensure all invalidation panels are cleared of cached widgets
+	FSlateApplicationBase::Get().InvalidateAllWidgets(false);
+
 	const bool bWaitForRelease = true;
 	ReleaseResources(bWaitForRelease);
 
