Slate: Fix crash when widget removes itself from the parent in OnPaint and reindex and invalidation are enabled. The cached value needs to be fetched again after an OnPaint. The proxy will become invalid at the end of the loop (and need to stay valid or everything else will break).

diff --git a/Engine/Source/Runtime/SlateCore/Private/FastUpdate/SlateInvalidationRoot.cpp b/Engine/Source/Runtime/SlateCore/Private/FastUpdate/SlateInvalidationRoot.cpp
--- a/Engine/Source/Runtime/SlateCore/Private/FastUpdate/SlateInvalidationRoot.cpp
+++ b/Engine/Source/Runtime/SlateCore/Private/FastUpdate/SlateInvalidationRoot.cpp
@@ -662,17 +662,16 @@ bool FSlateInvalidationRoot::PaintFastPath_UpdateNextWidget(const FSlateInvalida
 	const FSlateInvalidationWidgetIndex MyIndex = FinalUpdateList.Pop(EAllowShrinking::No).GetWidgetIndex();
 
 	FSlateInvalidationWidgetList::InvalidationWidgetType& InvalidationWidget = (*FastWidgetPathList)[MyIndex];
-	SWidget* WidgetPtr = InvalidationWidget.GetWidget();
 
 #if UE_SLATE_WITH_INVALIDATIONWIDGETLIST_DEBUGGING
 	if (GSlateInvalidationRootVerifyWidgetsAreUpdatedOnce)
 	{
-		ensureAlwaysMsgf(!InvalidationWidget.bDebug_Updated, TEXT("VerifyWidgetsAreUpdatedOnce failed. Widget '%s' is going to be updated more than once"), *FReflectionMetaData::GetWidgetDebugInfo(WidgetPtr));
+		ensureAlwaysMsgf(!InvalidationWidget.bDebug_Updated, TEXT("VerifyWidgetsAreUpdatedOnce failed. Widget '%s' is going to be updated more than once"), *FReflectionMetaData::GetWidgetDebugInfo(InvalidationWidget.GetWidget()));
 	}
 #endif
 
 	// Check visibility, it was tested before adding it to the list but another widget may have change while updating.
-	if (InvalidationWidget.Visibility.IsVisible() && WidgetPtr)
+	if (InvalidationWidget.Visibility.IsVisible() && InvalidationWidget.GetWidget())
 	{
 		const FWidgetProxy::FUpdateResult UpdateResult = InvalidationWidget.Update(*Context.PaintArgs, *Context.WindowElementList);
 
@@ -702,7 +701,7 @@ bool FSlateInvalidationRoot::PaintFastPath_UpdateNextWidget(const FSlateInvalida
 			// Did it painted more elements than it previously had
 			if (UpdateResult.NewOutgoingLayerId > UpdateResult.PreviousOutgoingLayerId && GSlateInvalidationEnableReindexLayerId)
 			{
-				if (InvalidationWidget.Visibility.IsVisible() && WidgetPtr)
+				if (InvalidationWidget.Visibility.IsVisible() && InvalidationWidget.GetWidget())
 				{
 					PaintFastPath_FixupLayerId(FastPaintContext, InvalidationWidget, UpdateResult.NewOutgoingLayerId);
 				}
