Slate SDF fonts - feature enabled

#tests Tested in editor and PIE

diff --git a/Engine/Source/Runtime/SlateCore/Private/Fonts/FontCache.cpp b/Engine/Source/Runtime/SlateCore/Private/Fonts/FontCache.cpp
--- a/Engine/Source/Runtime/SlateCore/Private/Fonts/FontCache.cpp
+++ b/Engine/Source/Runtime/SlateCore/Private/Fonts/FontCache.cpp
@@ -72,7 +72,7 @@ FAutoConsoleVariableRef CVarFlushFontCache(
 	TEXT("Flush the font cache."));
 #endif
 
-static TAutoConsoleVariable<bool> CVarSlateSdfTextEnable(TEXT("SlateSdfText.Enable"), false, TEXT("Enables MSDF-based text rendering in Slate"));
+static TAutoConsoleVariable<bool> CVarSlateSdfTextEnable(TEXT("SlateSdfText.Enable"), true, TEXT("Enables MSDF-based text rendering in Slate"));
 
 bool IsSlateSdfTextFeatureEnabled()
 {
