Early Exit in FWidgetPath::ContainsWidget if the widget is null.

#rnx

diff --git a/Engine/Source/Runtime/SlateCore/Private/Layout/WidgetPath.cpp b/Engine/Source/Runtime/SlateCore/Private/Layout/WidgetPath.cpp
--- a/Engine/Source/Runtime/SlateCore/Private/Layout/WidgetPath.cpp
+++ b/Engine/Source/Runtime/SlateCore/Private/Layout/WidgetPath.cpp
@@ -58,11 +58,14 @@ FWidgetPath FWidgetPath::GetPathDownTo( TSharedRef<const SWidget> MarkerWidget )
 
 bool FWidgetPath::ContainsWidget( const SWidget* WidgetToFind ) const
 {
-	for(int32 WidgetIndex = 0; WidgetIndex < Widgets.Num(); ++WidgetIndex)
+	if (WidgetToFind)
 	{
-		if ( &Widgets[WidgetIndex].Widget.Get() == WidgetToFind )
+		for (int32 WidgetIndex = 0; WidgetIndex < Widgets.Num(); ++WidgetIndex)
 		{
-			return true;
+			if (&Widgets[WidgetIndex].Widget.Get() == WidgetToFind)
+			{
+				return true;
+			}
 		}
 	}
 		
@@ -416,11 +419,14 @@ FWeakWidgetPath::EPathResolutionResult::Result FWeakWidgetPath::ToWidgetPath( FW
 
 bool FWeakWidgetPath::ContainsWidget(const SWidget* SomeWidget) const
 {
-	for (int32 WidgetIndex = 0; WidgetIndex < Widgets.Num(); ++WidgetIndex)
+	if (SomeWidget)
 	{
-		if (Widgets[WidgetIndex].Pin().Get() == SomeWidget)
+		for (int32 WidgetIndex = 0; WidgetIndex < Widgets.Num(); ++WidgetIndex)
 		{
-			return true;
+			if (Widgets[WidgetIndex].Pin().Get() == SomeWidget)
+			{
+				return true;
+			}
 		}
 	}
 
