Replace ActiveTimerHandles TArray with a TSet to optimize freeing large numbers of UI items, e.g, 6400 instances ~40s -> ~3s (16k instances 4min -> 7s)

diff --git a/Engine/Source/Runtime/SlateCore/Private/Application/SlateApplicationBase.cpp b/Engine/Source/Runtime/SlateCore/Private/Application/SlateApplicationBase.cpp
--- a/Engine/Source/Runtime/SlateCore/Private/Application/SlateApplicationBase.cpp
+++ b/Engine/Source/Runtime/SlateCore/Private/Application/SlateApplicationBase.cpp
@@ -114,7 +114,7 @@ void FSlateApplicationBase::RegisterActiveTimer( const TSharedRef<FActiveTimerHa
 void FSlateApplicationBase::UnRegisterActiveTimer( const TSharedRef<FActiveTimerHandle>& ActiveTimerHandle )
 {
 	FScopeLock ActiveTimerLock(&ActiveTimerCS);
-	ActiveTimerHandles.RemoveSingleSwap(ActiveTimerHandle);
+	ActiveTimerHandles.Remove(ActiveTimerHandle);
 }
 
 void FSlateApplicationBase::UpdateAnyActiveTimersArePending()
@@ -124,11 +124,13 @@ void FSlateApplicationBase::UpdateAnyActiveTimersArePending()
 	// first remove any tick handles that may have become invalid.
 	// If we didn't remove invalid handles here, they would never get removed because
 	// we don't force widgets to UnRegister before they are destroyed.
-	ActiveTimerHandles.RemoveAll([](const TWeakPtr<FActiveTimerHandle>& ActiveTimerHandle)
+	for (auto It = ActiveTimerHandles.CreateIterator(); It; ++It)
 	{
-		// only check the weak pointer to the handle. Just want to make sure to clear out any widgets that have since been deleted.
-		return !ActiveTimerHandle.IsValid();
-	});
+		if(!(*It).IsValid())
+		{
+			It.RemoveCurrent();
+		}
+	}
 
 	// The rest are valid. Update their pending status and see if any are ready.
 	const double CurrentTime = GetCurrentTime();
diff --git a/Engine/Source/Runtime/SlateCore/Public/Application/SlateApplicationBase.h b/Engine/Source/Runtime/SlateCore/Public/Application/SlateApplicationBase.h
--- a/Engine/Source/Runtime/SlateCore/Public/Application/SlateApplicationBase.h
+++ b/Engine/Source/Runtime/SlateCore/Public/Application/SlateApplicationBase.h
@@ -497,8 +497,8 @@ private:
 	 */
 	SLATECORE_API void UnRegisterActiveTimer( const TSharedRef<FActiveTimerHandle>& ActiveTimerHandle );
 
-	/** The list of active timer handles. */
-	TArray<TWeakPtr<FActiveTimerHandle>> ActiveTimerHandles;
+	/** The set of active timer handles. */
+	TSet<TWeakPtr<FActiveTimerHandle>> ActiveTimerHandles;
 
 protected:
 	enum class ECustomSafeZoneState : uint8
