Fix typo in comment

diff --git a/Engine/Source/Runtime/SlateCore/Public/Layout/ChildrenBase.h b/Engine/Source/Runtime/SlateCore/Public/Layout/ChildrenBase.h
--- a/Engine/Source/Runtime/SlateCore/Public/Layout/ChildrenBase.h
+++ b/Engine/Source/Runtime/SlateCore/Public/Layout/ChildrenBase.h
@@ -42,7 +42,7 @@ public:
 
 	FChildren(std::nullptr_t) = delete;
 
-	//~ Prevents allocation of FChilren. It creates confusion between FSlot and FChildren
+	//~ Prevents allocation of FChildren. It creates confusion between FSlot and FChildren.
 	void* operator new (size_t) = delete;
 	void* operator new[](size_t) = delete;
 
