Fixed typo that selects the maximum number of indices.

diff --git a/Engine/Source/Runtime/SlateCore/Private/Rendering/ElementBatcher.cpp b/Engine/Source/Runtime/SlateCore/Private/Rendering/ElementBatcher.cpp
--- a/Engine/Source/Runtime/SlateCore/Private/Rendering/ElementBatcher.cpp
+++ b/Engine/Source/Runtime/SlateCore/Private/Rendering/ElementBatcher.cpp
@@ -241,7 +241,7 @@ void FSlateBatchData::MergeRenderBatches()
 		}
 
 		MaxNumFinalVertices = FMath::Max(MaxNumFinalVertices, FinalVertexData.Num());
-		MaxNumFinalIndices = FMath::Max(MaxNumFinalVertices, FinalIndexData.Num());
+		MaxNumFinalIndices = FMath::Max(MaxNumFinalIndices, FinalIndexData.Num());
 	}
 }
 
