Deprecated a FSlateFontInfo constructor using a font path

#rnx

diff --git a/Engine/Plugins/Runtime/AudioWidgets/Source/AudioWidgets/Private/AudioMaterialSlate/AudioMaterialSlateTypes.cpp b/Engine/Plugins/Runtime/AudioWidgets/Source/AudioWidgets/Private/AudioMaterialSlate/AudioMaterialSlateTypes.cpp
--- a/Engine/Plugins/Runtime/AudioWidgets/Source/AudioWidgets/Private/AudioMaterialSlate/AudioMaterialSlateTypes.cpp
+++ b/Engine/Plugins/Runtime/AudioWidgets/Source/AudioWidgets/Private/AudioMaterialSlate/AudioMaterialSlateTypes.cpp
@@ -198,6 +198,17 @@ void FAudioMaterialMeterStyle::GetResources(TArray<const FSlateBrush*>& OutBrush
 	// Add any brush resources here so that Slate can correctly atlas and reference them
 }
 
+PRAGMA_DISABLE_DEPRECATION_WARNINGS
+FAudioMaterialMeterStyle& FAudioMaterialMeterStyle::SetFont(const FName& InFontName, uint16 InSize) { Font = FSlateFontInfo(InFontName, InSize); return *this; }
+FAudioMaterialMeterStyle& FAudioMaterialMeterStyle::SetFont(const FString& InFontName, uint16 InSize) { Font = FSlateFontInfo(*InFontName, InSize); return *this; }
+FAudioMaterialMeterStyle& FAudioMaterialMeterStyle::SetFont(const WIDECHAR* InFontName, uint16 InSize) { Font = FSlateFontInfo(InFontName, InSize); return *this; }
+FAudioMaterialMeterStyle& FAudioMaterialMeterStyle::SetFont(const ANSICHAR* InFontName, uint16 InSize) { Font = FSlateFontInfo(InFontName, InSize); return *this; }
+FAudioMaterialMeterStyle& FAudioMaterialMeterStyle::SetFontName(const FName& InFontName) { Font = FSlateFontInfo(InFontName, Font.Size); return *this; }
+FAudioMaterialMeterStyle& FAudioMaterialMeterStyle::SetFontName(const FString& InFontName) { Font = FSlateFontInfo(InFontName, Font.Size); return *this; }
+FAudioMaterialMeterStyle& FAudioMaterialMeterStyle::SetFontName(const WIDECHAR* InFontName) { Font = FSlateFontInfo(InFontName, Font.Size); return *this; }
+FAudioMaterialMeterStyle& FAudioMaterialMeterStyle::SetFontName(const ANSICHAR* InFontName) { Font = FSlateFontInfo(InFontName, Font.Size); return *this; }
+PRAGMA_ENABLE_DEPRECATION_WARNINGS
+
 FAudioMaterialEnvelopeStyle::FAudioMaterialEnvelopeStyle()
 	: CurveColor(FLinearColor::White)
 	, BackgroundColor(FLinearColor::Black)
diff --git a/Engine/Plugins/Runtime/AudioWidgets/Source/AudioWidgets/Private/AudioMeterStyle.cpp b/Engine/Plugins/Runtime/AudioWidgets/Source/AudioWidgets/Private/AudioMeterStyle.cpp
--- a/Engine/Plugins/Runtime/AudioWidgets/Source/AudioWidgets/Private/AudioMeterStyle.cpp
+++ b/Engine/Plugins/Runtime/AudioWidgets/Source/AudioWidgets/Private/AudioMeterStyle.cpp
@@ -36,4 +36,13 @@ const FAudioMeterStyle& FAudioMeterStyle::GetDefault()
 	return Default;
 }
 
-
+PRAGMA_DISABLE_DEPRECATION_WARNINGS
+FAudioMeterStyle& FAudioMeterStyle::SetFont(const FName& InFontName, uint16 InSize) { Font = FSlateFontInfo(InFontName, InSize); return *this; }
+FAudioMeterStyle& FAudioMeterStyle::SetFont(const FString& InFontName, uint16 InSize) { Font = FSlateFontInfo(*InFontName, InSize); return *this; }
+FAudioMeterStyle& FAudioMeterStyle::SetFont(const WIDECHAR* InFontName, uint16 InSize) { Font = FSlateFontInfo(InFontName, InSize); return *this; }
+FAudioMeterStyle& FAudioMeterStyle::SetFont(const ANSICHAR* InFontName, uint16 InSize) { Font = FSlateFontInfo(InFontName, InSize); return *this; }
+FAudioMeterStyle& FAudioMeterStyle::SetFontName(const FName& InFontName) { Font = FSlateFontInfo(InFontName, Font.Size); return *this; }
+FAudioMeterStyle& FAudioMeterStyle::SetFontName(const FString& InFontName) { Font = FSlateFontInfo(InFontName, Font.Size); return *this; }
+FAudioMeterStyle& FAudioMeterStyle::SetFontName(const WIDECHAR* InFontName) { Font = FSlateFontInfo(InFontName, Font.Size); return *this; }
+FAudioMeterStyle& FAudioMeterStyle::SetFontName(const ANSICHAR* InFontName) { Font = FSlateFontInfo(InFontName, Font.Size); return *this; }
+PRAGMA_ENABLE_DEPRECATION_WARNINGS
\ No newline at end of file
diff --git a/Engine/Plugins/Runtime/AudioWidgets/Source/AudioWidgets/Public/AudioMaterialSlate/AudioMaterialSlateTypes.h b/Engine/Plugins/Runtime/AudioWidgets/Source/AudioWidgets/Public/AudioMaterialSlate/AudioMaterialSlateTypes.h
--- a/Engine/Plugins/Runtime/AudioWidgets/Source/AudioWidgets/Public/AudioMaterialSlate/AudioMaterialSlateTypes.h
+++ b/Engine/Plugins/Runtime/AudioWidgets/Source/AudioWidgets/Public/AudioMaterialSlate/AudioMaterialSlateTypes.h
@@ -310,14 +310,24 @@ struct AUDIOWIDGETS_API FAudioMaterialMeterStyle : public FAudioMaterialWidgetSt
 	FAudioMaterialMeterStyle& SetFont(const FSlateFontInfo& InFont) { Font = InFont; return *this; }
 	FAudioMaterialMeterStyle& SetFont(TSharedPtr<const FCompositeFont> InCompositeFont, const int32 InSize, const FName& InTypefaceFontName = NAME_None) { Font = FSlateFontInfo(InCompositeFont, InSize, InTypefaceFontName); return *this; }
 	FAudioMaterialMeterStyle& SetFont(const UObject* InFontObject, const int32 InSize, const FName& InTypefaceFontName = NAME_None) { Font = FSlateFontInfo(InFontObject, InSize, InTypefaceFontName); return *this; }
-	FAudioMaterialMeterStyle& SetFont(const FName& InFontName, uint16 InSize) { Font = FSlateFontInfo(InFontName, InSize); return *this; }
-	FAudioMaterialMeterStyle& SetFont(const FString& InFontName, uint16 InSize) { Font = FSlateFontInfo(*InFontName, InSize); return *this; }
-	FAudioMaterialMeterStyle& SetFont(const WIDECHAR* InFontName, uint16 InSize) { Font = FSlateFontInfo(InFontName, InSize); return *this; }
-	FAudioMaterialMeterStyle& SetFont(const ANSICHAR* InFontName, uint16 InSize) { Font = FSlateFontInfo(InFontName, InSize); return *this; }
-	FAudioMaterialMeterStyle& SetFontName(const FName& InFontName) { Font = FSlateFontInfo(InFontName, Font.Size); return *this; }
-	FAudioMaterialMeterStyle& SetFontName(const FString& InFontName) { Font = FSlateFontInfo(InFontName, Font.Size); return *this; }
-	FAudioMaterialMeterStyle& SetFontName(const WIDECHAR* InFontName) { Font = FSlateFontInfo(InFontName, Font.Size); return *this; }
-	FAudioMaterialMeterStyle& SetFontName(const ANSICHAR* InFontName) { Font = FSlateFontInfo(InFontName, Font.Size); return *this; }
+
+	UE_DEPRECATED(5.6, "Use SetFont using FSlateFontInfo instead. FSlateFontInfo's constructors using a FontName are deprecated.")
+	FAudioMaterialMeterStyle& SetFont(const FName& InFontName, uint16 InSize);
+	UE_DEPRECATED(5.6, "Use SetFont using FSlateFontInfo instead. FSlateFontInfo's constructors using a FontName are deprecated.")
+	FAudioMaterialMeterStyle& SetFont(const FString& InFontName, uint16 InSize);
+	UE_DEPRECATED(5.6, "Use SetFont using FSlateFontInfo instead. FSlateFontInfo's constructors using a FontName are deprecated.")
+	FAudioMaterialMeterStyle& SetFont(const WIDECHAR* InFontName, uint16 InSize);
+	UE_DEPRECATED(5.6, "Use SetFont using FSlateFontInfo instead. FSlateFontInfo's constructors using a FontName are deprecated.")
+	FAudioMaterialMeterStyle& SetFont(const ANSICHAR* InFontName, uint16 InSize);
+	UE_DEPRECATED(5.6, "Use SetFont using FSlateFontInfo instead. FSlateFontInfo's constructors using a FontName are deprecated.")
+	FAudioMaterialMeterStyle& SetFontName(const FName& InFontName);
+	UE_DEPRECATED(5.6, "Use SetFont using FSlateFontInfo instead. FSlateFontInfo's constructors using a FontName are deprecated.")
+	FAudioMaterialMeterStyle& SetFontName(const FString& InFontName);
+	UE_DEPRECATED(5.6, "Use SetFont using FSlateFontInfo instead. FSlateFontInfo's constructors using a FontName are deprecated.")
+	FAudioMaterialMeterStyle& SetFontName(const WIDECHAR* InFontName);
+	UE_DEPRECATED(5.6, "Use SetFont using FSlateFontInfo instead. FSlateFontInfo's constructors using a FontName are deprecated.")
+	FAudioMaterialMeterStyle& SetFontName(const ANSICHAR* InFontName);
+
 	FAudioMaterialMeterStyle& SetFontSize(uint16 InSize) { Font.Size = InSize; return *this; }
 	FAudioMaterialMeterStyle& SetTypefaceFontName(const FName& InTypefaceFontName) { Font.TypefaceFontName = InTypefaceFontName; return *this; }
 
diff --git a/Engine/Plugins/Runtime/AudioWidgets/Source/AudioWidgets/Public/AudioMeterStyle.h b/Engine/Plugins/Runtime/AudioWidgets/Source/AudioWidgets/Public/AudioMeterStyle.h
--- a/Engine/Plugins/Runtime/AudioWidgets/Source/AudioWidgets/Public/AudioMeterStyle.h
+++ b/Engine/Plugins/Runtime/AudioWidgets/Source/AudioWidgets/Public/AudioMeterStyle.h
@@ -114,14 +114,24 @@ struct AUDIOWIDGETS_API FAudioMeterStyle : public FSlateWidgetStyle
 	FAudioMeterStyle& SetFont(const FSlateFontInfo& InFont) { Font = InFont; return *this; }
 	FAudioMeterStyle& SetFont(TSharedPtr<const FCompositeFont> InCompositeFont, const int32 InSize, const FName& InTypefaceFontName = NAME_None) { Font = FSlateFontInfo(InCompositeFont, InSize, InTypefaceFontName); return *this; }
 	FAudioMeterStyle& SetFont(const UObject* InFontObject, const int32 InSize, const FName& InTypefaceFontName = NAME_None) { Font = FSlateFontInfo(InFontObject, InSize, InTypefaceFontName); return *this; }
-	FAudioMeterStyle& SetFont(const FName& InFontName, uint16 InSize) { Font = FSlateFontInfo(InFontName, InSize); return *this; }
-	FAudioMeterStyle& SetFont(const FString& InFontName, uint16 InSize) { Font = FSlateFontInfo(*InFontName, InSize); return *this; }
-	FAudioMeterStyle& SetFont(const WIDECHAR* InFontName, uint16 InSize) { Font = FSlateFontInfo(InFontName, InSize); return *this; }
-	FAudioMeterStyle& SetFont(const ANSICHAR* InFontName, uint16 InSize) { Font = FSlateFontInfo(InFontName, InSize); return *this; }
-	FAudioMeterStyle& SetFontName(const FName& InFontName) { Font = FSlateFontInfo(InFontName, Font.Size); return *this; }
-	FAudioMeterStyle& SetFontName(const FString& InFontName) { Font = FSlateFontInfo(InFontName, Font.Size); return *this; }
-	FAudioMeterStyle& SetFontName(const WIDECHAR* InFontName) { Font = FSlateFontInfo(InFontName, Font.Size); return *this; }
-	FAudioMeterStyle& SetFontName(const ANSICHAR* InFontName) { Font = FSlateFontInfo(InFontName, Font.Size); return *this; }
+
+	UE_DEPRECATED(5.6, "Use SetFont using FSlateFontInfo instead. FSlateFontInfo's constructors using a FontName are deprecated.")
+	FAudioMeterStyle& SetFont(const FName& InFontName, uint16 InSize);
+	UE_DEPRECATED(5.6, "Use SetFont using FSlateFontInfo instead. FSlateFontInfo's constructors using a FontName are deprecated.")
+	FAudioMeterStyle& SetFont(const FString& InFontName, uint16 InSize);
+	UE_DEPRECATED(5.6, "Use SetFont using FSlateFontInfo instead. FSlateFontInfo's constructors using a FontName are deprecated.")
+	FAudioMeterStyle& SetFont(const WIDECHAR* InFontName, uint16 InSize);
+	UE_DEPRECATED(5.6, "Use SetFont using FSlateFontInfo instead. FSlateFontInfo's constructors using a FontName are deprecated.")
+	FAudioMeterStyle& SetFont(const ANSICHAR* InFontName, uint16 InSize);
+	UE_DEPRECATED(5.6, "Use SetFont using FSlateFontInfo instead. FSlateFontInfo's constructors using a FontName are deprecated.")
+	FAudioMeterStyle& SetFontName(const FName& InFontName);
+	UE_DEPRECATED(5.6, "Use SetFont using FSlateFontInfo instead. FSlateFontInfo's constructors using a FontName are deprecated.")
+	FAudioMeterStyle& SetFontName(const FString& InFontName);
+	UE_DEPRECATED(5.6, "Use SetFont using FSlateFontInfo instead. FSlateFontInfo's constructors using a FontName are deprecated.")
+	FAudioMeterStyle& SetFontName(const WIDECHAR* InFontName);
+	UE_DEPRECATED(5.6, "Use SetFont using FSlateFontInfo instead. FSlateFontInfo's constructors using a FontName are deprecated.")
+	FAudioMeterStyle& SetFontName(const ANSICHAR* InFontName);
+
 	FAudioMeterStyle& SetFontSize(uint16 InSize) { Font.Size = InSize; return *this; }
 	FAudioMeterStyle& SetTypefaceFontName(const FName& InTypefaceFontName) { Font.TypefaceFontName = InTypefaceFontName; return *this; }
 
diff --git a/Engine/Source/Runtime/SlateCore/Private/Styling/SlateTypes.cpp b/Engine/Source/Runtime/SlateCore/Private/Styling/SlateTypes.cpp
--- a/Engine/Source/Runtime/SlateCore/Private/Styling/SlateTypes.cpp
+++ b/Engine/Source/Runtime/SlateCore/Private/Styling/SlateTypes.cpp
@@ -112,6 +112,17 @@ const FTextBlockStyle& FTextBlockStyle::GetDefault()
 	return *Default.Get();
 }
 
+PRAGMA_DISABLE_DEPRECATION_WARNINGS
+FTextBlockStyle& FTextBlockStyle::SetFont(const FName& InFontName, float InSize) { Font = FSlateFontInfo(InFontName, InSize); return *this; }
+FTextBlockStyle& FTextBlockStyle::SetFont(const FString& InFontName, float InSize) { Font = FSlateFontInfo(*InFontName, InSize); return *this; }
+FTextBlockStyle& FTextBlockStyle::SetFont(const WIDECHAR* InFontName, float InSize) { Font = FSlateFontInfo(InFontName, InSize); return *this; }
+FTextBlockStyle& FTextBlockStyle::SetFont(const ANSICHAR* InFontName, float InSize) { Font = FSlateFontInfo(InFontName, InSize); return *this; }
+FTextBlockStyle& FTextBlockStyle::SetFontName(const FName& InFontName) { Font = FSlateFontInfo(InFontName, Font.GetClampSize()); return *this; }
+FTextBlockStyle& FTextBlockStyle::SetFontName(const FString& InFontName) { Font = FSlateFontInfo(InFontName, Font.GetClampSize()); return *this; }
+FTextBlockStyle& FTextBlockStyle::SetFontName(const WIDECHAR* InFontName) { Font = FSlateFontInfo(InFontName, Font.GetClampSize()); return *this; }
+FTextBlockStyle& FTextBlockStyle::SetFontName(const ANSICHAR* InFontName) { Font = FSlateFontInfo(InFontName, Font.GetClampSize()); return *this; }
+PRAGMA_ENABLE_DEPRECATION_WARNINGS
+
 FButtonStyle::FButtonStyle()
 : Normal()
 , Hovered()
@@ -269,6 +280,10 @@ const FEditableTextStyle& FEditableTextStyle::GetDefault()
 	return Default;
 }
 
+PRAGMA_DISABLE_DEPRECATION_WARNINGS
+FEditableTextStyle& FEditableTextStyle::SetFont(const FName& InFontName, float InSize) { Font = FSlateFontInfo(InFontName, InSize); return *this; }
+PRAGMA_ENABLE_DEPRECATION_WARNINGS
+
 FEditableTextBoxStyle::FEditableTextBoxStyle()
 	: BackgroundImageNormal()
 	, BackgroundImageHovered()
@@ -314,6 +329,9 @@ const FEditableTextBoxStyle& FEditableTextBoxStyle::GetDefault()
 	return Default;
 }
 
+PRAGMA_DISABLE_DEPRECATION_WARNINGS
+FEditableTextBoxStyle& FEditableTextBoxStyle::SetFont(const FName& InFontName, float InSize) { return SetFont(FSlateFontInfo(InFontName, InSize)); }
+PRAGMA_ENABLE_DEPRECATION_WARNINGS
 
 FInlineEditableTextBlockStyle::FInlineEditableTextBlockStyle()
 {
diff --git a/Engine/Source/Runtime/SlateCore/Public/Fonts/SlateFontInfo.h b/Engine/Source/Runtime/SlateCore/Public/Fonts/SlateFontInfo.h
--- a/Engine/Source/Runtime/SlateCore/Public/Fonts/SlateFontInfo.h
+++ b/Engine/Source/Runtime/SlateCore/Public/Fonts/SlateFontInfo.h
@@ -236,39 +236,43 @@ public:
 	SLATECORE_API FSlateFontInfo( const UObject* InFontObject, const float InSize, const FName& InTypefaceFontName = NAME_None, const FFontOutlineSettings& InOutlineSettings = FFontOutlineSettings());
 
 	/**
-	 * DEPRECATED - Creates and initializes a new instance with the specified font name and size.
+	 * Creates and initializes a new instance with the specified font name and size.
 	 *
 	 * @param InFontName The name of the font.
 	 * @param InSize The size of the font.
 	 * @param InHinting The type of hinting to use for the font.
 	 */
+	UE_DEPRECATED(5.6, "Use constructor with FCompositeFont instead.")
 	SLATECORE_API FSlateFontInfo( const FString& InFontName, float InSize, EFontHinting InHinting = EFontHinting::Default, const FFontOutlineSettings& InOutlineSettings = FFontOutlineSettings());
 
 	/**
-	 * DEPRECATED - Creates and initializes a new instance with the specified font name and size.
+	 * Creates and initializes a new instance with the specified font name and size.
 	 *
 	 * @param InFontName The name of the font.
 	 * @param InSize The size of the font.
 	 * @param InHinting The type of hinting to use for the font.
 	 */
+	UE_DEPRECATED(5.6, "Use constructor with FCompositeFont instead.")
 	SLATECORE_API FSlateFontInfo( const FName& InFontName, float InSize, EFontHinting InHinting = EFontHinting::Default );
 
 	/**
-	 * DEPRECATED - Creates and initializes a new instance with the specified font name and size.
+	 * Creates and initializes a new instance with the specified font name and size.
 	 *
 	 * @param InFontName The name of the font.
 	 * @param InSize The size of the font.
 	 * @param InHinting The type of hinting to use for the font.
 	 */
+	UE_DEPRECATED(5.6, "Use constructor with FCompositeFont instead.")
 	SLATECORE_API FSlateFontInfo( const ANSICHAR* InFontName, float InSize, EFontHinting InHinting = EFontHinting::Default );
 
 	/**
-	 * DEPRECATED - Creates and initializes a new instance with the specified font name and size.
+	 * Creates and initializes a new instance with the specified font name and size.
 	 *
 	 * @param InFontName The name of the font.
 	 * @param InSize The size of the font.
 	 * @param InHinting The type of hinting to use for the font.
 	 */
+	UE_DEPRECATED(5.6, "Use constructor with FCompositeFont instead.")
 	SLATECORE_API FSlateFontInfo( const WIDECHAR* InFontName, float InSize, EFontHinting InHinting = EFontHinting::Default );
 
 public:
diff --git a/Engine/Source/Runtime/SlateCore/Public/Styling/SlateTypes.h b/Engine/Source/Runtime/SlateCore/Public/Styling/SlateTypes.h
--- a/Engine/Source/Runtime/SlateCore/Public/Styling/SlateTypes.h
+++ b/Engine/Source/Runtime/SlateCore/Public/Styling/SlateTypes.h
@@ -343,14 +343,24 @@ struct FTextBlockStyle : public FSlateWidgetStyle
 	FTextBlockStyle& SetFont(const FSlateFontInfo& InFont) { Font = InFont; return *this; }
 	FTextBlockStyle& SetFont(TSharedPtr<const FCompositeFont> InCompositeFont, const float InSize, const FName& InTypefaceFontName = NAME_None) { Font = FSlateFontInfo(InCompositeFont, InSize, InTypefaceFontName); return *this; }
 	FTextBlockStyle& SetFont(const UObject* InFontObject, const float InSize, const FName& InTypefaceFontName = NAME_None) { Font = FSlateFontInfo(InFontObject, InSize, InTypefaceFontName); return *this; }
-	FTextBlockStyle& SetFont(const FName& InFontName, float InSize) { Font = FSlateFontInfo(InFontName, InSize); return *this; }
-	FTextBlockStyle& SetFont(const FString& InFontName, float InSize) { Font = FSlateFontInfo(*InFontName, InSize); return *this; }
-	FTextBlockStyle& SetFont(const WIDECHAR* InFontName, float InSize) { Font = FSlateFontInfo(InFontName, InSize); return *this; }
-	FTextBlockStyle& SetFont(const ANSICHAR* InFontName, float InSize) { Font = FSlateFontInfo(InFontName, InSize); return *this; }
-	FTextBlockStyle& SetFontName(const FName& InFontName) { Font = FSlateFontInfo(InFontName, Font.GetClampSize()); return *this; }
-	FTextBlockStyle& SetFontName(const FString& InFontName) { Font = FSlateFontInfo(InFontName, Font.GetClampSize()); return *this; }
-	FTextBlockStyle& SetFontName(const WIDECHAR* InFontName) { Font = FSlateFontInfo(InFontName, Font.GetClampSize()); return *this; }
-	FTextBlockStyle& SetFontName(const ANSICHAR* InFontName) { Font = FSlateFontInfo(InFontName, Font.GetClampSize()); return *this; }
+	
+	UE_DEPRECATED(5.6, "Use SetFont using FCompositeFont instead. FSlateFontInfo's constructors using a FontName are deprecated.")
+	FTextBlockStyle& SetFont(const FName& InFontName, float InSize);
+	UE_DEPRECATED(5.6, "Use SetFont using FCompositeFont instead. FSlateFontInfo's constructors using a FontName are deprecated.")
+	FTextBlockStyle& SetFont(const FString& InFontName, float InSize);
+	UE_DEPRECATED(5.6, "Use SetFont using FCompositeFont instead. FSlateFontInfo's constructors using a FontName are deprecated.")
+	FTextBlockStyle& SetFont(const WIDECHAR* InFontName, float InSize);
+	UE_DEPRECATED(5.6, "Use SetFont using FCompositeFont instead. FSlateFontInfo's constructors using a FontName are deprecated.")
+	FTextBlockStyle& SetFont(const ANSICHAR* InFontName, float InSize);
+	UE_DEPRECATED(5.6, "Use SetFont using FCompositeFont instead. FSlateFontInfo's constructors using a FontName are deprecated.")
+	FTextBlockStyle& SetFontName(const FName& InFontName);
+	UE_DEPRECATED(5.6, "Use SetFont using FCompositeFont instead. FSlateFontInfo's constructors using a FontName are deprecated.")
+	FTextBlockStyle& SetFontName(const FString& InFontName);
+	UE_DEPRECATED(5.6, "Use SetFont using FCompositeFont instead. FSlateFontInfo's constructors using a FontName are deprecated.")
+	FTextBlockStyle& SetFontName(const WIDECHAR* InFontName);
+	UE_DEPRECATED(5.6, "Use SetFont using FCompositeFont instead. FSlateFontInfo's constructors using a FontName are deprecated.")
+	FTextBlockStyle& SetFontName(const ANSICHAR* InFontName);
+
 	FTextBlockStyle& SetFontSize(float InSize) { Font.Size = InSize; return *this; }
 	FTextBlockStyle& SetTypefaceFontName(const FName& InTypefaceFontName) { Font.TypefaceFontName = InTypefaceFontName; return *this; }
 	FTextBlockStyle& SetFontMaterial(UObject* InMaterial) { Font.FontMaterial = InMaterial; return *this; }
@@ -866,7 +876,10 @@ struct FEditableTextStyle : public FSlateWidgetStyle
 	UPROPERTY(EditAnywhere, BlueprintReadWrite, Category=Appearance)
 	FSlateFontInfo Font;
 	FEditableTextStyle& SetFont(const FSlateFontInfo& InFont) { Font = InFont; return *this; }
-	FEditableTextStyle& SetFont(const FName& InFontName, float InSize) { Font = FSlateFontInfo(InFontName, InSize); return *this; }
+
+	UE_DEPRECATED(5.6, "Use SetFont using FSlateFontInfo instead. FSlateFontInfo's constructors using a FontName are deprecated.")
+	FEditableTextStyle& SetFont(const FName& InFontName, float InSize);
+
 	FEditableTextStyle& SetFontMaterial(UObject* InMaterial) { Font.FontMaterial = InMaterial; return *this; }
 	FEditableTextStyle& SetFontOutlineMaterial(UObject* InMaterial) { Font.OutlineSettings.OutlineMaterial = InMaterial; return *this; }
 
@@ -1043,7 +1056,9 @@ struct FEditableTextBoxStyle : public FSlateWidgetStyle
 	FSlateFontInfo Font_DEPRECATED;
 #endif
 	FEditableTextBoxStyle& SetFont(const FSlateFontInfo& InFont) { TextStyle.Font = InFont; return *this; }
-	FEditableTextBoxStyle& SetFont(const FName& InFontName, float InSize) { return SetFont(FSlateFontInfo(InFontName, InSize)); }
+
+	UE_DEPRECATED(5.6, "Use SetFont using FSlateFontInfo instead. FSlateFontInfo's constructors using a FontName are deprecated.")
+	FEditableTextBoxStyle& SetFont(const FName& InFontName, float InSize);
 
 	/** The style of the text block, which dictates the font, color, and shadow options. Style overrides all other properties! */
 	UPROPERTY(EditAnywhere, BlueprintReadWrite, Category = Appearance)
