Slate: Fix bad access to widget invalidatoin list.

diff --git a/Engine/Source/Runtime/SlateCore/Private/FastUpdate/SlateInvalidationWidgetList.cpp b/Engine/Source/Runtime/SlateCore/Private/FastUpdate/SlateInvalidationWidgetList.cpp
--- a/Engine/Source/Runtime/SlateCore/Private/FastUpdate/SlateInvalidationWidgetList.cpp
+++ b/Engine/Source/Runtime/SlateCore/Private/FastUpdate/SlateInvalidationWidgetList.cpp
@@ -1257,23 +1257,18 @@ void FSlateInvalidationWidgetList::Internal_RemoveRangeFromSameParent(const FInd
 		};
 
 		// Remove the other arrays that are between min and max	(ie. i, ii)
-		//if (Range.GetInclusiveMax().ArrayIndex - Range.GetInclusiveMin().ArrayIndex > 1)
 		if (!bRangeIsInSameElementArray)
 		{
 			const IndexType BeginArrayIndex = (IndexType)Data[Range.GetInclusiveMinWidgetIndex().ArrayIndex].NextArrayIndex;
-			const IndexType EndArrayIndex = (IndexType)Data[Range.GetInclusiveMaxWidgetIndex().ArrayIndex].PreviousArrayIndex;
+			const IndexType EndArrayIndex = Range.GetInclusiveMaxWidgetIndex().ArrayIndex;
 
-			if (BeginArrayIndex != Range.GetInclusiveMaxWidgetIndex().ArrayIndex)
+			IndexType CurrentArrayIndex = BeginArrayIndex;
+			while (CurrentArrayIndex != EndArrayIndex)
 			{
-				IndexType NextToRemove = BeginArrayIndex;
-				IndexType CurrentArrayIndex = NextToRemove;
-				do
-				{
-					ResetInvalidationWidgetHandle(CurrentArrayIndex, Data[CurrentArrayIndex].StartIndex, Data[CurrentArrayIndex].ElementList.Num());
-					CurrentArrayIndex = NextToRemove;
-					NextToRemove = (IndexType)Data[CurrentArrayIndex].NextArrayIndex;
-					RemoveDataNode(CurrentArrayIndex);
-				} while (CurrentArrayIndex != EndArrayIndex);
+				ResetInvalidationWidgetHandle(CurrentArrayIndex, Data[CurrentArrayIndex].StartIndex, Data[CurrentArrayIndex].ElementList.Num());
+				IndexType NextToRemove = (IndexType)Data[CurrentArrayIndex].NextArrayIndex;
+				RemoveDataNode(CurrentArrayIndex);
+				CurrentArrayIndex = NextToRemove;
 			}
 		}
 
