CIS Fix: Fix missing window style "window" referenced from color picker style set.

#rnx

diff --git a/Engine/Source/Runtime/SlateCore/Private/Styling/StarshipCoreStyle.cpp b/Engine/Source/Runtime/SlateCore/Private/Styling/StarshipCoreStyle.cpp
--- a/Engine/Source/Runtime/SlateCore/Private/Styling/StarshipCoreStyle.cpp
+++ b/Engine/Source/Runtime/SlateCore/Private/Styling/StarshipCoreStyle.cpp
@@ -1995,7 +1995,7 @@ void FStarshipCoreStyle::SetupColorPickerStyles(TSharedRef<FStyle>& Style)
 		Style->Set("ColorPicker.ColorThemesSmall", new IMAGE_BRUSH_SVG("Starship/Common/ColorThemes_16", Icon16x16));
 		Style->Set("ColorPicker.ColorThemesOff", new IMAGE_BRUSH_SVG("Starship/Common/ColorThemesOff_20", Icon20x20));
 
-		FWindowStyle CursorDecorator = Style->GetWidgetStyle<FWindowStyle>("Window");
+		FWindowStyle CursorDecorator = FWindowStyle();
 		CursorDecorator.SetCornerRadius(4);
 		Style->Set("ColorPicker.CursorDecorator", CursorDecorator);
 
