Font geometry normalization for signed distance field generator
For SDF-enabled font faces, this will increase font data size by roughly 5% (4.4 MB for 98 MB of largest fonts, more details below)

#tests Tested in editor
#rnx <Release Notes exclude>

Base + new preprocessed data sizes for specific fonts:
Noto Sans SC Black: 16.8 + 1.14 MiB
Noto Sans SC Regular: 16 + 0.74 MiB
Noto Sans TC Black: 11.4 + 0.88 MiB
Noto Sans TC Regular: 10.8 + 0.56 MiB
Noto Sans KR: 8.7 + 0.23 MiB
Noto Sans JP: 8.8 + 0.49 MiB
Tazugane for Epic Games: 5.7 + 0.25 MiB
NIS_JYAUN_P: 3 + 0.01 MiB
Noto Sans Italic: 919.2 + 23.1 KiB
Noto Sans Regular: 889.3 + 24.8 KiB
Noto Sans Bold: 889.3 + 27 KiB
Noto Sans Devanagari: 446.4 + 42.5 KiB
Burbank Big Condensed: 432.4 + 0.02 KiB
Noto Sans Arabic: 389.8 + 447.6 KiB
Heading Now: 251.5 + 0.16 KiB

diff --git a/Engine/Source/Editor/FontEditor/Private/FontFaceEditor.cpp b/Engine/Source/Editor/FontEditor/Private/FontFaceEditor.cpp
--- a/Engine/Source/Editor/FontEditor/Private/FontFaceEditor.cpp
+++ b/Engine/Source/Editor/FontEditor/Private/FontFaceEditor.cpp
@@ -493,6 +493,8 @@ void FFontFaceEditor::UpdatePreviewFonts()
 	const int PreviewFontNum = UE_ARRAY_COUNT(*PreviewTextBlocks);
 	if (IsFontFaceDistanceFieldEnabled())
 	{
+		// This ensures that font geometry is preprocessed before cloning the face, otherwise it would be needlessly redone for each copy.
+		FontFace->CacheSubFaces();
 		PreviewFaces.SetNum(PreviewFontNum, EAllowShrinking::No);
 		PreviewFonts.SetNum(PreviewFontNum, EAllowShrinking::No);
 		ClonePreviewFontFace(PreviewFaces[0], EFontRasterizationMode::Bitmap);
diff --git a/Engine/Source/Runtime/Core/Public/UObject/FortniteMainBranchObjectVersion.h b/Engine/Source/Runtime/Core/Public/UObject/FortniteMainBranchObjectVersion.h
--- a/Engine/Source/Runtime/Core/Public/UObject/FortniteMainBranchObjectVersion.h
+++ b/Engine/Source/Runtime/Core/Public/UObject/FortniteMainBranchObjectVersion.h
@@ -553,6 +553,9 @@ struct FFortniteMainBranchObjectVersion
 		// Added debug color for physical materials
 		DebugColorForPhysicalMaterials,
 
+		// Added PreprocessedFontGeometry to FFontFaceData
+		AddedPreprocessedFontGeometry,
+
 		// -----<new versions can be added above this line>-------------------------------------------------
 		VersionPlusOne,
 		LatestVersion = VersionPlusOne - 1
diff --git a/Engine/Source/Runtime/Engine/Private/FontFace.cpp b/Engine/Source/Runtime/Engine/Private/FontFace.cpp
--- a/Engine/Source/Runtime/Engine/Private/FontFace.cpp
+++ b/Engine/Source/Runtime/Engine/Private/FontFace.cpp
@@ -209,6 +209,12 @@ void UFontFace::CacheSubFaces()
 	if (FSlateApplication::IsInitialized())
 	{
 		SubFaces = FSlateApplication::Get().GetRenderer()->GetFontCache()->GetAvailableFontSubFaces(FontFaceData);
+#if WITH_EDITOR && WITH_FREETYPE
+		if (bEnableDistanceFieldRendering && GetDefault<UUserInterfaceSettings>()->bEnableDistanceFieldFontRasterization && IsSlateSdfTextFeatureEnabled())
+		{
+			FSlateApplication::Get().GetRenderer()->GetFontCache()->EnsurePreprocessedFontGeometry(FontFaceData);
+		}
+#endif
 	}
 }
 #endif // WITH_EDITORONLY_DATA
diff --git a/Engine/Source/Runtime/SlateCore/Private/Fonts/CompositeFont.cpp b/Engine/Source/Runtime/SlateCore/Private/Fonts/CompositeFont.cpp
--- a/Engine/Source/Runtime/SlateCore/Private/Fonts/CompositeFont.cpp
+++ b/Engine/Source/Runtime/SlateCore/Private/Fonts/CompositeFont.cpp
@@ -3,7 +3,9 @@
 #include "Fonts/CompositeFont.h"
 #include "SlateGlobals.h"
 #include "UObject/EditorObjectVersion.h"
+#include "UObject/FortniteMainBranchObjectVersion.h"
 #include "Fonts/FontFaceInterface.h"
+#include "Fonts/PreprocessedFontGeometry.h"
 #include "Templates/Casts.h"
 #include "Fonts/FontBulkData.h"
 
@@ -14,14 +16,76 @@
 // The total true type memory we are using for resident font faces
 DECLARE_MEMORY_STAT(TEXT("Resident Font Memory (TTF/OTF)"), STAT_SlateRawFontDataMemory, STATGROUP_SlateMemory);
 
+FFontFaceData::FFontFaceData()
+{
+}
+
+FFontFaceData::FFontFaceData(TArray<uint8>&& InData)
+	: Data(MoveTemp(InData))
+{
+	TrackMemoryUsage();
+}
+
+FFontFaceData::~FFontFaceData()
+{
+	UntrackMemoryUsage();
+}
+
+void FFontFaceData::SetData(TArray<uint8>&& InData)
+{
+	UntrackMemoryUsage();
+	Data = MoveTemp(InData);
+	PreprocessedFontGeometries.Reset();
+	TrackMemoryUsage();
+}
+
+void FFontFaceData::SetPreprocessedFontGeometries(TArray<UE::Slate::FPreprocessedFontGeometry>&& InPreprocessedFontGeometries)
+{
+	UntrackMemoryUsage();
+	PreprocessedFontGeometries = MoveTemp(InPreprocessedFontGeometries);
+	TrackMemoryUsage();
+}
+
+SIZE_T FFontFaceData::GetAllocatedSize() const
+{
+	SIZE_T TotalSize = Data.GetAllocatedSize()+PreprocessedFontGeometries.GetAllocatedSize();
+	for (const UE::Slate::FPreprocessedFontGeometry &SubFaceGeometry : PreprocessedFontGeometries)
+	{
+		TotalSize += SubFaceGeometry.GetAllocatedSize();
+	}
+	return TotalSize;
+}
+
+void FFontFaceData::Serialize(FArchive& Ar)
+{
+	Ar.UsingCustomVersion(FFortniteMainBranchObjectVersion::GUID);
+	if (Ar.IsLoading())
+	{
+		UntrackMemoryUsage();
+	}
+	Ar << Data;
+	if (Ar.CustomVer(FFortniteMainBranchObjectVersion::GUID) >= FFortniteMainBranchObjectVersion::AddedPreprocessedFontGeometry)
+	{
+		Ar << PreprocessedFontGeometries;
+	}
+	else if (Ar.IsLoading())
+	{
+		PreprocessedFontGeometries.Reset();
+	}
+	if (Ar.IsLoading())
+	{
+		TrackMemoryUsage();
+	}
+}
+
 void FFontFaceData::TrackMemoryUsage() const
 {
-	INC_DWORD_STAT_BY(STAT_SlateRawFontDataMemory, Data.GetAllocatedSize());
+	INC_DWORD_STAT_BY(STAT_SlateRawFontDataMemory, GetAllocatedSize());
 }
 
 void FFontFaceData::UntrackMemoryUsage() const
 {
-	DEC_DWORD_STAT_BY(STAT_SlateRawFontDataMemory, Data.GetAllocatedSize());
+	DEC_DWORD_STAT_BY(STAT_SlateRawFontDataMemory, GetAllocatedSize());
 }
 
 FFontData::FFontData()
diff --git a/Engine/Source/Runtime/SlateCore/Private/Fonts/FontCache.cpp b/Engine/Source/Runtime/SlateCore/Private/Fonts/FontCache.cpp
--- a/Engine/Source/Runtime/SlateCore/Private/Fonts/FontCache.cpp
+++ b/Engine/Source/Runtime/SlateCore/Private/Fonts/FontCache.cpp
@@ -14,6 +14,7 @@
 #include "Fonts/FontProviderInterface.h"
 #include "Fonts/FontCacheUtils.h"
 #include "Fonts/SlateSdfGenerator.h"
+#include "Fonts/FontGeometryPreprocessing.h"
 
 #include UE_INLINE_GENERATED_CPP_BY_NAME(FontCache)
 
@@ -1981,6 +1982,26 @@ SIZE_T FSlateFontCache::GetFontDataAssetResidentMemory(const UObject* FontDataAs
 	return CompositeFontCache->GetFontDataAssetResidentMemory(FontDataAsset);
 }
 
+#if WITH_EDITOR && WITH_FREETYPE
+bool FSlateFontCache::EnsurePreprocessedFontGeometry(FFontFaceDataRef FontFaceData)
+{
+	if (!IsSlateSdfTextFeatureEnabled() || !FontFaceData->GetPreprocessedFontGeometries().IsEmpty() || !FontFaceData->HasData())
+	{
+		return false;
+	}
+	const TArray<FString> SubFaces = FFreeTypeFace::GetAvailableSubFaces(FTLibrary.Get(), FontFaceData);
+	TArray<UE::Slate::FPreprocessedFontGeometry> PreprocessedFontGeometries;
+	PreprocessedFontGeometries.SetNum(SubFaces.Num());
+	for (int32 SubFaceIndex = 0; SubFaceIndex < SubFaces.Num(); ++SubFaceIndex)
+	{
+		const FFreeTypeFace FreeTypeFace(FTLibrary.Get(), FontFaceData, SubFaceIndex, EFontLayoutMethod::Metrics);
+		UE::Slate::PreprocessFontGeometry(PreprocessedFontGeometries[SubFaceIndex], FreeTypeFace.GetFace());
+	}
+	FontFaceData->SetPreprocessedFontGeometries((TArray<UE::Slate::FPreprocessedFontGeometry> &&) PreprocessedFontGeometries);
+	return true;
+}
+#endif
+
 int32 FSlateFontCache::GetAllFontTexturesIndex(const int32 InIndex) const
 {
 	int32 Index = InIndex;
diff --git a/Engine/Source/Runtime/SlateCore/Private/Fonts/FontCacheFreeType.cpp b/Engine/Source/Runtime/SlateCore/Private/Fonts/FontCacheFreeType.cpp
--- a/Engine/Source/Runtime/SlateCore/Private/Fonts/FontCacheFreeType.cpp
+++ b/Engine/Source/Runtime/SlateCore/Private/Fonts/FontCacheFreeType.cpp
@@ -12,6 +12,7 @@
 #include "HAL/PlatformProcess.h"
 #include "Trace/SlateMemoryTags.h"
 
+#include "Fonts/PreprocessedFontGeometry.h"
 #include "Fonts/UnicodeBlockRange.h"
 #include <limits>
 
@@ -420,6 +421,32 @@ FFreeTypeFace::FFreeTypeFace(const EFontLayoutMethod InLayoutMethod)
 	LayoutMethod = InLayoutMethod;
 }
 
+SIZE_T FFreeTypeFace::GetAllocatedMemorySize() const
+{
+#if WITH_FREETYPE
+	if (Memory.IsValid())
+	{
+		return Memory->GetAllocatedSize();
+	}
+#endif
+	return 0;
+}
+
+const UE::Slate::FPreprocessedFontGeometry* FFreeTypeFace::GetPreprocessedFontGeometry() const
+{
+#if WITH_FREETYPE
+	if (Memory.IsValid())
+	{
+		const int32 FaceIndex = FTFace->face_index&0x7fff; // Bits 0-15 according to FreeType documentation
+		if (Memory->GetPreprocessedFontGeometries().IsValidIndex(FaceIndex))
+		{
+			return &Memory->GetPreprocessedFontGeometries()[FaceIndex];
+		}
+	}
+#endif
+	return nullptr;
+}
+
 void FFreeTypeFace::FailAsyncLoad()
 {
 #if WITH_FREETYPE
@@ -449,6 +476,8 @@ void FFreeTypeFace::CompleteAsyncLoad(const FFreeTypeLibrary* InFTLibrary, FFont
 		FTFace = nullptr;
 	}
 
+	ensureAlwaysMsgf(FTFace == nullptr || (FTFace->face_index&0x7fff) == InFaceIndex, TEXT("FTFace->face_index value differs from requested face index! Avoid using the face_index value"));
+
 	ParseAttributes();
 
 	if (Memory->HasData())
diff --git a/Engine/Source/Runtime/SlateCore/Private/Fonts/FontCacheFreeType.h b/Engine/Source/Runtime/SlateCore/Private/Fonts/FontCacheFreeType.h
--- a/Engine/Source/Runtime/SlateCore/Private/Fonts/FontCacheFreeType.h
+++ b/Engine/Source/Runtime/SlateCore/Private/Fonts/FontCacheFreeType.h
@@ -330,14 +330,13 @@ public:
 	/**
 	 * Gets the memory size of the loaded font or 0 if the font is streamed
 	 */
-	FORCEINLINE SIZE_T GetAllocatedMemorySize() const
-	{
-#if WITH_FREETYPE
-		return Memory.IsValid() ? Memory->GetData().GetAllocatedSize() : 0;
-#else
-		return 0;
-#endif
-	}
+	SIZE_T GetAllocatedMemorySize() const;
+
+	/**
+	 * Returns the preprocessed font geometry object associated with the font data if available.
+	 * It can only be available if this was initialized with a FontFaceData object.
+	 */
+	const UE::Slate::FPreprocessedFontGeometry* GetPreprocessedFontGeometry() const;
 
 	void OverrideAscent(bool InOverride, int32 Value = 0)
 	{
diff --git a/Engine/Source/Runtime/SlateCore/Private/Fonts/FontGeometryPreprocessing.cpp b/Engine/Source/Runtime/SlateCore/Private/Fonts/FontGeometryPreprocessing.cpp
new file mode 100644
--- /dev/null
+++ b/Engine/Source/Runtime/SlateCore/Private/Fonts/FontGeometryPreprocessing.cpp
@@ -0,0 +1,127 @@
+// Copyright Epic Games, Inc. All Rights Reserved.
+
+#include "FontGeometryPreprocessing.h"
+
+#if WITH_EDITOR && WITH_FREETYPE
+
+#define MSDFGEN_USE_SKIA
+#include "ThirdParty/skia/skia-simplify.h"
+
+#define MSDFGEN_PARENT_NAMESPACE SlateMsdfgen
+#include "ThirdParty/msdfgen/msdfgen.h"
+
+namespace UE
+{
+namespace Slate
+{
+
+using namespace SlateMsdfgen;
+
+// Corners with an angle greater than 3 radians (~171 degrees) won't be treated as corners.
+static constexpr double SdfCornerAngleThreshold = 3.0;
+
+bool PreprocessFontGeometry(FPreprocessedFontGeometry& OutPreprocessedFontGeometry, FT_Face InFreeTypeFace)
+{
+	int NumCorrectWindings = 0;
+	int NumReverseWindings = 0;
+	/** Maps winding directions of each glyph in font:
+	 *      +1 = correct winding (and no preprocessing required)
+	 *      -1 = reversed winding
+	 *       0 = inconsistend winding or unnormalized geometry (individual glyph preprocessing data will be stored) or empty shape
+	 */
+	TArray<int8> GlyphWindingDirections;
+	GlyphWindingDirections.SetNumZeroed(InFreeTypeFace->num_glyphs);
+	for (FT_Long GlyphIndex = 0; GlyphIndex < InFreeTypeFace->num_glyphs; ++GlyphIndex)
+	{
+		// Load glyph into an MSDFgen Shape object
+		if (FT_Load_Glyph(InFreeTypeFace, GlyphIndex, FT_LOAD_NO_SCALE))
+		{
+			continue;
+		}
+		msdfgen::Shape MsdfgenShape;
+		if (msdfgen::readFreetypeOutline(MsdfgenShape, &InFreeTypeFace->glyph->outline, 1))
+		{
+			continue;
+		}
+		if (MsdfgenShape.contours.empty())
+		{
+			continue;
+		}
+
+		// Detect reversed or inconsistent contour windings
+		bool bPreprocessed = false;
+		bool bReversedWindings = false;
+		FPreprocessedGlyphGeometry PreprocessedGlyphGeometry;
+		FPreprocessedGlyphGeometry::FWindingFingerprint OriginalGlyphWinding(MsdfgenShape);
+		MsdfgenShape.orientContours();
+		FPreprocessedGlyphGeometry::FWindingFingerprint ResolvedGlyphWinding(MsdfgenShape);
+		if (ResolvedGlyphWinding != OriginalGlyphWinding)
+		{
+			PreprocessedGlyphGeometry = FPreprocessedGlyphGeometry(OriginalGlyphWinding, ResolvedGlyphWinding);
+			const FPreprocessedGlyphGeometryView PreprocessedGlyphGeometryView = PreprocessedGlyphGeometry.View();
+			// Restore MsdfgenShape to its initial state by re-applying contour reversals
+			PreprocessedGlyphGeometryView.UpdateWindings(MsdfgenShape);
+			// If contour windings are fully reversed, do not store individual glyph preprocessing data yet in case this is true for the majority of glyphs
+			if (PreprocessedGlyphGeometryView.ReversedWindings())
+			{
+				bReversedWindings = true;
+			}
+			else // Inconsistent windings case - list of contours to reverse must be stored either way
+			{
+				bPreprocessed = true;
+			}
+		}
+
+		// Detect self-intersecting geometry and store full geometry of the resolved shape
+		FPreprocessedGlyphGeometry::FGeometryFingerprint OriginalGeometry(MsdfgenShape);
+		msdfgen::resolveShapeGeometry(MsdfgenShape);
+		FPreprocessedGlyphGeometry::FGeometryFingerprint ResolvedGeometry(MsdfgenShape);
+		if (ResolvedGeometry != OriginalGeometry)
+		{
+			msdfgen::edgeColoringInkTrap(MsdfgenShape, SdfCornerAngleThreshold);
+			PreprocessedGlyphGeometry = FPreprocessedGlyphGeometry(MsdfgenShape);
+			bPreprocessed = true;
+		}
+
+		if (bPreprocessed)
+		{
+			OutPreprocessedFontGeometry.AddGlyph((int32) GlyphIndex, PreprocessedGlyphGeometry.View());
+		}
+		else
+		{
+			if (bReversedWindings)
+			{
+				++NumReverseWindings;
+				GlyphWindingDirections[GlyphIndex] = -1;
+			}
+			else
+			{
+				++NumCorrectWindings;
+				GlyphWindingDirections[GlyphIndex] = 1;
+			}
+		}
+	}
+
+	int8 AtypicalWinding = -1;
+	// If at least 2/3 of (non-preprocessed) glyphs have reversed windings, make that the default and mark correctly oriented glyphs as "reversed" to negate this.
+	if (NumReverseWindings > 2*NumCorrectWindings)
+	{
+		OutPreprocessedFontGeometry.SetGlobalWindingReversal(true);
+		AtypicalWinding = 1;
+	}
+	// Mark the winding direction deemed as "atypical" with the reversal flag
+	const FPreprocessedGlyphGeometryView GlyphWindingReversalView(FPreprocessedGlyphGeometry::FLAG_REVERSE_WINDINGS, 0, nullptr, nullptr, 0, 0);
+	for (int32 GlyphIndex = 0; GlyphIndex < GlyphWindingDirections.Num(); ++GlyphIndex)
+	{
+		if (GlyphWindingDirections[GlyphIndex] == AtypicalWinding)
+		{
+			OutPreprocessedFontGeometry.AddGlyph(GlyphIndex, GlyphWindingReversalView);
+		}
+	}
+	return true;
+}
+
+} // namespace Slate
+} // namespace UE
+
+#endif
diff --git a/Engine/Source/Runtime/SlateCore/Private/Fonts/FontGeometryPreprocessing.h b/Engine/Source/Runtime/SlateCore/Private/Fonts/FontGeometryPreprocessing.h
new file mode 100644
--- /dev/null
+++ b/Engine/Source/Runtime/SlateCore/Private/Fonts/FontGeometryPreprocessing.h
@@ -0,0 +1,24 @@
+// Copyright Epic Games, Inc. All Rights Reserved.
+
+#pragma once
+
+#include "CoreMinimal.h"
+#include "Fonts/FontCacheFreeType.h"
+#include "Fonts/PreprocessedFontGeometry.h"
+
+#if WITH_EDITOR && WITH_FREETYPE
+
+namespace UE
+{
+namespace Slate
+{
+
+/**
+ * Produces preprocessed font geometry object for the given FreeType font face. Returns true on success.
+ */
+bool PreprocessFontGeometry(FPreprocessedFontGeometry& OutPreprocessedFontGeometry, FT_Face InFreeTypeFace);
+
+}
+}
+
+#endif
diff --git a/Engine/Source/Runtime/SlateCore/Private/Fonts/PreprocessedFontGeometry.cpp b/Engine/Source/Runtime/SlateCore/Private/Fonts/PreprocessedFontGeometry.cpp
new file mode 100644
--- /dev/null
+++ b/Engine/Source/Runtime/SlateCore/Private/Fonts/PreprocessedFontGeometry.cpp
@@ -0,0 +1,464 @@
+// Copyright Epic Games, Inc. All Rights Reserved.
+
+#include "Fonts/PreprocessedFontGeometry.h"
+
+// Only enable Skia-dependent features of MSDFgen for Editor build
+#if WITH_EDITOR
+#define MSDFGEN_USE_SKIA
+#include "ThirdParty/skia/skia-simplify.h"
+#endif
+
+#define MSDFGEN_PARENT_NAMESPACE SlateMsdfgen
+#if !WITH_FREETYPE
+#define MSDFGEN_NO_FREETYPE
+#endif
+#include "ThirdParty/msdfgen/msdfgen.h"
+
+namespace UE
+{
+namespace Slate
+{
+
+using namespace SlateMsdfgen;
+
+struct FPreprocessedGlyphGeometry::FWindingFingerprint::FContourFingerprint
+{
+	const msdfgen::EdgeSegment* FirstSegment;
+	msdfgen::Point2 FirstPoint;
+};
+
+static int16 CompressCoordinate(double Value)
+{
+	return (int16) FMath::Clamp(FMath::RoundToInt(Value), -0x8000, 0x7fff);
+}
+
+static double DecompressCoordinate(int16 CompressedValue)
+{
+	return (double) CompressedValue;
+}
+
+static msdfgen::Point2 DecompressPoint(const int16* Coordinates)
+{
+	return msdfgen::Point2(DecompressCoordinate(Coordinates[0]), DecompressCoordinate(Coordinates[1]));
+}
+
+bool FPreprocessedGlyphGeometryView::UpdateGeometry(msdfgen::Shape& OutMsdfgenShape) const
+{
+	if (!(Flags&FPreprocessedGlyphGeometry::FLAG_FULL_GEOMETRY))
+	{
+		return false;
+	}
+	check(!(Flags&(FPreprocessedGlyphGeometry::FLAG_REVERSE_WINDINGS|FPreprocessedGlyphGeometry::FLAG_CONTOUR_WINDINGS)));
+	OutMsdfgenShape.inverseYAxis = (Flags&FPreprocessedGlyphGeometry::FLAG_INVERSE_Y_AXIS) != 0;
+	OutMsdfgenShape.contours.clear();
+	OutMsdfgenShape.contours.resize(ContourCount);
+	const uint8* CurEdgePtr = ContourDataPtr;
+	const int16* CurCoordinatePtr = CoordinateDataPtr;
+	const uint8* const EdgeEndPtr = ContourDataPtr+ContourDataLength;
+	const int16* const CoordinateEndPtr = CoordinateDataPtr+CoordinateDataLength;
+	for (int32 ContourIndex = 0; ContourIndex < ContourCount; ++ContourIndex)
+	{
+		const int16* const InitialContourCoordinate = CurCoordinatePtr;
+		msdfgen::Contour& OutContour = OutMsdfgenShape.contours[ContourIndex];
+		msdfgen::Point2 DummyLastPoint;
+		msdfgen::Point2* LastPointPtr = &DummyLastPoint;
+		OutContour.edges.clear();
+		bool ContourEnd = false;
+		while (!ContourEnd)
+		{
+			checkSlow(CurEdgePtr < EdgeEndPtr);
+			// Extract bits 4 through 6 as edge color (see FPreprocessedGlyphGeometry::ContourData description)
+			msdfgen::EdgeColor edgeColor = msdfgen::EdgeColor(*CurEdgePtr>>4&0x07u);
+			// The two least significant bits represent the number of control points in edge segment (see FPreprocessedGlyphGeometry::ContourData description)
+			switch (*CurEdgePtr&0x03u) {
+				case 1u:
+					{
+						checkSlow(CurCoordinatePtr+2 <= CoordinateEndPtr);
+						// New msdfgen::LinearSegment will be freed by msdfgen::EdgeHolder which is implicitly constructed in the following statement
+						msdfgen::LinearSegment* edgeSegment = new msdfgen::LinearSegment(
+							*LastPointPtr = DecompressPoint(CurCoordinatePtr),
+							msdfgen::Point2(),
+							edgeColor
+						);
+						OutContour.addEdge(edgeSegment);
+						LastPointPtr = &edgeSegment->p[1];
+						CurCoordinatePtr += 2;
+						break;
+					}
+				case 2u:
+					{
+						checkSlow(CurCoordinatePtr+4 <= CoordinateEndPtr);
+						// New msdfgen::QuadraticSegment will be freed by msdfgen::EdgeHolder which is implicitly constructed in the following statement
+						msdfgen::QuadraticSegment* edgeSegment = new msdfgen::QuadraticSegment(
+							*LastPointPtr = DecompressPoint(CurCoordinatePtr),
+							DecompressPoint(CurCoordinatePtr + 2),
+							msdfgen::Point2(),
+							edgeColor
+						);
+						OutContour.addEdge(edgeSegment);
+						LastPointPtr = &edgeSegment->p[2];
+						CurCoordinatePtr += 4;
+						break;
+					}
+				case 3u:
+					{
+						checkSlow(CurCoordinatePtr+6 <= CoordinateEndPtr);
+						// New msdfgen::CubicSegment will be freed by msdfgen::EdgeHolder which is implicitly constructed in the following statement
+						msdfgen::CubicSegment* edgeSegment = new msdfgen::CubicSegment(
+							*LastPointPtr = DecompressPoint(CurCoordinatePtr),
+							DecompressPoint(CurCoordinatePtr + 2),
+							DecompressPoint(CurCoordinatePtr + 4),
+							msdfgen::Point2(),
+							edgeColor
+						);
+						OutContour.addEdge(edgeSegment);
+						LastPointPtr = &edgeSegment->p[3];
+						CurCoordinatePtr += 6;
+						break;
+					}
+			}
+			ContourEnd = (*CurEdgePtr&FPreprocessedGlyphGeometry::FLAG_CONTOUR_END) != 0;
+			++CurEdgePtr;
+		}
+		if (InitialContourCoordinate < CurCoordinatePtr)
+		{
+			*LastPointPtr = DecompressPoint(InitialContourCoordinate);
+		}
+	}
+	return CurEdgePtr == EdgeEndPtr && CurCoordinatePtr == CoordinateEndPtr;
+}
+
+bool FPreprocessedGlyphGeometryView::UpdateWindings(msdfgen::Shape& InOutMsdfgenShape) const
+{
+	if (Flags&FPreprocessedGlyphGeometry::FLAG_REVERSE_WINDINGS)
+	{
+		check(!(Flags&(FPreprocessedGlyphGeometry::FLAG_CONTOUR_WINDINGS|FPreprocessedGlyphGeometry::FLAG_FULL_GEOMETRY)));
+		for (msdfgen::Contour& Contour : InOutMsdfgenShape.contours)
+		{
+			Contour.reverse();
+		}
+		return true;
+	}
+	if (!((Flags&FPreprocessedGlyphGeometry::FLAG_CONTOUR_WINDINGS) && InOutMsdfgenShape.contours.size() == ContourCount))
+	{
+		return false;
+	}
+	check(!(Flags&(FPreprocessedGlyphGeometry::FLAG_REVERSE_WINDINGS|FPreprocessedGlyphGeometry::FLAG_FULL_GEOMETRY)));
+	check(ContourDataLength == (ContourCount+7)/8);
+	int32 ContourIndex = 0;
+	const uint8* CurContourBits = ContourDataPtr;
+	const uint8* const ContourBitsEnd = ContourDataPtr+ContourDataLength;
+	while (ContourIndex < ContourCount && CurContourBits < ContourBitsEnd)
+	{
+		for (int32 Bit = 0; Bit < 8 && ContourIndex < ContourCount; ++Bit, ++ContourIndex)
+		{
+			if (*CurContourBits&1u<<Bit)
+			{
+				InOutMsdfgenShape.contours[ContourIndex].reverse();
+			}
+		}
+		++CurContourBits;
+	}
+	check(CurContourBits == ContourBitsEnd);
+	return ContourIndex == ContourCount;
+}
+
+bool FPreprocessedGlyphGeometryView::ReversedWindings() const
+{
+	return (Flags&FPreprocessedGlyphGeometry::FLAG_REVERSE_WINDINGS) != 0;
+}
+
+FPreprocessedGlyphGeometry::FWindingFingerprint::FWindingFingerprint(const msdfgen::Shape& MsdfgenShape)
+{
+	Contours.Reserve(MsdfgenShape.contours.size());
+	for (const msdfgen::Contour& MsdfgenContour : MsdfgenShape.contours)
+	{
+		if (MsdfgenContour.edges.empty())
+		{
+			Contours.Add(FContourFingerprint());
+		}
+		else
+		{
+			Contours.Add(FContourFingerprint
+				{
+					(const msdfgen::EdgeSegment*) MsdfgenContour.edges.front(),
+					MsdfgenContour.edges.front()->controlPoints()[1]
+				}
+			);
+		}
+	}
+}
+
+FPreprocessedGlyphGeometry::FWindingFingerprint::~FWindingFingerprint() = default;
+
+bool FPreprocessedGlyphGeometry::FWindingFingerprint::operator==(const FWindingFingerprint& Other) const
+{
+	if (Contours.Num() != Other.Contours.Num())
+	{
+		return false;
+	}
+	for (int32 Index = 0; Index < Contours.Num(); ++Index)
+	{
+		if (!(Contours[Index].FirstSegment == Other.Contours[Index].FirstSegment && Contours[Index].FirstPoint == Other.Contours[Index].FirstPoint))
+		{
+			return false;
+		}
+	}
+	return true;
+}
+
+bool FPreprocessedGlyphGeometry::FWindingFingerprint::operator!=(const FWindingFingerprint& Other) const
+{
+	return !operator==(Other);
+}
+
+bool FPreprocessedGlyphGeometry::FWindingFingerprint::Diff(bool& OutAllDiff, TArray<uint8>& OutDiffVector, const FWindingFingerprint &A, const FWindingFingerprint &B)
+{
+	if (A.Contours.Num() != B.Contours.Num())
+	{
+		return false;
+	}
+	OutAllDiff = true;
+	OutDiffVector.SetNumZeroed((A.Contours.Num()+7)/8);
+	for (int32 Index = 0; Index < A.Contours.Num(); ++Index)
+	{
+		if (!(A.Contours[Index].FirstSegment == B.Contours[Index].FirstSegment && A.Contours[Index].FirstPoint == B.Contours[Index].FirstPoint))
+		{
+			OutDiffVector[Index/8] |= 1u<<(Index%8);
+		}
+		else
+		{
+			OutAllDiff = false;
+		}
+	}
+	return true;
+}
+
+FPreprocessedGlyphGeometry::FGeometryFingerprint::FGeometryFingerprint(const msdfgen::Shape& MsdfgenShape)
+{
+	Contours.SetNum(MsdfgenShape.contours.size());
+	for (size_t ContourIndex = 0; ContourIndex < MsdfgenShape.contours.size(); ++ContourIndex)
+	{
+		for (const msdfgen::EdgeHolder& MsdfgenEdge : MsdfgenShape.contours[ContourIndex].edges)
+		{
+			const msdfgen::Point2 MsdfgenPoint = MsdfgenEdge->point(0);
+			const TPair<int32, int32> RoundedPoint(
+				FMath::RoundToInt32(MsdfgenPoint.x),
+				FMath::RoundToInt32(MsdfgenPoint.y)
+			);
+			bool MultiContourPoint = false;
+			for (size_t OtherContourIndex = 0; OtherContourIndex < ContourIndex; ++OtherContourIndex)
+			{
+				if (Contours[OtherContourIndex].Contains(RoundedPoint))
+				{
+					Contours[OtherContourIndex].Remove(RoundedPoint);
+					MultiContourPoint = true;
+					break;
+				}
+			}
+			if (!MultiContourPoint)
+			{
+				Contours[ContourIndex].Add(RoundedPoint);
+			}
+		}
+	}
+}
+
+bool FPreprocessedGlyphGeometry::FGeometryFingerprint::operator==(const FGeometryFingerprint& Other) const
+{
+	return OneWayMatch(*this, Other) && OneWayMatch(Other, *this);
+}
+
+bool FPreprocessedGlyphGeometry::FGeometryFingerprint::operator!=(const FGeometryFingerprint& Other) const
+{
+	return !operator==(Other);
+}
+
+bool FPreprocessedGlyphGeometry::FGeometryFingerprint::OneWayMatch(const FGeometryFingerprint& A, const FGeometryFingerprint& B)
+{
+	if (A.Contours.Num() != B.Contours.Num())
+	{
+		return false;
+	}
+	TArray<bool, TInlineAllocator<64> > UsedBContours;
+	UsedBContours.SetNumZeroed(A.Contours.Num());
+	for (const FContourVertices& AContour : A.Contours)
+	{
+		int32 MatchedBContourIndex = -1;
+		for (const TPair<int32, int32>& AVertex : AContour)
+		{
+			for (int32 BContourIndex = 0; BContourIndex < B.Contours.Num(); ++BContourIndex)
+			{
+				if (B.Contours[BContourIndex].Contains(AVertex))
+				{
+					if (BContourIndex != MatchedBContourIndex)
+					{
+						if (MatchedBContourIndex < 0)
+						{
+							if (UsedBContours[BContourIndex])
+							{
+								return false;
+							}
+							UsedBContours[BContourIndex] = true;
+							MatchedBContourIndex = BContourIndex;
+						}
+						else
+						{
+							return false;
+						}
+					}
+				}
+			}
+		}
+	}
+	return true;
+}
+
+FPreprocessedGlyphGeometry::FPreprocessedGlyphGeometry()
+	: Flags(0u)
+	, ContourCount(0)
+{ }
+
+FPreprocessedGlyphGeometry::FPreprocessedGlyphGeometry(const FWindingFingerprint& WindingFingerprintA, const FWindingFingerprint& WindingFingerprintB)
+	: Flags(0u)
+	, ContourCount(0)
+{
+	bool bAllDiff = false;
+	if (FWindingFingerprint::Diff(bAllDiff, ContourData, WindingFingerprintA, WindingFingerprintB))
+	{
+		if (bAllDiff)
+		{
+			Flags = FLAG_REVERSE_WINDINGS;
+			ContourData.Reset();
+		}
+		else
+		{
+			Flags = FLAG_CONTOUR_WINDINGS;
+			ContourCount = WindingFingerprintA.Contours.Num();
+		}
+	}
+}
+
+FPreprocessedGlyphGeometry::FPreprocessedGlyphGeometry(const msdfgen::Shape& MsdfgenShape)
+{
+	Flags = FLAG_FULL_GEOMETRY;
+	if (MsdfgenShape.inverseYAxis)
+	{
+		Flags |= FLAG_INVERSE_Y_AXIS;
+	}
+	ContourCount = (int32) MsdfgenShape.contours.size();
+	for (int32 ContourIndex = 0; ContourIndex < ContourCount; ++ContourIndex)
+	{
+		const msdfgen::Contour& MsdfgenContour = MsdfgenShape.contours[ContourIndex];
+		if (!MsdfgenContour.edges.empty())
+		{
+			for (const msdfgen::EdgeHolder& srcEdge : MsdfgenContour.edges)
+			{
+				const msdfgen::Point2* const MsdfgenPoints = srcEdge->controlPoints();
+				uint8 EdgeFlags = (uint8) ((unsigned) srcEdge->color<<4);
+				int16 Coordinates[6];
+				switch (srcEdge->type())
+				{
+					case (int) msdfgen::LinearSegment::EDGE_TYPE:
+						Coordinates[0] = CompressCoordinate(MsdfgenPoints[0].x);
+						Coordinates[1] = CompressCoordinate(MsdfgenPoints[0].y);
+						CoordinateData.Append(Coordinates, 2);
+						EdgeFlags |= 1u;
+						break;
+					case (int) msdfgen::QuadraticSegment::EDGE_TYPE:
+						Coordinates[0] = CompressCoordinate(MsdfgenPoints[0].x);
+						Coordinates[1] = CompressCoordinate(MsdfgenPoints[0].y);
+						Coordinates[2] = CompressCoordinate(MsdfgenPoints[1].x);
+						Coordinates[3] = CompressCoordinate(MsdfgenPoints[1].y);
+						CoordinateData.Append(Coordinates, 4);
+						EdgeFlags |= 2u;
+						break;
+					case (int) msdfgen::CubicSegment::EDGE_TYPE:
+						Coordinates[0] = CompressCoordinate(MsdfgenPoints[0].x);
+						Coordinates[1] = CompressCoordinate(MsdfgenPoints[0].y);
+						Coordinates[2] = CompressCoordinate(MsdfgenPoints[1].x);
+						Coordinates[3] = CompressCoordinate(MsdfgenPoints[1].y);
+						Coordinates[4] = CompressCoordinate(MsdfgenPoints[2].x);
+						Coordinates[5] = CompressCoordinate(MsdfgenPoints[2].y);
+						CoordinateData.Append(Coordinates, 6);
+						EdgeFlags |= 3u;
+						break;
+				}
+				ContourData.Add(EdgeFlags);
+			}
+			ContourData.Last() |= FLAG_CONTOUR_END;
+		}
+	}
+}
+
+FPreprocessedGlyphGeometryView FPreprocessedGlyphGeometry::View() const
+{
+	return FPreprocessedGlyphGeometryView(
+		Flags,
+		ContourCount,
+		ContourData.GetData(),
+		CoordinateData.GetData(),
+		ContourData.Num(),
+		CoordinateData.Num()
+	);
+}
+
+void FPreprocessedFontGeometry::AddGlyph(const int32 GlyphIndex, const FPreprocessedGlyphGeometryView& GlyphView)
+{
+	check(GlyphView.ContourCount >= 0 && GlyphView.ContourDataLength >= 0 && GlyphView.CoordinateDataLength >= 0);
+	FGlyphHeader GlyphHeader = { };
+	GlyphHeader.Flags = GlyphView.Flags;
+	GlyphHeader.ContourCount = GlyphView.ContourCount;
+	GlyphHeader.ContourDataStart = ContourData.Num();
+	GlyphHeader.ContourDataLength = GlyphView.ContourDataLength;
+	GlyphHeader.CoordinateDataStart = CoordinateData.Num();
+	GlyphHeader.CoordinateDataLength = GlyphView.CoordinateDataLength;
+	ContourData.Append(GlyphView.ContourDataPtr, GlyphView.ContourDataLength);
+	CoordinateData.Append(GlyphView.CoordinateDataPtr, GlyphView.CoordinateDataLength);
+	Glyphs.Add(GlyphIndex, GlyphHeader);
+}
+
+FPreprocessedGlyphGeometryView FPreprocessedFontGeometry::ViewGlyph(const int32 GlyphIndex) const
+{
+	if (const FGlyphHeader* GlyphHeader = Glyphs.Find(GlyphIndex))
+	{
+		uint8 GlyphFlags = GlyphHeader->Flags;
+		if (GlobalWindingReversal && !(GlyphFlags&~FPreprocessedGlyphGeometry::FLAG_REVERSE_WINDINGS&0xffu))
+		{
+			GlyphFlags ^= FPreprocessedGlyphGeometry::FLAG_REVERSE_WINDINGS;
+		}
+		return FPreprocessedGlyphGeometryView(
+			GlyphFlags,
+			GlyphHeader->ContourCount,
+			ContourData.GetData()+GlyphHeader->ContourDataStart,
+			CoordinateData.GetData()+GlyphHeader->CoordinateDataStart,
+			GlyphHeader->ContourDataLength,
+			GlyphHeader->CoordinateDataLength
+		);
+	}
+	return FPreprocessedGlyphGeometryView(
+		GlobalWindingReversal ? FPreprocessedGlyphGeometry::FLAG_REVERSE_WINDINGS : 0,
+		0,
+		nullptr,
+		nullptr,
+		0,
+		0
+	);
+}
+
+SIZE_T FPreprocessedFontGeometry::GetAllocatedSize() const
+{
+	return Glyphs.GetAllocatedSize() + ContourData.GetAllocatedSize() + CoordinateData.GetAllocatedSize();
+}
+
+FArchive &operator<<(FArchive &Ar, FPreprocessedFontGeometry &FontGeometry)
+{
+	return Ar
+		<< FontGeometry.GlobalWindingReversal
+		<< FontGeometry.Glyphs
+		<< FontGeometry.ContourData
+		<< FontGeometry.CoordinateData;
+}
+
+} // namespace Slate
+} // namespace UE
diff --git a/Engine/Source/Runtime/SlateCore/Private/Fonts/PreprocessedFontGeometry.h b/Engine/Source/Runtime/SlateCore/Private/Fonts/PreprocessedFontGeometry.h
new file mode 100644
--- /dev/null
+++ b/Engine/Source/Runtime/SlateCore/Private/Fonts/PreprocessedFontGeometry.h
@@ -0,0 +1,272 @@
+// Copyright Epic Games, Inc. All Rights Reserved.
+
+#pragma once
+
+#include "CoreMinimal.h"
+
+// Forward declaration of MSDFgen types
+namespace SlateMsdfgen
+{
+namespace msdfgen
+{
+class Shape;
+}
+typedef msdfgen::Shape FShape;
+}
+
+namespace UE
+{
+namespace Slate
+{
+
+class FPreprocessedFontGeometry;
+
+/** Representation of an FPreprocessedGlyphGeometry object stored elsewhere. */
+class FPreprocessedGlyphGeometryView
+{
+
+public:
+	friend FPreprocessedFontGeometry;
+
+	FPreprocessedGlyphGeometryView()
+		: Flags(0)
+		, ContourCount(0)
+		, ContourDataPtr(nullptr)
+		, CoordinateDataPtr(nullptr)
+		, ContourDataLength(0)
+		, CoordinateDataLength(0)
+	{ }
+
+	FPreprocessedGlyphGeometryView(
+		uint8 InFlags,
+		int32 InContourCount,
+		const uint8* InContourDataPtr,
+		const int16* InCoordinateDataPtr,
+		int32 InContourDataLength,
+		int32 InCoordinateDataLength
+	)
+		: Flags(InFlags)
+		, ContourCount(InContourCount)
+		, ContourDataPtr(InContourDataPtr)
+		, CoordinateDataPtr(InCoordinateDataPtr)
+		, ContourDataLength(InContourDataLength)
+		, CoordinateDataLength(InCoordinateDataLength)
+	{ }
+
+	/** If full geometry data is available, OutMsdfgenShape will be overwritten with this data. Otherwise, false is returned. */
+	bool UpdateGeometry(SlateMsdfgen::FShape& OutMsdfgenShape) const;
+
+	/** If winding data is available and the provided MSDFgen Shape is compatible, its contours will be reversed based on this data. Otherwise, false is returned. */
+	bool UpdateWindings(SlateMsdfgen::FShape& InOutMsdfgenShape) const;
+
+	/** Returns true if all contours have reversed windings. */
+	bool ReversedWindings() const;
+
+private:
+	/** A bitfield of glyph flags (see FLAG_ constants in FPreprocessedGlyphGeometry) */
+	uint8 Flags;
+	/** The number of contours the glyph has. May be zero if preprocessed data is not present. */
+	int32 ContourCount;
+	/** Pointer to the beginning of FGlyphGeometryFingerprint::ContourData */
+	const uint8* ContourDataPtr;
+	/** Pointer to the beginning of FGlyphGeometryFingerprint::CoordinateData */
+	const int16* CoordinateDataPtr;
+	/** Number of elements of the contour data array */
+	int32 ContourDataLength;
+	/** Number of elements of the coordinate data array */
+	int32 CoordinateDataLength;
+
+};
+
+/** Preprocessed geometry data for a single glyph. */
+class FPreprocessedGlyphGeometry
+{
+public:
+
+	/** Glyph flag which indicates that all contours need to be reversed (no contour data present) */
+	static constexpr uint8 FLAG_REVERSE_WINDINGS = 0x01;
+	/** Glyph flag which indicates that contour data represents a bit field of whether each contour's winding should be reversed. */
+	static constexpr uint8 FLAG_CONTOUR_WINDINGS = 0x02;
+	/** Glyph flag which indicates that contour data represents a sequence of edge segments. */
+	static constexpr uint8 FLAG_FULL_GEOMETRY = 0x04;
+	/** Glyph flag which indicates that the source msdfgen::Shape's inverseYAxis attribute is true. */
+	static constexpr uint8 FLAG_INVERSE_Y_AXIS = 0x10;
+
+	/** Edge segment flag which indicates that it is the last segment of the current contour and if another one follows, it initializes the next contour. */
+	static constexpr uint8 FLAG_CONTOUR_END = 0x08;
+
+	/**
+	 * Captures the windings of the contours of a glyph so that it can be later determined which ones have been reversed.
+	 */
+	class FWindingFingerprint
+	{
+
+	public:
+		friend FPreprocessedGlyphGeometry;
+
+		explicit FWindingFingerprint(const SlateMsdfgen::FShape& MsdfgenShape);
+		~FWindingFingerprint();
+
+		bool operator==(const FWindingFingerprint& Other) const;
+		bool operator!=(const FWindingFingerprint& Other) const;
+
+		/** Compares two fingerprints, outputs the different contours in OutDiffVector, and sets OutAllDiff to whether all are different. Returns false if the fingerprints are incompatible. */
+		static bool Diff(bool& OutAllDiff, TArray<uint8>& OutDiffVector, const FWindingFingerprint &A, const FWindingFingerprint &B);
+
+	private:
+		struct FContourFingerprint;
+
+		/** The fingerprints of individual contours. */
+		TArray<FContourFingerprint> Contours;
+
+	};
+
+	/**
+	 * Captures the geometry of a glyph so that it can be determined if contours were split or merged by simplification.
+	 */
+	class FGeometryFingerprint
+	{
+
+	public:
+		explicit FGeometryFingerprint(const SlateMsdfgen::FShape& MsdfgenShape);
+
+		bool operator==(const FGeometryFingerprint& Other) const;
+		bool operator!=(const FGeometryFingerprint& Other) const;
+
+	private:
+		typedef TSet<TPair<int32, int32> > FContourVertices;
+
+		/** Checks that B matches A but not if A matches B. */
+		static bool OneWayMatch(const FGeometryFingerprint& A, const FGeometryFingerprint& B);
+
+		/** The vertices of individual contours. */
+		TArray<FContourVertices> Contours;
+
+	};
+
+	FPreprocessedGlyphGeometry();
+	FPreprocessedGlyphGeometry(const FWindingFingerprint& WindingFingerprintA, const FWindingFingerprint& WindingFingerprintB);
+	explicit FPreprocessedGlyphGeometry(const SlateMsdfgen::FShape& MsdfgenShape);
+
+	/** Returns a FPreprocessedGlyphGeometryView object corresponding to this object. */
+	FPreprocessedGlyphGeometryView View() const;
+
+private:
+	/** A bitfield of glyph flags (see FLAG_ constants above) */
+	uint8 Flags;
+	/** The number of contours the glyph has. May be zero if preprocessed data is not present. */
+	int32 ContourCount;
+
+	/**
+	 * If Flags contains FLAG_CONTOUR_WINDINGS, this is a sequence of bits (8 per array element),
+	 * indicating which contours should be reversed.
+	 * If Flags contains FLAG_FULL_GEOMETRY, this is a sequence of bitfields for each edge segment:
+	 * With 0 being the least significant bit and 7 being the most significant bit of each array element,
+	 * bits 0 and 1 contain the number of the edge segment's control points
+	 *   (1 = linear edge segment, 2 = quadratic Bezier curve, 3 = cubic Bezier curve),
+	 * bit 3 may contain FLAG_CONTOUR_END, indicating the end of one contour and start of the next one,
+	 * and bits 4, 5, and 6 encode the edge's MSDF color (the red, green, and blue channel).
+	 * Bits 2 and 7 are unused.
+	 * Flags cannot contain both of the above, that would be an invalid state.
+	 * If Flags contain neither, ContourData is empty.
+	 */
+	TArray<uint8> ContourData;
+
+	/**
+	 * If Flags contains FLAG_FULL_GEOMETRY, contains a sequence of coordinates for each edge segment.
+	 * The number of coordinates to be read from CoordinateData is dictated by the number of control points
+	 * indicated in ContourData. The coordinates are listed in the same order as the edge segments
+	 * with the X coordinate always followed by the Y coordinate.
+	 * The closing points of contours are not included, so for each contour the initial point must be remembered.
+	 */
+	TArray<int16> CoordinateData;
+
+};
+
+/**
+ * This object represents additional per-glyph font geometry data, which override the base font file data
+ * for "unnormalized" glyphs. The original font data are still required and this object serves as their extension.
+ * If the font file data change for any reason, this object becomes invalid.
+ *
+ * A glyph's geometry is normalized when:
+ *  - it contains no self-intersections, that is, no part of its edge intersects another part
+ *  - and for each edge segment, looking from its initial point towards its next control point,
+ *    the filled portion of the glyph is always on the right side of the edge and empty portion is on the left side.
+ * 
+ * Normalized geometry is required for generating signed distance fields (directly from vector geometry).
+ *
+ * For glyphs that already satisfy both conditions (are normalized), no data is stored in this object.
+ * For glyphs that only satisfy the first condition and contain no false edge segments
+ * (ones that do not lie at the glyph's boundary), only a boolean array will be stored,
+ * which dictates which of its contours need to be reversed.
+ * For other unnormalized glyphs, the entire shape geometry is stored.
+ */
+class FPreprocessedFontGeometry
+{
+public:
+
+	/** Adds a single preprocessed glyph, identified by the numeric glyph index. */
+	void AddGlyph(const int32 GlyphIndex, const FPreprocessedGlyphGeometryView& GlyphView);
+
+	/**
+	 * Returns the view of the glyph identified by the numeric glyph index.
+	 * The view is only valid until this FPreprocessedFontGeometry object is modified.
+	 * If the glyph is not present in the preprocessed data, a valid object will be returned,
+	 * indicating no preprocessing is required.
+	 */
+	FPreprocessedGlyphGeometryView ViewGlyph(const int32 GlyphIndex) const;
+
+	/** Calling this will cause all glyphs not added to be reported with the FLAG_REVERSE_WINDINGS flag and those added with that flag as clean */
+	void SetGlobalWindingReversal(bool InGlobalWindingReversal)
+	{
+		GlobalWindingReversal = InGlobalWindingReversal;
+	}
+
+	/** Returns the object's total allocated size */
+	SIZE_T GetAllocatedSize() const;
+
+	/** Serializes the font geometry object to/from an archive */
+	friend FArchive& operator<<(FArchive& Ar, FPreprocessedFontGeometry& FontGeometry);
+
+private:
+
+	/** Contains fixed-size glyph properties and maps the glyph's data within ContourData and CoordinateData arrays. */
+	struct FGlyphHeader
+	{
+		/** Glyph flags - see constants in FPreprocessedGlyphGeometry */
+		uint8 Flags;
+		/** Number of glyph's contours */
+		int32 ContourCount;
+		/** Initial index of the glyph's contour data within the ContourData array */
+		int32 ContourDataStart;
+		/** Number of elements of the ContourData array for this glyph */
+		int32 ContourDataLength;
+		/** Initial index of the glyph's coordinates within the CoordinateData array */
+		int32 CoordinateDataStart;
+		/** Number of elements of the CoordinateData array for this glyph */
+		int32 CoordinateDataLength;
+	};
+
+	friend FArchive& operator<<(FArchive& Ar, FGlyphHeader& GlyphHeader) {
+		return Ar
+			<< GlyphHeader.Flags
+			<< GlyphHeader.ContourCount
+			<< GlyphHeader.ContourDataStart
+			<< GlyphHeader.ContourDataLength
+			<< GlyphHeader.CoordinateDataStart
+			<< GlyphHeader.CoordinateDataLength;
+	}
+
+	/** Indicates whether the entire font is encoded with the wrong contour winding.  */
+	bool GlobalWindingReversal = false;
+	/** Maps glyph headers of the included glyphs to glyph indices. */
+	TMap<int32, FGlyphHeader> Glyphs;
+	/** Concatenation of FPreprocessedGlyphGeometry::ContourData of all contained glyphs */
+	TArray<uint8> ContourData;
+	/** Concatenation of FPreprocessedGlyphGeometry::CoordinateData of all contained glyphs */
+	TArray<int16> CoordinateData;
+
+};
+
+} // namespace Slate
+} // namespace UE
diff --git a/Engine/Source/Runtime/SlateCore/Private/Fonts/SlateSdfGenerator.cpp b/Engine/Source/Runtime/SlateCore/Private/Fonts/SlateSdfGenerator.cpp
--- a/Engine/Source/Runtime/SlateCore/Private/Fonts/SlateSdfGenerator.cpp
+++ b/Engine/Source/Runtime/SlateCore/Private/Fonts/SlateSdfGenerator.cpp
@@ -8,16 +8,26 @@
 #include "Containers/StaticArray.h"
 #include "HAL/IConsoleManager.h"
 #include "SlateFontRenderer.h"
+#include "PreprocessedFontGeometry.h"
 #include "ProfilingDebugging/CsvProfiler.h"
 #include "Stats/StatsMisc.h"
-#include "ThirdParty/msdfgen/msdfgen.h"
 
-// Disable msdfgen FreeType features if FreeType is not available
-#if defined(MSDFGEN_USE_FREETYPE) && !WITH_FREETYPE
-#undef MSDFGEN_USE_FREETYPE
+// Only enable Skia-dependent features of MSDFgen for Editor build
+#if WITH_EDITOR
+#define MSDFGEN_USE_SKIA
+// Include Windows.h before Skia so that its unnecessary defines are cleaned up and Skia's include will be ignored
+#if PLATFORM_WINDOWS
+#include "Windows/WindowsHWrapper.h"
 #endif
-
+#include "ThirdParty/skia/skia-simplify.h"
 // Including source file directly, as seen in dlmalloc third-party dependency
+#include "ThirdParty/skia/skia-simplify.cpp"
+#endif // WITH_EDITOR
+
+#define MSDFGEN_PARENT_NAMESPACE SlateMsdfgen
+#if !WITH_FREETYPE
+#define MSDFGEN_NO_FREETYPE
+#endif
 #include "ThirdParty/msdfgen/msdfgen.cpp"
 
 CSV_DECLARE_CATEGORY_MODULE_EXTERN(SLATECORE_API, Slate);
@@ -35,6 +45,8 @@ static constexpr double SDF_BOUNDS_MITER_LIMIT = 1.0;
 namespace SdfUtils
 {
 
+using namespace SlateMsdfgen;
+
 /**
  * Specifies how the glyph's shape geometry will be mapped into the distance field's pixel coordinate system
  * and how the distance values are mapped to the 0 - 255 range.
@@ -211,29 +223,42 @@ bool FFreeTypeShapeBuilder::Build(TSharedPtr<FFreeTypeFace> InFace,
 		return false;
 	}
 
-	FT_Error Error = ::FT_Load_Glyph(InFace->GetFace(), 
-									 InGlyphIndex, 
-									 FT_LOAD_NO_SCALE
-									 | FT_LOAD_IGNORE_TRANSFORM
-									 | FT_LOAD_NO_HINTING
-									 | FT_LOAD_NO_AUTOHINT
-									 | FT_LOAD_NO_BITMAP);
-	if (Error != 0)
+	UE::Slate::FPreprocessedGlyphGeometryView PreprocessedGlyphGeometry;
+	if (const UE::Slate::FPreprocessedFontGeometry* PreprocessedFontGeometry = InFace->GetPreprocessedFontGeometry())
 	{
-		return false;
-	}
-	if (!FreeTypeUtils::IsGlyphEligibleForSdf(InFace->GetFace()->glyph)
-		|| InFace->GetFace()->glyph->metrics.width <= 0
-		|| InFace->GetFace()->glyph->metrics.height <= 0
-		|| InFace->GetFace()->glyph->outline.n_points <= 0)
-	{
-		return false;
+		PreprocessedGlyphGeometry = PreprocessedFontGeometry->ViewGlyph(InGlyphIndex);
 	}
 
-	Error = msdfgen::readFreetypeOutline(OutMsdfgenShape, &InFace->GetFace()->glyph->outline, 1);
-	if (Error != 0 || OutMsdfgenShape.contours.empty())
+	// Preprocessed glyph geometry may contain the full shape, only load it from FreeType if it does not
+	if (!PreprocessedGlyphGeometry.UpdateGeometry(OutMsdfgenShape))
 	{
-		return false;
+		FT_Error Error = ::FT_Load_Glyph(InFace->GetFace(), 
+										 InGlyphIndex, 
+										 FT_LOAD_NO_SCALE
+										 | FT_LOAD_IGNORE_TRANSFORM
+										 | FT_LOAD_NO_HINTING
+										 | FT_LOAD_NO_AUTOHINT
+										 | FT_LOAD_NO_BITMAP);
+		if (Error != 0)
+		{
+			return false;
+		}
+		if (!FreeTypeUtils::IsGlyphEligibleForSdf(InFace->GetFace()->glyph)
+			|| InFace->GetFace()->glyph->metrics.width <= 0
+			|| InFace->GetFace()->glyph->metrics.height <= 0
+			|| InFace->GetFace()->glyph->outline.n_points <= 0)
+		{
+			return false;
+		}
+
+		Error = msdfgen::readFreetypeOutline(OutMsdfgenShape, &InFace->GetFace()->glyph->outline, 1);
+		if (Error != 0 || OutMsdfgenShape.contours.empty())
+		{
+			return false;
+		}
+
+		// Preprocessed glyph data may also contain information about reversed contour windings. If not, this will have no effect.
+		PreprocessedGlyphGeometry.UpdateWindings(OutMsdfgenShape);
 	}
 
 	// MSDFgen uses bottom-up Y coordinates but UE uses top-down.
diff --git a/Engine/Source/Runtime/SlateCore/Public/Fonts/CompositeFont.h b/Engine/Source/Runtime/SlateCore/Public/Fonts/CompositeFont.h
--- a/Engine/Source/Runtime/SlateCore/Public/Fonts/CompositeFont.h
+++ b/Engine/Source/Runtime/SlateCore/Public/Fonts/CompositeFont.h
@@ -12,6 +12,14 @@
 class UFontBulkData;
 struct FFontFaceData;
 
+namespace UE
+{
+namespace Slate
+{
+class FPreprocessedFontGeometry;
+}
+}
+
 UENUM()
 enum class EFontHinting : uint8
 {
@@ -60,22 +68,13 @@ struct FFontFaceData
 {
 public:
 	/** Default constructor */
-	FFontFaceData()
-	{
-	}
+	SLATECORE_API FFontFaceData();
 
 	/** Construct from an existing block of font data */
-	FFontFaceData(TArray<uint8>&& InData)
-		: Data(MoveTemp(InData))
-	{
-		TrackMemoryUsage();
-	}
+	SLATECORE_API FFontFaceData(TArray<uint8>&& InData);
 
 	/** Destructor */
-	~FFontFaceData()
-	{
-		UntrackMemoryUsage();
-	}
+	SLATECORE_API ~FFontFaceData();
 
 	/** @return true if this instance has data set on it */
 	bool HasData() const
@@ -84,12 +83,10 @@ public:
 	}
 
 	/** Set from an existing block of font data */
-	void SetData(TArray<uint8>&& InData)
-	{
-		UntrackMemoryUsage();
-		Data = MoveTemp(InData);
-		TrackMemoryUsage();
-	}
+	SLATECORE_API void SetData(TArray<uint8>&& InData);
+	
+	/** Set the preprocessed font geometry data per sub-face for the given font data */
+	SLATECORE_API void SetPreprocessedFontGeometries(TArray<UE::Slate::FPreprocessedFontGeometry>&& InPreprocessedFontGeometries);
 
 	/** @return The data set on this instance */
 	const TArray<uint8>& GetData() const
@@ -97,20 +94,18 @@ public:
 		return Data;
 	}
 
-	/** Serialize the data on this instance to/from an archive */
-	void Serialize(FArchive& Ar)
+	/** @return The preprocessed font geometry per sub-face of the font represented by Data */
+	TArrayView<const UE::Slate::FPreprocessedFontGeometry> GetPreprocessedFontGeometries() const
 	{
-		if (Ar.IsLoading())
-		{
-			UntrackMemoryUsage();
-		}
-		Ar << Data;
-		if (Ar.IsLoading())
-		{
-			TrackMemoryUsage();
-		}
+		return PreprocessedFontGeometries;
 	}
 
+	/** @return the memory size of the raw font data and the associated preprocessed data */
+	SLATECORE_API SIZE_T GetAllocatedSize() const;
+
+	/** Serialize the data on this instance to/from an archive */
+	SLATECORE_API void Serialize(FArchive& Ar);
+
 	/** Helper function to construct a default font face data */
 	static FFontFaceDataRef MakeFontFaceData()
 	{
@@ -130,6 +125,9 @@ private:
 
 	/** Internal data */
 	TArray<uint8> Data;
+
+	/** Preprocessed font geometry data per font sub-face */
+	TArray<UE::Slate::FPreprocessedFontGeometry> PreprocessedFontGeometries;
 };
 
 /** Font face rasterization settings */
diff --git a/Engine/Source/Runtime/SlateCore/Public/Fonts/FontCache.h b/Engine/Source/Runtime/SlateCore/Public/Fonts/FontCache.h
--- a/Engine/Source/Runtime/SlateCore/Public/Fonts/FontCache.h
+++ b/Engine/Source/Runtime/SlateCore/Public/Fonts/FontCache.h
@@ -992,6 +992,10 @@ public:
 	 */
 	SLATECORE_API SIZE_T GetFontDataAssetResidentMemory(const UObject* FontDataAsset) const;
 
+#if WITH_EDITOR && WITH_FREETYPE
+	SLATECORE_API bool EnsurePreprocessedFontGeometry(FFontFaceDataRef FontFaceData);
+#endif
+
 private:
 	// Non-copyable
 	SLATECORE_API FSlateFontCache(const FSlateFontCache&);
