Fix for FSlateFontMeasure::MeasureStringInternal() generating access failure calling GetCharacter() multiple times
#rnx
#tests Opened problematic Niagara asset with -stompmalloc

diff --git a/Engine/Source/Runtime/SlateCore/Private/Fonts/FontMeasure.cpp b/Engine/Source/Runtime/SlateCore/Private/Fonts/FontMeasure.cpp
--- a/Engine/Source/Runtime/SlateCore/Private/Fonts/FontMeasure.cpp
+++ b/Engine/Source/Runtime/SlateCore/Private/Fonts/FontMeasure.cpp
@@ -207,7 +207,8 @@ FVector2f FSlateFontMeasure::MeasureStringInternal( FStringView Text, int32 Star
 		}
 		else
 		{
-			const FCharacterEntry& Entry = CharacterList.GetCharacter(CurrentChar, InFontInfo.FontFallback);
+			// copy the value as we may be invoking GetCharacter() below and it could invalidate our reference
+			const FCharacterEntry Entry = CharacterList.GetCharacter(CurrentChar, InFontInfo.FontFallback);
 
 			int32 Kerning = 0;
 			if( PreviousChar != 0 )
