InstancedStruct: Add support for custom icons

diff --git a/Engine/Source/Editor/StructUtilsEditor/Private/SInstancedStructPicker.cpp b/Engine/Source/Editor/StructUtilsEditor/Private/SInstancedStructPicker.cpp
--- a/Engine/Source/Editor/StructUtilsEditor/Private/SInstancedStructPicker.cpp
+++ b/Engine/Source/Editor/StructUtilsEditor/Private/SInstancedStructPicker.cpp
@@ -225,7 +225,7 @@ const FSlateBrush* SInstancedStructPicker::GetDisplayValueIcon() const
 	const UScriptStruct* CommonStruct = nullptr;
 	if (UE::StructUtils::Private::GetCommonScriptStruct(StructProperty, CommonStruct) == FPropertyAccess::Success)
 	{
-		return FSlateIconFinder::FindIconBrushForClass(UScriptStruct::StaticClass());
+		return FSlateIconFinder::FindIconBrushForClass(CommonStruct, "ClassIcon.Object");
 	}
 
 	return nullptr;
diff --git a/Engine/Source/Runtime/SlateCore/Private/Styling/SlateIconFinder.cpp b/Engine/Source/Runtime/SlateCore/Private/Styling/SlateIconFinder.cpp
--- a/Engine/Source/Runtime/SlateCore/Private/Styling/SlateIconFinder.cpp
+++ b/Engine/Source/Runtime/SlateCore/Private/Styling/SlateIconFinder.cpp
@@ -4,34 +4,32 @@
 #include "Styling/ISlateStyle.h"
 #include "Styling/SlateStyleRegistry.h"
 
-FSlateIcon FSlateIconFinder::FindIconForClass(const UClass* InClass, const FName& InDefaultName)
+FSlateIcon FSlateIconFinder::FindIconForClass(const UStruct* InStruct, const FName& InDefaultName)
 {
-	return FindCustomIconForClass(InClass, TEXT("ClassIcon"), InDefaultName);
+	return FindCustomIconForClass(InStruct, TEXT("ClassIcon"), InDefaultName);
 }
 
-FSlateIcon FSlateIconFinder::FindCustomIconForClass(const UClass* InClass, const TCHAR* StyleBasePath, const FName& InDefaultName)
+FSlateIcon FSlateIconFinder::FindCustomIconForClass(const UStruct* InStruct, const TCHAR* InStyleBasePath, const FName& InDefaultName)
 {
-	FSlateIcon Icon;
-
-	FString IconPath(StyleBasePath);
+	FString IconPath(InStyleBasePath);
 	IconPath += TEXT(".");
 	const int32 BastPathLength = IconPath.Len();
 
-	if (InClass)
+	if (InStruct)
 	{
-		// walk up class hierarchy until we find an icon
-		const UClass* CurrentClass = InClass;
-		while ( CurrentClass)
+		// walk up struct hierarchy until we find an icon
+		const UStruct* CurrentStruct = InStruct;
+		while (CurrentStruct)
 		{
-			CurrentClass->AppendName(IconPath);
-			Icon = FindIcon(*IconPath);
+			CurrentStruct->AppendName(IconPath);
+			FSlateIcon Icon = FindIcon(*IconPath);
 
 			if (Icon.IsSet())
 			{
 				return Icon;
 			}
 
-			CurrentClass = CurrentClass->GetSuperClass();
+			CurrentStruct = CurrentStruct->GetSuperStruct();
 			IconPath.RemoveAt(BastPathLength, IconPath.Len() - BastPathLength, EAllowShrinking::No);
 		}
 	}
@@ -46,16 +44,16 @@ FSlateIcon FSlateIconFinder::FindCustomIconForClass(const UClass* InClass, const
 	return FindIcon(InDefaultName);
 }
 
-FSlateIcon FSlateIconFinder::FindIcon(const FName& IconName)
+FSlateIcon FSlateIconFinder::FindIcon(const FName& InIconName)
 {
 	FSlateIcon Icon;
 
 	FSlateStyleRegistry::IterateAllStyles(
 		[&](const ISlateStyle& Style)
 		{
-			if (Style.GetOptionalBrush(IconName, nullptr, nullptr))
+			if (Style.GetOptionalBrush(InIconName, nullptr, nullptr))
 			{
-				Icon = FSlateIcon(Style.GetStyleSetName(), IconName);
+				Icon = FSlateIcon(Style.GetStyleSetName(), InIconName);
 				// terminate iteration
 				return false;
 			}
diff --git a/Engine/Source/Runtime/SlateCore/Public/Styling/SlateIconFinder.h b/Engine/Source/Runtime/SlateCore/Public/Styling/SlateIconFinder.h
--- a/Engine/Source/Runtime/SlateCore/Public/Styling/SlateIconFinder.h
+++ b/Engine/Source/Runtime/SlateCore/Public/Styling/SlateIconFinder.h
@@ -13,51 +13,51 @@ class FSlateIconFinder
 public:
 
 	/**
-	 * Find the icon to use for the supplied class
+	 * Find the icon to use for the supplied struct
 	 *
-	 * @param InClass			The class to locate an icon for
+	 * @param InStruct			The struct to locate an icon for
 	 * @param InDefaultName		The default icon name to use, if no specialized icon could be found
-	 * @return An FSlateIcon structure defining the class's icon
+	 * @return An FSlateIcon structure defining the struct's icon
 	 */
-	static SLATECORE_API FSlateIcon FindIconForClass(const UClass* InClass, const FName& InDefaultName = FName());
+	static SLATECORE_API FSlateIcon FindIconForClass(const UStruct* InStruct, const FName& InDefaultName = FName());
 
 	/**
-	 * Find a custom icon to use for the supplied class, according to the specified base style
+	 * Find a custom icon to use for the supplied struct, according to the specified base style
 	 *
-	 * @param InClass			The class to locate an icon for
+	 * @param InStruct			The struct to locate an icon for
 	 * @param InStyleBasePath	Style base path to use for the search (e.g. ClassIcon, or ClassThumbnail)
 	 * @param InDefaultName		The default icon name to use, if no specialized icon could be found
-	 * @return An FSlateIcon structure defining the class's icon
+	 * @return An FSlateIcon structure defining the struct's icon
 	 */
-	static SLATECORE_API FSlateIcon FindCustomIconForClass(const UClass* InClass, const TCHAR* StyleBasePath, const FName& InDefaultName = FName());
+	static SLATECORE_API FSlateIcon FindCustomIconForClass(const UStruct* InStruct, const TCHAR* InStyleBasePath, const FName& InDefaultName = FName());
 
 	/**
-	 * Find a slate brush to use for the supplied class's icon
+	 * Find a slate brush to use for the supplied struct's icon
 	 *
-	 * @param InClass			The class to locate an icon for
+	 * @param InStruct			The struct to locate an icon for
 	 * @param InDefaultName		The default icon name to use, if no specialized icon could be found
 	 * @return A slate brush, or nullptr if no icon was found
 	 */
-	static const FSlateBrush* FindIconBrushForClass(const UClass* InClass, const FName& InDefaultName = FName())
+	static const FSlateBrush* FindIconBrushForClass(const UStruct* InStruct, const FName& InDefaultName = FName())
 	{
-		return FindIconForClass(InClass, InDefaultName).GetIcon();
+		return FindIconForClass(InStruct, InDefaultName).GetIcon();
 	}
 	
 	/**
-	 * Find a custom icon to use for the supplied class, according to the specified base style
+	 * Find a custom icon to use for the supplied struct, according to the specified base style
 	 *
-	 * @param InClass			The class to locate an icon for
+	 * @param InStruct			The struct to locate an icon for
 	 * @param InStyleBasePath	Style base path to use for the search (e.g. ClassIcon, or ClassThumbnail)
 	 * @param InDefaultName		The default icon name to use, if no specialized icon could be found
 	 * @return A slate brush, or nullptr if no icon was found
 	 */
-	static const FSlateBrush* FindCustomIconBrushForClass(const UClass* InClass, const TCHAR* StyleBasePath, const FName& InDefaultName = FName())
+	static const FSlateBrush* FindCustomIconBrushForClass(const UStruct* InStruct, const TCHAR* InStyleBasePath, const FName& InDefaultName = FName())
 	{
-		return FindCustomIconForClass(InClass, StyleBasePath, InDefaultName).GetIcon();
+		return FindCustomIconForClass(InStruct, InStyleBasePath, InDefaultName).GetIcon();
 	}
 
 	/**
-	 * Find the first occurance of a brush represented by the specified IconName in any of the registered style sets
+	 * Find the first occurrence of a brush represented by the specified IconName in any of the registered style sets
 	 * @param InIconName 		The fully qualified style name of the icon to find
 	 * @return An FSlateIcon structure
 	 */
