Add New HasFontDataForCodepoint method in FSlateFontCache

#tests used an US editor and tested with Mandarin, English and Greek texts that it did detect the need to load a font correctly,. After the test FText was displayed and a font loaded, it then returned true.

diff --git a/Engine/Source/Runtime/SlateCore/Private/Fonts/FontCache.cpp b/Engine/Source/Runtime/SlateCore/Private/Fonts/FontCache.cpp
--- a/Engine/Source/Runtime/SlateCore/Private/Fonts/FontCache.cpp
+++ b/Engine/Source/Runtime/SlateCore/Private/Fonts/FontCache.cpp
@@ -1591,6 +1591,11 @@ const FFontData& FSlateFontCache::GetFontDataForCodepoint( const FSlateFontInfo&
 	return CompositeFontCache->GetFontDataForCodepoint(InFontInfo, InCodepoint, OutScalingFactor);
 }
 
+bool FSlateFontCache::HasLoadedFontForCodepoint(const FSlateFontInfo& InFontInfo, const UTF32CHAR InCodepoint) const
+{
+	return CompositeFontCache->HasLoadedFontForCodepoint(InFontInfo, InCodepoint);
+}
+
 uint16 FSlateFontCache::GetMaxCharacterHeight( const FSlateFontInfo& InFontInfo, float FontScale ) const
 {
 	return FontRenderer->GetMaxHeight(InFontInfo, FontScale);
diff --git a/Engine/Source/Runtime/SlateCore/Private/Fonts/FontCacheCompositeFont.cpp b/Engine/Source/Runtime/SlateCore/Private/Fonts/FontCacheCompositeFont.cpp
--- a/Engine/Source/Runtime/SlateCore/Private/Fonts/FontCacheCompositeFont.cpp
+++ b/Engine/Source/Runtime/SlateCore/Private/Fonts/FontCacheCompositeFont.cpp
@@ -301,6 +301,86 @@ const FFontData& FCompositeFontCache::GetDefaultFontData(const FSlateFontInfo& I
 	return DummyFontData;
 }
 
+bool FCompositeFontCache::HasLoadedFontForCodepoint (const FSlateFontInfo& InFontInfo, const UTF32CHAR InCodepoint)
+{
+	const FCompositeFont* const ResolvedCompositeFont = InFontInfo.GetCompositeFont();
+	const FCachedTypefaceData* const CachedDefaultTypefaceData = GetDefaultCachedTypeface(ResolvedCompositeFont);
+	const FCachedTypefaceData* const CachedTypefaceData = GetCachedTypefaceForCodepoint(ResolvedCompositeFont, InCodepoint);
+	if (!CachedTypefaceData)
+	{
+		return false;
+	}
+
+	auto IsFontValidAndNotLoading = [this](const FFontData* InFoundFontData)
+	{
+		if (!InFoundFontData)
+		{
+			return false;
+		}
+
+		// if this font supports the codepoint and it's not in loading state, we are good, no need for further tests
+		TSharedPtr<FFreeTypeFace> FaceAndMemory = FontFaceMap.FindRef(*InFoundFontData);
+		if (FaceAndMemory.IsValid() && FaceAndMemory->IsFaceValid() && !FaceAndMemory->IsFaceLoading())
+		{
+			return true;
+		}
+
+		return false;
+	};
+
+	auto FontDataMightSupportCodepoint = [this, InCodepoint, &IsFontValidAndNotLoading](const FFontData* InFoundFontData)
+	{
+		if (!InFoundFontData)
+		{
+			return false;
+		}
+
+		if (!IsFontValidAndNotLoading(InFoundFontData))
+		{
+			// Font not loaded; assume it can render the codepoint
+			return true;
+		}
+
+		return DoesFontDataSupportCodepoint(*InFoundFontData, InCodepoint);
+	};
+
+	// Check the preferred typeface for the font
+	const FFontData* FoundFontData = CachedTypefaceData->GetFontData(InFontInfo.TypefaceFontName);
+	if (!FoundFontData)
+	{
+		// if we can't find a font, return immediately false
+		return false;
+	}
+
+	if (FontDataMightSupportCodepoint(FoundFontData))
+	{
+		return true;
+	}
+
+	// Failing that, try and find a font by the attributes of the default font with the given name
+	if (CachedTypefaceData && CachedTypefaceData != CachedDefaultTypefaceData)
+	{
+		if (const FFontData* const FoundDefaultFontData = CachedDefaultTypefaceData->GetFontData(InFontInfo.TypefaceFontName))
+		{
+			const TSet<FName>& DefaultFontAttributes = GetFontAttributes(*FoundDefaultFontData);
+			FoundFontData = GetBestMatchFontForAttributes(CachedTypefaceData, DefaultFontAttributes);
+			if (FontDataMightSupportCodepoint(FoundFontData))
+			{
+				return true;
+			}
+		}
+	}
+
+	// if we are still try, try the primary font
+	FoundFontData = CachedTypefaceData->GetPrimaryFontData();
+	if (FontDataMightSupportCodepoint(FoundFontData))
+	{
+		return true;
+	}
+
+	return false;
+}
+
 const FFontData& FCompositeFontCache::GetFontDataForCodepoint(const FSlateFontInfo& InFontInfo, const UTF32CHAR InCodepoint, float& OutScalingFactor)
 {
 	static const FFontData DummyFontData;
diff --git a/Engine/Source/Runtime/SlateCore/Private/Fonts/FontCacheCompositeFont.h b/Engine/Source/Runtime/SlateCore/Private/Fonts/FontCacheCompositeFont.h
--- a/Engine/Source/Runtime/SlateCore/Private/Fonts/FontCacheCompositeFont.h
+++ b/Engine/Source/Runtime/SlateCore/Private/Fonts/FontCacheCompositeFont.h
@@ -214,6 +214,9 @@ public:
 	/** Get the font data to use for the given font info and codepoint */
 	const FFontData& GetFontDataForCodepoint(const FSlateFontInfo& InFontInfo, const UTF32CHAR InCodepoint, float& OutScalingFactor);
 
+	/** Return true if the font data to use for the given font info and codepoint is already loaded */
+	bool HasLoadedFontForCodepoint(const FSlateFontInfo& InFontInfo, const UTF32CHAR InCodepoint);
+
 	/** Gets or loads a FreeType font face */
 	TSharedPtr<FFreeTypeFace> GetFontFace(const FFontData& InFontData);
 
diff --git a/Engine/Source/Runtime/SlateCore/Public/Fonts/FontCache.h b/Engine/Source/Runtime/SlateCore/Public/Fonts/FontCache.h
--- a/Engine/Source/Runtime/SlateCore/Public/Fonts/FontCache.h
+++ b/Engine/Source/Runtime/SlateCore/Public/Fonts/FontCache.h
@@ -894,6 +894,16 @@ public:
 	 */
 	SLATECORE_API const FFontData& GetFontDataForCodepoint( const FSlateFontInfo& InFontInfo, const UTF32CHAR InCodepoint, float& OutScalingFactor ) const;
 
+	/**
+	 * Returns true if the font to use from the typeface associated with the given codepoint is already loaded
+	 *
+	 * @param InFontInfo		A descriptor of the font to get the typeface for
+	 * @param InCodepoint		The codepoint to get the typeface associated with
+	 *
+	 * @return FontData is already loaded
+	 */
+	SLATECORE_API bool HasLoadedFontForCodepoint(const FSlateFontInfo& InFontInfo, const UTF32CHAR InCodepoint) const;
+
 	/**
 	 * Returns the height of the largest character in the font. 
 	 *
