Remove glyph SDF generation times from log, which were reported for each individual glyph

#tests Tested in editor and PIE
#rnx <Release Notes exclude>

diff --git a/Engine/Source/Runtime/SlateCore/Private/Fonts/SlateSdfGenerator.cpp b/Engine/Source/Runtime/SlateCore/Private/Fonts/SlateSdfGenerator.cpp
--- a/Engine/Source/Runtime/SlateCore/Private/Fonts/SlateSdfGenerator.cpp
+++ b/Engine/Source/Runtime/SlateCore/Private/Fonts/SlateSdfGenerator.cpp
@@ -12,6 +12,9 @@
 #include "ProfilingDebugging/CsvProfiler.h"
 #include "Stats/StatsMisc.h"
 
+// This option will log the time it took to generate SDF for each glyph
+//#define SLATE_LOG_SDF_GENERATOR_TIME 1
+
 // Only enable Skia-dependent features of MSDFgen for Editor build
 #if WITH_EDITOR
 #define MSDFGEN_USE_SKIA
@@ -446,7 +449,9 @@ void FSdfGeneratorTask::EndCsvTrace()
 
 void FSdfGeneratorTask::DoWork()
 {
+	#ifdef SLATE_LOG_SDF_GENERATOR_TIME
 	SCOPE_LOG_TIME("FSdfGeneratorTask", NULL);
+	#endif
 	DoOutlineDecomposition();
 }
 
