[SlateCore]
* Changed so UnicodeBlockRange compiles in 50 seconds instead of 150 seconds (on msvc toolchain)

diff --git a/Engine/Source/Runtime/SlateCore/Private/Fonts/UnicodeBlockRange.cpp b/Engine/Source/Runtime/SlateCore/Private/Fonts/UnicodeBlockRange.cpp
--- a/Engine/Source/Runtime/SlateCore/Private/Fonts/UnicodeBlockRange.cpp
+++ b/Engine/Source/Runtime/SlateCore/Private/Fonts/UnicodeBlockRange.cpp
@@ -5,7 +5,7 @@
 TArrayView<const FUnicodeBlockRange> FUnicodeBlockRange::GetUnicodeBlockRanges()
 {
 	static const FUnicodeBlockRange UnicodeBlockRanges[] = {
-		#define REGISTER_UNICODE_BLOCK_RANGE(LOWERBOUND, UPPERBOUND, SYMBOLNAME, DISPLAYNAME) { EUnicodeBlockRange::SYMBOLNAME, DISPLAYNAME, FInt32Range(FInt32Range::BoundsType::Inclusive(LOWERBOUND), FInt32Range::BoundsType::Inclusive(UPPERBOUND)) },
+		#define REGISTER_UNICODE_BLOCK_RANGE(LOWERBOUND, UPPERBOUND, SYMBOLNAME, DISPLAYNAME) { EUnicodeBlockRange::SYMBOLNAME, DISPLAYNAME, LOWERBOUND, UPPERBOUND },
 		#include "Fonts/UnicodeBlockRange.inl"
 		#undef REGISTER_UNICODE_BLOCK_RANGE
 	};
@@ -18,3 +18,10 @@ FUnicodeBlockRange FUnicodeBlockRange::GetUnicodeBlockRange(const EUnicodeBlockR
 	TArrayView<const FUnicodeBlockRange> UnicodeBlockRanges = GetUnicodeBlockRanges();
 	return UnicodeBlockRanges[(int32)InBlockIndex];
 }
+
+FUnicodeBlockRange::FUnicodeBlockRange(EUnicodeBlockRange Index_, const FText& DisplayName_, int32 Lower, int32 Upper)
+:	Index(Index_)
+,	DisplayName(DisplayName_)
+,	Range(FInt32Range::BoundsType::Inclusive(Lower), FInt32Range::BoundsType::Inclusive(Upper))
+{
+}
diff --git a/Engine/Source/Runtime/SlateCore/Public/Fonts/UnicodeBlockRange.h b/Engine/Source/Runtime/SlateCore/Public/Fonts/UnicodeBlockRange.h
--- a/Engine/Source/Runtime/SlateCore/Public/Fonts/UnicodeBlockRange.h
+++ b/Engine/Source/Runtime/SlateCore/Public/Fonts/UnicodeBlockRange.h
@@ -33,4 +33,6 @@ struct FUnicodeBlockRange
 
 	/** Returns the block corresponding to the given enum */
 	static SLATECORE_API FUnicodeBlockRange GetUnicodeBlockRange(const EUnicodeBlockRange InBlockIndex);
+
+	SLATECORE_API FUnicodeBlockRange(EUnicodeBlockRange Index, const FText& DisplayName, int32 Lower, int32 Upper);
 };
