[SlateCore]
* Redid compile time optimization for UnicodeBlockRange. Changed so localization text is instead created inside ctor and made sure ctor is not inlined.

MSVC - Before: 50s  After: 0.3s
Clang - Before: 2.5s After 1.7s

diff --git a/Engine/Source/Runtime/SlateCore/Private/Fonts/UnicodeBlockRange.cpp b/Engine/Source/Runtime/SlateCore/Private/Fonts/UnicodeBlockRange.cpp
--- a/Engine/Source/Runtime/SlateCore/Private/Fonts/UnicodeBlockRange.cpp
+++ b/Engine/Source/Runtime/SlateCore/Private/Fonts/UnicodeBlockRange.cpp
@@ -5,7 +5,7 @@
 TArrayView<const FUnicodeBlockRange> FUnicodeBlockRange::GetUnicodeBlockRanges()
 {
 	static const FUnicodeBlockRange UnicodeBlockRanges[] = {
-		#define REGISTER_UNICODE_BLOCK_RANGE(LOWERBOUND, UPPERBOUND, SYMBOLNAME, DISPLAYNAME) { EUnicodeBlockRange::SYMBOLNAME, DISPLAYNAME, LOWERBOUND, UPPERBOUND },
+		#define REGISTER_UNICODE_BLOCK_RANGE(LOWERBOUND, UPPERBOUND, SYMBOLNAME, LOCTEXT_NAMESPACE, LOCTEXT_KEY, LOCTEXT_LITERAL) { EUnicodeBlockRange::SYMBOLNAME, TEXT(LOCTEXT_NAMESPACE), TEXT(LOCTEXT_KEY), TEXT(LOCTEXT_LITERAL), LOWERBOUND, UPPERBOUND },
 		#include "Fonts/UnicodeBlockRange.inl"
 		#undef REGISTER_UNICODE_BLOCK_RANGE
 	};
@@ -19,9 +19,9 @@ FUnicodeBlockRange FUnicodeBlockRange::GetUnicodeBlockRange(const EUnicodeBlockR
 	return UnicodeBlockRanges[(int32)InBlockIndex];
 }
 
-FUnicodeBlockRange::FUnicodeBlockRange(EUnicodeBlockRange Index_, const FText& DisplayName_, int32 Lower, int32 Upper)
+FUnicodeBlockRange::FUnicodeBlockRange(EUnicodeBlockRange Index_, const TCHAR* Namespace, const TCHAR* Key, const TCHAR* Literal, int32 Lower, int32 Upper)
 :	Index(Index_)
-,	DisplayName(DisplayName_)
+,	DisplayName(FText::AsLocalizable_Advanced(Namespace, Key, Literal))
 ,	Range(FInt32Range::BoundsType::Inclusive(Lower), FInt32Range::BoundsType::Inclusive(Upper))
 {
 }
diff --git a/Engine/Source/Runtime/SlateCore/Public/Fonts/UnicodeBlockRange.h b/Engine/Source/Runtime/SlateCore/Public/Fonts/UnicodeBlockRange.h
--- a/Engine/Source/Runtime/SlateCore/Public/Fonts/UnicodeBlockRange.h
+++ b/Engine/Source/Runtime/SlateCore/Public/Fonts/UnicodeBlockRange.h
@@ -11,7 +11,7 @@
 /** Enumeration of pre-defined Unicode block ranges that can be used to access entries from FUnicodeBlockRange */
 enum class EUnicodeBlockRange : uint16
 {
-#define REGISTER_UNICODE_BLOCK_RANGE(LOWERBOUND, UPPERBOUND, SYMBOLNAME, DISPLAYNAME) SYMBOLNAME,
+#define REGISTER_UNICODE_BLOCK_RANGE(LOWERBOUND, UPPERBOUND, SYMBOLNAME, LOCTEXT_NAMESPACE, LOCTEXT_KEY, LOCTEXT_LITERAL) SYMBOLNAME,
 #include "UnicodeBlockRange.inl"
 #undef REGISTER_UNICODE_BLOCK_RANGE
 };
@@ -34,5 +34,6 @@ struct FUnicodeBlockRange
 	/** Returns the block corresponding to the given enum */
 	static SLATECORE_API FUnicodeBlockRange GetUnicodeBlockRange(const EUnicodeBlockRange InBlockIndex);
 
-	SLATECORE_API FUnicodeBlockRange(EUnicodeBlockRange Index, const FText& DisplayName, int32 Lower, int32 Upper);
+	/** Ctor. Forced to not be inlined to improve compile times */
+	FORCENOINLINE FUnicodeBlockRange(EUnicodeBlockRange Index, const TCHAR* Namespace, const TCHAR* Key, const TCHAR* Literal, int32 Lower, int32 Upper);
 };
diff --git a/Engine/Source/Runtime/SlateCore/Public/Fonts/UnicodeBlockRange.inl b/Engine/Source/Runtime/SlateCore/Public/Fonts/UnicodeBlockRange.inl
--- a/Engine/Source/Runtime/SlateCore/Public/Fonts/UnicodeBlockRange.inl
+++ b/Engine/Source/Runtime/SlateCore/Public/Fonts/UnicodeBlockRange.inl
@@ -1,279 +1,279 @@
 // Copyright Epic Games, Inc. All Rights Reserved.
 
 // XMacro header
-// #define REGISTER_UNICODE_BLOCK_RANGE(LOWERBOUND, UPPERBOUND, SYMBOLNAME, DISPLAYNAME) before including this header
+// #define REGISTER_UNICODE_BLOCK_RANGE(LOWERBOUND, UPPERBOUND, SYMBOLNAME, LOCTEXT_NAMESPACE, LOCTEXT_KEY, LOCTEXT_LITERAL) before including this header
 
-REGISTER_UNICODE_BLOCK_RANGE(0x0000, 0x001f, ControlCharacter,							NSLOCTEXT("UnicodeBlock", "ControlCharacter", "Control Character"))
-REGISTER_UNICODE_BLOCK_RANGE(0x0020, 0x007f, BasicLatin,								NSLOCTEXT("UnicodeBlock", "BasicLatin", "Basic Latin"))
-REGISTER_UNICODE_BLOCK_RANGE(0x0080, 0x00ff, Latin1Supplement,							NSLOCTEXT("UnicodeBlock", "Latin1Supplement", "Latin - 1 Supplement"))
-REGISTER_UNICODE_BLOCK_RANGE(0x0100, 0x017f, LatinExtendedA,							NSLOCTEXT("UnicodeBlock", "LatinExtendedA", "Latin Extended - A"))
-REGISTER_UNICODE_BLOCK_RANGE(0x0180, 0x024f, LatinExtendedB,							NSLOCTEXT("UnicodeBlock", "LatinExtendedB", "Latin Extended - B"))
-REGISTER_UNICODE_BLOCK_RANGE(0x0250, 0x02af, IPAExtensions,								NSLOCTEXT("UnicodeBlock", "IPAExtensions", "IPA Extensions"))
-REGISTER_UNICODE_BLOCK_RANGE(0x02b0, 0x02ff, SpacingModifierLetters,					NSLOCTEXT("UnicodeBlock", "SpacingModifierLetters", "Spacing Modifier Letters"))
-REGISTER_UNICODE_BLOCK_RANGE(0x0300, 0x036f, CombiningDiacriticalMarks,					NSLOCTEXT("UnicodeBlock", "CombiningDiacriticalMarks", "Combining Diacritical Marks"))
-REGISTER_UNICODE_BLOCK_RANGE(0x0370, 0x03ff, GreekAndCoptic,							NSLOCTEXT("UnicodeBlock", "GreekAndCoptic", "Greek and Coptic"))
-REGISTER_UNICODE_BLOCK_RANGE(0x0400, 0x04ff, Cyrillic,									NSLOCTEXT("UnicodeBlock", "Cyrillic", "Cyrillic"))
-REGISTER_UNICODE_BLOCK_RANGE(0x0500, 0x052f, CyrillicSupplementary,						NSLOCTEXT("UnicodeBlock", "CyrillicSupplementary", "Cyrillic Supplementary"))
-REGISTER_UNICODE_BLOCK_RANGE(0x0530, 0x058f, Armenian,									NSLOCTEXT("UnicodeBlock", "Armenian", "Armenian"))
-REGISTER_UNICODE_BLOCK_RANGE(0x0590, 0x05ff, Hebrew,									NSLOCTEXT("UnicodeBlock", "Hebrew", "Hebrew"))
-REGISTER_UNICODE_BLOCK_RANGE(0x0600, 0x06ff, Arabic,									NSLOCTEXT("UnicodeBlock", "Arabic", "Arabic"))
-REGISTER_UNICODE_BLOCK_RANGE(0x0700, 0x074f, Syriac,									NSLOCTEXT("UnicodeBlock", "Syriac", "Syriac"))
-REGISTER_UNICODE_BLOCK_RANGE(0x0750, 0x077f, ArabicSupplement,							NSLOCTEXT("UnicodeBlock", "ArabicSupplement", "Arabic Supplement"))
-REGISTER_UNICODE_BLOCK_RANGE(0x0780, 0x07bf, Thaana,									NSLOCTEXT("UnicodeBlock", "Thaana", "Thaana"))
-REGISTER_UNICODE_BLOCK_RANGE(0x07c0, 0x07ff, NKo,										NSLOCTEXT("UnicodeBlock", "NKo", "NKo"))
-REGISTER_UNICODE_BLOCK_RANGE(0x0800, 0x083f, Samaritan,									NSLOCTEXT("UnicodeBlock", "Samaritan", "Samaritan"))
-REGISTER_UNICODE_BLOCK_RANGE(0x0840, 0x085f, Mandaic,									NSLOCTEXT("UnicodeBlock", "Mandaic", "Mandaic"))
-REGISTER_UNICODE_BLOCK_RANGE(0x08a0, 0x08ff, ArabicExtendedA,							NSLOCTEXT("UnicodeBlock", "ArabicExtendedA", "Arabic Extended - A"))
-REGISTER_UNICODE_BLOCK_RANGE(0x0900, 0x097f, Devanagari,								NSLOCTEXT("UnicodeBlock", "Devanagari", "Devanagari"))
-REGISTER_UNICODE_BLOCK_RANGE(0x0980, 0x09ff, Bengali,									NSLOCTEXT("UnicodeBlock", "Bengali", "Bengali"))
-REGISTER_UNICODE_BLOCK_RANGE(0x0a00, 0x0a7f, Gurmukhi,									NSLOCTEXT("UnicodeBlock", "Gurmukhi", "Gurmukhi"))
-REGISTER_UNICODE_BLOCK_RANGE(0x0a80, 0x0aff, Gujarati,									NSLOCTEXT("UnicodeBlock", "Gujarati", "Gujarati"))
-REGISTER_UNICODE_BLOCK_RANGE(0x0b00, 0x0b7f, Oriya,										NSLOCTEXT("UnicodeBlock", "Oriya", "Oriya"))
-REGISTER_UNICODE_BLOCK_RANGE(0x0b80, 0x0bff, Tamil,										NSLOCTEXT("UnicodeBlock", "Tamil", "Tamil"))
-REGISTER_UNICODE_BLOCK_RANGE(0x0c00, 0x0c7f, Telugu,									NSLOCTEXT("UnicodeBlock", "Telugu", "Telugu"))
-REGISTER_UNICODE_BLOCK_RANGE(0x0c80, 0x0cff, Kannada,									NSLOCTEXT("UnicodeBlock", "Kannada", "Kannada"))
-REGISTER_UNICODE_BLOCK_RANGE(0x0d00, 0x0d7f, Malayalam,									NSLOCTEXT("UnicodeBlock", "Malayalam", "Malayalam"))
-REGISTER_UNICODE_BLOCK_RANGE(0x0d80, 0x0dff, Sinhala,									NSLOCTEXT("UnicodeBlock", "Sinhala", "Sinhala"))
-REGISTER_UNICODE_BLOCK_RANGE(0x0e00, 0x0e7f, Thai,										NSLOCTEXT("UnicodeBlock", "Thai", "Thai"))
-REGISTER_UNICODE_BLOCK_RANGE(0x0e80, 0x0eff, Lao,										NSLOCTEXT("UnicodeBlock", "Lao", "Lao"))
-REGISTER_UNICODE_BLOCK_RANGE(0x0f00, 0x0fff, Tibetan,									NSLOCTEXT("UnicodeBlock", "Tibetan", "Tibetan"))
-REGISTER_UNICODE_BLOCK_RANGE(0x1000, 0x109f, Myanmar,									NSLOCTEXT("UnicodeBlock", "Myanmar", "Myanmar"))
-REGISTER_UNICODE_BLOCK_RANGE(0x10a0, 0x10ff, Georgian,									NSLOCTEXT("UnicodeBlock", "Georgian", "Georgian"))
-REGISTER_UNICODE_BLOCK_RANGE(0x1100, 0x11ff, HangulJamo,								NSLOCTEXT("UnicodeBlock", "HangulJamo", "Hangul Jamo"))
-REGISTER_UNICODE_BLOCK_RANGE(0x1200, 0x137f, Ethiopic,									NSLOCTEXT("UnicodeBlock", "Ethiopic", "Ethiopic"))
-REGISTER_UNICODE_BLOCK_RANGE(0x1380, 0x139f, EthiopicSupplement,						NSLOCTEXT("UnicodeBlock", "EthiopicSupplement", "Ethiopic Supplement"))
-REGISTER_UNICODE_BLOCK_RANGE(0x13a0, 0x13ff, Cherokee,									NSLOCTEXT("UnicodeBlock", "Cherokee", "Cherokee"))
-REGISTER_UNICODE_BLOCK_RANGE(0x1400, 0x167f, UnifiedCanadianAboriginalSyllabics,		NSLOCTEXT("UnicodeBlock", "UnifiedCanadianAboriginalSyllabics", "Unified Canadian Aboriginal Syllabics"))
-REGISTER_UNICODE_BLOCK_RANGE(0x1680, 0x169f, Ogham,										NSLOCTEXT("UnicodeBlock", "Ogham", "Ogham"))
-REGISTER_UNICODE_BLOCK_RANGE(0x16a0, 0x16ff, Runic,										NSLOCTEXT("UnicodeBlock", "Runic", "Runic"))
-REGISTER_UNICODE_BLOCK_RANGE(0x1700, 0x171f, Tagalog,									NSLOCTEXT("UnicodeBlock", "Tagalog", "Tagalog"))
-REGISTER_UNICODE_BLOCK_RANGE(0x1720, 0x173f, Hanunoo,									NSLOCTEXT("UnicodeBlock", "Hanunoo", "Hanunoo"))
-REGISTER_UNICODE_BLOCK_RANGE(0x1740, 0x175f, Buhid,										NSLOCTEXT("UnicodeBlock", "Buhid", "Buhid"))
-REGISTER_UNICODE_BLOCK_RANGE(0x1760, 0x177f, Tagbanwa,									NSLOCTEXT("UnicodeBlock", "Tagbanwa", "Tagbanwa"))
-REGISTER_UNICODE_BLOCK_RANGE(0x1780, 0x17ff, Khmer,										NSLOCTEXT("UnicodeBlock", "Khmer", "Khmer"))
-REGISTER_UNICODE_BLOCK_RANGE(0x1800, 0x18af, Mongolian,									NSLOCTEXT("UnicodeBlock", "Mongolian", "Mongolian"))
-REGISTER_UNICODE_BLOCK_RANGE(0x18b0, 0x18ff, UnifiedCanadianAboriginalSyllabicsExtended,NSLOCTEXT("UnicodeBlock", "UnifiedCanadianAboriginalSyllabicsExtended", "Unified Canadian Aboriginal Syllabics Extended"))
-REGISTER_UNICODE_BLOCK_RANGE(0x1900, 0x194f, Limbu,										NSLOCTEXT("UnicodeBlock", "Limbu", "Limbu"))
-REGISTER_UNICODE_BLOCK_RANGE(0x1950, 0x197f, TaiLe,										NSLOCTEXT("UnicodeBlock", "TaiLe", "Tai Le"))
-REGISTER_UNICODE_BLOCK_RANGE(0x1980, 0x19Df, NewTaiLue,									NSLOCTEXT("UnicodeBlock", "NewTaiLue", "New Tai Lue"))
-REGISTER_UNICODE_BLOCK_RANGE(0x19e0, 0x19ff, KhmerSymbols,								NSLOCTEXT("UnicodeBlock", "KhmerSymbols", "Khmer Symbols"))
-REGISTER_UNICODE_BLOCK_RANGE(0x1a00, 0x1a1f, Buginese,									NSLOCTEXT("UnicodeBlock", "Buginese", "Buginese"))
-REGISTER_UNICODE_BLOCK_RANGE(0x1a20, 0x1aaf, TaiTham,									NSLOCTEXT("UnicodeBlock", "TaiTham", "Tai Tham"))
-REGISTER_UNICODE_BLOCK_RANGE(0x1ab0, 0x1aff, CombiningDiacriticalMarksExtended,			NSLOCTEXT("UnicodeBlock", "CombiningDiacriticalMarksExtended", "Combining Diacritical Marks Extended"))
-REGISTER_UNICODE_BLOCK_RANGE(0x1b00, 0x1b7f, Balinese,									NSLOCTEXT("UnicodeBlock", "Balinese", "Balinese"))
-REGISTER_UNICODE_BLOCK_RANGE(0x1b80, 0x1bbf, Sundanese,									NSLOCTEXT("UnicodeBlock", "Sundanese", "Sundanese"))
-REGISTER_UNICODE_BLOCK_RANGE(0x1bc0, 0x1bff, Batak,										NSLOCTEXT("UnicodeBlock", "Batak", "Batak"))
-REGISTER_UNICODE_BLOCK_RANGE(0x1c00, 0x1c4f, Lepcha,									NSLOCTEXT("UnicodeBlock", "Lepcha", "Lepcha"))
-REGISTER_UNICODE_BLOCK_RANGE(0x1c50, 0x1c7f, OlChiki,									NSLOCTEXT("UnicodeBlock", "OlChiki", "Ol Chiki"))
-REGISTER_UNICODE_BLOCK_RANGE(0x1c80, 0x1c87, CyrillicExtendedC,							NSLOCTEXT("UnicodeBlock", "CyrillicExtendedC", "Cyrillic Extended C"))
-REGISTER_UNICODE_BLOCK_RANGE(0x1cc0, 0x1ccf, SundaneseSupplement,						NSLOCTEXT("UnicodeBlock", "SundaneseSupplement", "Sundanese Supplement"))
-REGISTER_UNICODE_BLOCK_RANGE(0x1cd0, 0x1cff, VedicExtensions,							NSLOCTEXT("UnicodeBlock", "VedicExtensions", "Vedic Extensions"))
-REGISTER_UNICODE_BLOCK_RANGE(0x1d00, 0x1d7f, PhoneticExtensions,						NSLOCTEXT("UnicodeBlock", "PhoneticExtensions", "Phonetic Extensions"))
-REGISTER_UNICODE_BLOCK_RANGE(0x1d80, 0x1dbf, PhoneticExtensionsSupplement,				NSLOCTEXT("UnicodeBlock", "PhoneticExtensionsSupplement", "Phonetic Extensions Supplement"))
-REGISTER_UNICODE_BLOCK_RANGE(0x1dc0, 0x1dff, CombiningDiacriticalMarksSupplement,		NSLOCTEXT("UnicodeBlock", "CombiningDiacriticalMarksSupplement", "Combining Diacritical Marks Supplement"))
-REGISTER_UNICODE_BLOCK_RANGE(0x1e00, 0x1eff, LatinExtendedAdditional,					NSLOCTEXT("UnicodeBlock", "LatinExtendedAdditional", "Latin Extended Additional"))
-REGISTER_UNICODE_BLOCK_RANGE(0x1f00, 0x1fff, GreekExtended,								NSLOCTEXT("UnicodeBlock", "GreekExtended", "Greek Extended"))
-REGISTER_UNICODE_BLOCK_RANGE(0x2000, 0x206f, GeneralPunctuation,						NSLOCTEXT("UnicodeBlock", "GeneralPunctuation", "General Punctuation"))
-REGISTER_UNICODE_BLOCK_RANGE(0x2070, 0x209f, SuperscriptsAndSubscripts,					NSLOCTEXT("UnicodeBlock", "SuperscriptsAndSubscripts", "Superscripts and Subscripts"))
-REGISTER_UNICODE_BLOCK_RANGE(0x20a0, 0x20cf, CurrencySymbols,							NSLOCTEXT("UnicodeBlock", "CurrencySymbols", "Currency Symbols"))
-REGISTER_UNICODE_BLOCK_RANGE(0x20d0, 0x20ff, CombiningDiacriticalMarksForSymbols,		NSLOCTEXT("UnicodeBlock", "CombiningDiacriticalMarksForSymbols", "Combining Diacritical Marks for Symbols"))
-REGISTER_UNICODE_BLOCK_RANGE(0x2100, 0x214f, LetterlikeSymbols,							NSLOCTEXT("UnicodeBlock", "LetterlikeSymbols", "Letterlike Symbols"))
-REGISTER_UNICODE_BLOCK_RANGE(0x2150, 0x218f, NumberForms,								NSLOCTEXT("UnicodeBlock", "NumberForms", "Number Forms"))
-REGISTER_UNICODE_BLOCK_RANGE(0x2190, 0x21ff, Arrows,									NSLOCTEXT("UnicodeBlock", "Arrows", "Arrows"))
-REGISTER_UNICODE_BLOCK_RANGE(0x2200, 0x22ff, MathematicalOperators,						NSLOCTEXT("UnicodeBlock", "MathematicalOperators", "Mathematical Operators"))
-REGISTER_UNICODE_BLOCK_RANGE(0x2300, 0x23ff, MiscellaneousTechnical,					NSLOCTEXT("UnicodeBlock", "MiscellaneousTechnical", "Miscellaneous Technical"))
-REGISTER_UNICODE_BLOCK_RANGE(0x2400, 0x243f, ControlPictures,							NSLOCTEXT("UnicodeBlock", "ControlPictures", "Control Pictures"))
-REGISTER_UNICODE_BLOCK_RANGE(0x2440, 0x245f, OpticalCharacterRecognition,				NSLOCTEXT("UnicodeBlock", "OpticalCharacterRecognition", "Optical Character Recognition"))
-REGISTER_UNICODE_BLOCK_RANGE(0x2460, 0x24ff, EnclosedAlphanumerics,						NSLOCTEXT("UnicodeBlock", "EnclosedAlphanumerics", "Enclosed Alphanumerics"))
-REGISTER_UNICODE_BLOCK_RANGE(0x2500, 0x257f, BoxDrawing,								NSLOCTEXT("UnicodeBlock", "BoxDrawing", "Box Drawing"))
-REGISTER_UNICODE_BLOCK_RANGE(0x2580, 0x259f, BlockElements,								NSLOCTEXT("UnicodeBlock", "BlockElements", "Block Elements"))
-REGISTER_UNICODE_BLOCK_RANGE(0x25a0, 0x25ff, GeometricShapes,							NSLOCTEXT("UnicodeBlock", "GeometricShapes", "Geometric Shapes"))
-REGISTER_UNICODE_BLOCK_RANGE(0x2600, 0x26ff, MiscellaneousSymbols,						NSLOCTEXT("UnicodeBlock", "MiscellaneousSymbols", "Miscellaneous Symbols"))
-REGISTER_UNICODE_BLOCK_RANGE(0x2700, 0x27bf, Dingbats,									NSLOCTEXT("UnicodeBlock", "Dingbats", "Dingbats"))
-REGISTER_UNICODE_BLOCK_RANGE(0x27c0, 0x27ef, MiscellaneousMathematicalSymbolsA,			NSLOCTEXT("UnicodeBlock", "MiscellaneousMathematicalSymbolsA", "Miscellaneous Mathematical Symbols - A"))
-REGISTER_UNICODE_BLOCK_RANGE(0x27f0, 0x27ff, SupplementalArrowsA,						NSLOCTEXT("UnicodeBlock", "SupplementalArrowsA", "Supplemental Arrows - A"))
-REGISTER_UNICODE_BLOCK_RANGE(0x2800, 0x28ff, BraillePatterns,							NSLOCTEXT("UnicodeBlock", "BraillePatterns", "Braille Patterns"))
-REGISTER_UNICODE_BLOCK_RANGE(0x2900, 0x297f, SupplementalArrowsB,						NSLOCTEXT("UnicodeBlock", "SupplementalArrowsB", "Supplemental Arrows - B"))
-REGISTER_UNICODE_BLOCK_RANGE(0x2980, 0x29ff, MiscellaneousMathematicalSymbolsB,			NSLOCTEXT("UnicodeBlock", "MiscellaneousMathematicalSymbolsB", "Miscellaneous Mathematical Symbols - B"))
-REGISTER_UNICODE_BLOCK_RANGE(0x2a00, 0x2aff, SupplementalMathematicalOperators,			NSLOCTEXT("UnicodeBlock", "SupplementalMathematicalOperators", "Supplemental Mathematical Operators"))
-REGISTER_UNICODE_BLOCK_RANGE(0x2b00, 0x2bff, MiscellaneousSymbolsAndArrows,				NSLOCTEXT("UnicodeBlock", "MiscellaneousSymbolsAndArrows", "Miscellaneous Symbols and Arrows"))
-REGISTER_UNICODE_BLOCK_RANGE(0x2c00, 0x2c5f, Glagolitic,								NSLOCTEXT("UnicodeBlock", "Glagolitic", "Glagolitic"))
-REGISTER_UNICODE_BLOCK_RANGE(0x2c60, 0x2c7f, LatinExtendedC,							NSLOCTEXT("UnicodeBlock", "LatinExtendedC", "Latin Extended - C"))
-REGISTER_UNICODE_BLOCK_RANGE(0x2c80, 0x2cff, Coptic,									NSLOCTEXT("UnicodeBlock", "Coptic", "Coptic"))
-REGISTER_UNICODE_BLOCK_RANGE(0x2d00, 0x2d2f, GeorgianSupplement,						NSLOCTEXT("UnicodeBlock", "GeorgianSupplement", "Georgian Supplement"))
-REGISTER_UNICODE_BLOCK_RANGE(0x2d30, 0x2d7f, Tifinagh,									NSLOCTEXT("UnicodeBlock", "Tifinagh", "Tifinagh"))
-REGISTER_UNICODE_BLOCK_RANGE(0x2d80, 0x2ddf, EthiopicExtended,							NSLOCTEXT("UnicodeBlock", "EthiopicExtended", "Ethiopic Extended"))
-REGISTER_UNICODE_BLOCK_RANGE(0x2de0, 0x2dff, CyrillicExtendedA,							NSLOCTEXT("UnicodeBlock", "CyrillicExtendedA", "Cyrillic Extended - A"))
-REGISTER_UNICODE_BLOCK_RANGE(0x2e00, 0x2e7f, SupplementalPunctuation,					NSLOCTEXT("UnicodeBlock", "SupplementalPunctuation", "Supplemental Punctuation"))
-REGISTER_UNICODE_BLOCK_RANGE(0x2e80, 0x2eff, CJKRadicalsSupplement,						NSLOCTEXT("UnicodeBlock", "CJKRadicalsSupplement", "CJK Radicals Supplement"))
-REGISTER_UNICODE_BLOCK_RANGE(0x2f00, 0x2fdf, KangxiRadicals,							NSLOCTEXT("UnicodeBlock", "KangxiRadicals", "Kangxi Radicals"))
-REGISTER_UNICODE_BLOCK_RANGE(0x2ff0, 0x2fff, IdeographicDescriptionCharacters,			NSLOCTEXT("UnicodeBlock", "IdeographicDescriptionCharacters", "Ideographic Description Characters"))
-REGISTER_UNICODE_BLOCK_RANGE(0x3000, 0x303f, CJKSymbolsAndPunctuation,					NSLOCTEXT("UnicodeBlock", "CJKSymbolsAndPunctuation", "CJK Symbols and Punctuation"))
-REGISTER_UNICODE_BLOCK_RANGE(0x3040, 0x309f, Hiragana,									NSLOCTEXT("UnicodeBlock", "Hiragana", "Hiragana"))
-REGISTER_UNICODE_BLOCK_RANGE(0x30a0, 0x30ff, Katakana,									NSLOCTEXT("UnicodeBlock", "Katakana", "Katakana"))
-REGISTER_UNICODE_BLOCK_RANGE(0x3100, 0x312f, Bopomofo,									NSLOCTEXT("UnicodeBlock", "Bopomofo", "Bopomofo"))
-REGISTER_UNICODE_BLOCK_RANGE(0x3130, 0x318f, HangulCompatibilityJamo,					NSLOCTEXT("UnicodeBlock", "HangulCompatibilityJamo", "Hangul Compatibility Jamo"))
-REGISTER_UNICODE_BLOCK_RANGE(0x3190, 0x319f, Kanbun,									NSLOCTEXT("UnicodeBlock", "Kanbun", "Kanbun"))
-REGISTER_UNICODE_BLOCK_RANGE(0x31a0, 0x31bf, BopomofoExtended,							NSLOCTEXT("UnicodeBlock", "BopomofoExtended", "Bopomofo Extended"))
-REGISTER_UNICODE_BLOCK_RANGE(0x31c0, 0x31ef, CJKStrokes,								NSLOCTEXT("UnicodeBlock", "CJKStrokes", "CJK Strokes"))
-REGISTER_UNICODE_BLOCK_RANGE(0x31f0, 0x31ff, KatakanaPhoneticExtensions,				NSLOCTEXT("UnicodeBlock", "KatakanaPhoneticExtensions", "Katakana Phonetic Extensions"))
-REGISTER_UNICODE_BLOCK_RANGE(0x3200, 0x32ff, EnclosedCJKLettersAndMonths,				NSLOCTEXT("UnicodeBlock", "EnclosedCJKLettersAndMonths", "Enclosed CJK Letters and Months"))
-REGISTER_UNICODE_BLOCK_RANGE(0x3300, 0x33ff, CJKCompatibility,							NSLOCTEXT("UnicodeBlock", "CJKCompatibility", "CJK Compatibility"))
-REGISTER_UNICODE_BLOCK_RANGE(0x3400, 0x4dbf, CJKUnifiedIdeographsExtensionA,			NSLOCTEXT("UnicodeBlock", "CJKUnifiedIdeographsExtensionA", "CJK Unified Ideographs Extension A"))
-REGISTER_UNICODE_BLOCK_RANGE(0x4dc0, 0x4dff, YijingHexagramSymbols,						NSLOCTEXT("UnicodeBlock", "YijingHexagramSymbols", "Yijing Hexagram Symbols"))
-REGISTER_UNICODE_BLOCK_RANGE(0x4e00, 0x9fff, CJKUnifiedIdeographs,						NSLOCTEXT("UnicodeBlock", "CJKUnifiedIdeographs", "CJK Unified Ideographs"))
-REGISTER_UNICODE_BLOCK_RANGE(0xa000, 0xa48f, YiSyllables,								NSLOCTEXT("UnicodeBlock", "YiSyllables", "Yi Syllables"))
-REGISTER_UNICODE_BLOCK_RANGE(0xa490, 0xa4cf, YiRadicals,								NSLOCTEXT("UnicodeBlock", "YiRadicals", "Yi Radicals"))
-REGISTER_UNICODE_BLOCK_RANGE(0xa4d0, 0xa4ff, Lisu,										NSLOCTEXT("UnicodeBlock", "Lisu", "Lisu"))
-REGISTER_UNICODE_BLOCK_RANGE(0xa500, 0xa63f, Vai,										NSLOCTEXT("UnicodeBlock", "Vai", "Vai"))
-REGISTER_UNICODE_BLOCK_RANGE(0xa640, 0xa69f, CyrillicExtendedB,							NSLOCTEXT("UnicodeBlock", "CyrillicExtendedB", "Cyrillic Extended - B"))
-REGISTER_UNICODE_BLOCK_RANGE(0xa6a0, 0xa6ff, Bamum,										NSLOCTEXT("UnicodeBlock", "Bamum", "Bamum"))
-REGISTER_UNICODE_BLOCK_RANGE(0xa700, 0xa71f, ModifierToneLetters,						NSLOCTEXT("UnicodeBlock", "ModifierToneLetters", "Modifier Tone Letters"))
-REGISTER_UNICODE_BLOCK_RANGE(0xa720, 0xa7ff, LatinExtendedD,							NSLOCTEXT("UnicodeBlock", "LatinExtendedD", "Latin Extended - D"))
-REGISTER_UNICODE_BLOCK_RANGE(0xa800, 0xa82f, SylotiNagri,								NSLOCTEXT("UnicodeBlock", "SylotiNagri", "Syloti Nagri"))
-REGISTER_UNICODE_BLOCK_RANGE(0xa830, 0xa83f, CommonIndicNumberForms,					NSLOCTEXT("UnicodeBlock", "CommonIndicNumberForms", "Common Indic Number Forms"))
-REGISTER_UNICODE_BLOCK_RANGE(0xa840, 0xa87f, PhagsPa,									NSLOCTEXT("UnicodeBlock", "PhagsPa", "Phags - pa"))
-REGISTER_UNICODE_BLOCK_RANGE(0xa880, 0xa8df, Saurashtra,								NSLOCTEXT("UnicodeBlock", "Saurashtra", "Saurashtra"))
-REGISTER_UNICODE_BLOCK_RANGE(0xa8e0, 0xa8ff, DevanagariExtended,						NSLOCTEXT("UnicodeBlock", "DevanagariExtended", "Devanagari Extended"))
-REGISTER_UNICODE_BLOCK_RANGE(0xa900, 0xa92f, KayahLi,									NSLOCTEXT("UnicodeBlock", "KayahLi", "Kayah Li"))
-REGISTER_UNICODE_BLOCK_RANGE(0xa930, 0xa95f, Rejang,									NSLOCTEXT("UnicodeBlock", "Rejang", "Rejang"))
-REGISTER_UNICODE_BLOCK_RANGE(0xa960, 0xa97f, HangulJamoExtendedA,						NSLOCTEXT("UnicodeBlock", "HangulJamoExtendedA", "Hangul Jamo Extended - A"))
-REGISTER_UNICODE_BLOCK_RANGE(0xa980, 0xa9df, Javanese,									NSLOCTEXT("UnicodeBlock", "Javanese", "Javanese"))
-REGISTER_UNICODE_BLOCK_RANGE(0xa9e0, 0xa9ff, MyanmarExtendedB,							NSLOCTEXT("UnicodeBlock", "MyanmarExtendedB", "Myanmar Extended - B"))
-REGISTER_UNICODE_BLOCK_RANGE(0xaa00, 0xaa5f, Cham,										NSLOCTEXT("UnicodeBlock", "Cham", "Cham"))
-REGISTER_UNICODE_BLOCK_RANGE(0xaa60, 0xaa7f, MyanmarExtendedA,							NSLOCTEXT("UnicodeBlock", "MyanmarExtendedA", "Myanmar Extended - A"))
-REGISTER_UNICODE_BLOCK_RANGE(0xaa80, 0xaadf, TaiViet,									NSLOCTEXT("UnicodeBlock", "TaiViet", "Tai Viet"))
-REGISTER_UNICODE_BLOCK_RANGE(0xaae0, 0xaaff, MeeteiMayekExtensions,						NSLOCTEXT("UnicodeBlock", "MeeteiMayekExtensions", "Meetei Mayek Extensions"))
-REGISTER_UNICODE_BLOCK_RANGE(0xab00, 0xab2f, EthiopicExtendedA,							NSLOCTEXT("UnicodeBlock", "EthiopicExtendedA", "Ethiopic Extended - A"))
-REGISTER_UNICODE_BLOCK_RANGE(0xab30, 0xab6f, LatinExtendedE,							NSLOCTEXT("UnicodeBlock", "LatinExtendedE", "Latin Extended - E"))
-REGISTER_UNICODE_BLOCK_RANGE(0xab70, 0xabbf, CherokeeSupplement,						NSLOCTEXT("UnicodeBlock", "CherokeeSupplement", "Cherokee Supplement"))
-REGISTER_UNICODE_BLOCK_RANGE(0xabc0, 0xabff, MeeteiMayek,								NSLOCTEXT("UnicodeBlock", "MeeteiMayek", "Meetei Mayek"))
-REGISTER_UNICODE_BLOCK_RANGE(0xac00, 0xd7af, HangulSyllables,							NSLOCTEXT("UnicodeBlock", "HangulSyllables", "Hangul Syllables"))
-REGISTER_UNICODE_BLOCK_RANGE(0xd7b0, 0xd7ff, HangulJamoExtendedB,						NSLOCTEXT("UnicodeBlock", "HangulJamoExtendedB", "Hangul Jamo Extended - B"))
-REGISTER_UNICODE_BLOCK_RANGE(0xd800, 0xdb7f, HighSurrogates,							NSLOCTEXT("UnicodeBlock", "HighSurrogates", "High Surrogates"))
-REGISTER_UNICODE_BLOCK_RANGE(0xdb80, 0xdbff, HighPrivateUseSurrogates,					NSLOCTEXT("UnicodeBlock", "HighPrivateUseSurrogates", "High Private Use Surrogates"))
-REGISTER_UNICODE_BLOCK_RANGE(0xdc00, 0xdfff, LowSurrogates,								NSLOCTEXT("UnicodeBlock", "LowSurrogates", "Low Surrogates"))
-REGISTER_UNICODE_BLOCK_RANGE(0xe000, 0xf8ff, PrivateUseArea,							NSLOCTEXT("UnicodeBlock", "PrivateUseArea", "Private Use Area"))
-REGISTER_UNICODE_BLOCK_RANGE(0xf900, 0xfaff, CJKCompatibilityIdeographs,				NSLOCTEXT("UnicodeBlock", "CJKCompatibilityIdeographs", "CJK Compatibility Ideographs"))
-REGISTER_UNICODE_BLOCK_RANGE(0xfb00, 0xfb4f, AlphabeticPresentationForms,				NSLOCTEXT("UnicodeBlock", "AlphabeticPresentationForms", "Alphabetic Presentation Forms"))
-REGISTER_UNICODE_BLOCK_RANGE(0xfb50, 0xfdff, ArabicPresentationFormsA,					NSLOCTEXT("UnicodeBlock", "ArabicPresentationFormsA", "Arabic Presentation Forms - A"))
-REGISTER_UNICODE_BLOCK_RANGE(0xfe00, 0xfe0f, VariationSelectors,						NSLOCTEXT("UnicodeBlock", "VariationSelectors", "Variation Selectors"))
-REGISTER_UNICODE_BLOCK_RANGE(0xfe10, 0xfe1f, VerticalForms,								NSLOCTEXT("UnicodeBlock", "VerticalForms", "Vertical Forms"))
-REGISTER_UNICODE_BLOCK_RANGE(0xfe20, 0xfe2f, CombiningHalfMarks,						NSLOCTEXT("UnicodeBlock", "CombiningHalfMarks", "Combining Half Marks"))
-REGISTER_UNICODE_BLOCK_RANGE(0xfe30, 0xfe4f, CJKCompatibilityForms,						NSLOCTEXT("UnicodeBlock", "CJKCompatibilityForms", "CJK Compatibility Forms"))
-REGISTER_UNICODE_BLOCK_RANGE(0xfe50, 0xfe6f, SmallFormVariants,							NSLOCTEXT("UnicodeBlock", "SmallFormVariants", "Small Form Variants"))
-REGISTER_UNICODE_BLOCK_RANGE(0xfe70, 0xfeff, ArabicPresentationFormsB,					NSLOCTEXT("UnicodeBlock", "ArabicPresentationFormsB", "Arabic Presentation Forms - B"))
-REGISTER_UNICODE_BLOCK_RANGE(0xff00, 0xffef, HalfwidthAndFullwidthForms,				NSLOCTEXT("UnicodeBlock", "HalfwidthAndFullwidthForms", "Halfwidth and Fullwidth Forms"))
-REGISTER_UNICODE_BLOCK_RANGE(0xfff0, 0xffff, Specials,									NSLOCTEXT("UnicodeBlock", "Specials", "Specials"))
+REGISTER_UNICODE_BLOCK_RANGE(0x0000, 0x001f, ControlCharacter,							"UnicodeBlock", "ControlCharacter", "Control Character")
+REGISTER_UNICODE_BLOCK_RANGE(0x0020, 0x007f, BasicLatin,								"UnicodeBlock", "BasicLatin", "Basic Latin")
+REGISTER_UNICODE_BLOCK_RANGE(0x0080, 0x00ff, Latin1Supplement,							"UnicodeBlock", "Latin1Supplement", "Latin - 1 Supplement")
+REGISTER_UNICODE_BLOCK_RANGE(0x0100, 0x017f, LatinExtendedA,							"UnicodeBlock", "LatinExtendedA", "Latin Extended - A")
+REGISTER_UNICODE_BLOCK_RANGE(0x0180, 0x024f, LatinExtendedB,							"UnicodeBlock", "LatinExtendedB", "Latin Extended - B")
+REGISTER_UNICODE_BLOCK_RANGE(0x0250, 0x02af, IPAExtensions,								"UnicodeBlock", "IPAExtensions", "IPA Extensions")
+REGISTER_UNICODE_BLOCK_RANGE(0x02b0, 0x02ff, SpacingModifierLetters,					"UnicodeBlock", "SpacingModifierLetters", "Spacing Modifier Letters")
+REGISTER_UNICODE_BLOCK_RANGE(0x0300, 0x036f, CombiningDiacriticalMarks,					"UnicodeBlock", "CombiningDiacriticalMarks", "Combining Diacritical Marks")
+REGISTER_UNICODE_BLOCK_RANGE(0x0370, 0x03ff, GreekAndCoptic,							"UnicodeBlock", "GreekAndCoptic", "Greek and Coptic")
+REGISTER_UNICODE_BLOCK_RANGE(0x0400, 0x04ff, Cyrillic,									"UnicodeBlock", "Cyrillic", "Cyrillic")
+REGISTER_UNICODE_BLOCK_RANGE(0x0500, 0x052f, CyrillicSupplementary,						"UnicodeBlock", "CyrillicSupplementary", "Cyrillic Supplementary")
+REGISTER_UNICODE_BLOCK_RANGE(0x0530, 0x058f, Armenian,									"UnicodeBlock", "Armenian", "Armenian")
+REGISTER_UNICODE_BLOCK_RANGE(0x0590, 0x05ff, Hebrew,									"UnicodeBlock", "Hebrew", "Hebrew")
+REGISTER_UNICODE_BLOCK_RANGE(0x0600, 0x06ff, Arabic,									"UnicodeBlock", "Arabic", "Arabic")
+REGISTER_UNICODE_BLOCK_RANGE(0x0700, 0x074f, Syriac,									"UnicodeBlock", "Syriac", "Syriac")
+REGISTER_UNICODE_BLOCK_RANGE(0x0750, 0x077f, ArabicSupplement,							"UnicodeBlock", "ArabicSupplement", "Arabic Supplement")
+REGISTER_UNICODE_BLOCK_RANGE(0x0780, 0x07bf, Thaana,									"UnicodeBlock", "Thaana", "Thaana")
+REGISTER_UNICODE_BLOCK_RANGE(0x07c0, 0x07ff, NKo,										"UnicodeBlock", "NKo", "NKo")
+REGISTER_UNICODE_BLOCK_RANGE(0x0800, 0x083f, Samaritan,									"UnicodeBlock", "Samaritan", "Samaritan")
+REGISTER_UNICODE_BLOCK_RANGE(0x0840, 0x085f, Mandaic,									"UnicodeBlock", "Mandaic", "Mandaic")
+REGISTER_UNICODE_BLOCK_RANGE(0x08a0, 0x08ff, ArabicExtendedA,							"UnicodeBlock", "ArabicExtendedA", "Arabic Extended - A")
+REGISTER_UNICODE_BLOCK_RANGE(0x0900, 0x097f, Devanagari,								"UnicodeBlock", "Devanagari", "Devanagari")
+REGISTER_UNICODE_BLOCK_RANGE(0x0980, 0x09ff, Bengali,									"UnicodeBlock", "Bengali", "Bengali")
+REGISTER_UNICODE_BLOCK_RANGE(0x0a00, 0x0a7f, Gurmukhi,									"UnicodeBlock", "Gurmukhi", "Gurmukhi")
+REGISTER_UNICODE_BLOCK_RANGE(0x0a80, 0x0aff, Gujarati,									"UnicodeBlock", "Gujarati", "Gujarati")
+REGISTER_UNICODE_BLOCK_RANGE(0x0b00, 0x0b7f, Oriya,										"UnicodeBlock", "Oriya", "Oriya")
+REGISTER_UNICODE_BLOCK_RANGE(0x0b80, 0x0bff, Tamil,										"UnicodeBlock", "Tamil", "Tamil")
+REGISTER_UNICODE_BLOCK_RANGE(0x0c00, 0x0c7f, Telugu,									"UnicodeBlock", "Telugu", "Telugu")
+REGISTER_UNICODE_BLOCK_RANGE(0x0c80, 0x0cff, Kannada,									"UnicodeBlock", "Kannada", "Kannada")
+REGISTER_UNICODE_BLOCK_RANGE(0x0d00, 0x0d7f, Malayalam,									"UnicodeBlock", "Malayalam", "Malayalam")
+REGISTER_UNICODE_BLOCK_RANGE(0x0d80, 0x0dff, Sinhala,									"UnicodeBlock", "Sinhala", "Sinhala")
+REGISTER_UNICODE_BLOCK_RANGE(0x0e00, 0x0e7f, Thai,										"UnicodeBlock", "Thai", "Thai")
+REGISTER_UNICODE_BLOCK_RANGE(0x0e80, 0x0eff, Lao,										"UnicodeBlock", "Lao", "Lao")
+REGISTER_UNICODE_BLOCK_RANGE(0x0f00, 0x0fff, Tibetan,									"UnicodeBlock", "Tibetan", "Tibetan")
+REGISTER_UNICODE_BLOCK_RANGE(0x1000, 0x109f, Myanmar,									"UnicodeBlock", "Myanmar", "Myanmar")
+REGISTER_UNICODE_BLOCK_RANGE(0x10a0, 0x10ff, Georgian,									"UnicodeBlock", "Georgian", "Georgian")
+REGISTER_UNICODE_BLOCK_RANGE(0x1100, 0x11ff, HangulJamo,								"UnicodeBlock", "HangulJamo", "Hangul Jamo")
+REGISTER_UNICODE_BLOCK_RANGE(0x1200, 0x137f, Ethiopic,									"UnicodeBlock", "Ethiopic", "Ethiopic")
+REGISTER_UNICODE_BLOCK_RANGE(0x1380, 0x139f, EthiopicSupplement,						"UnicodeBlock", "EthiopicSupplement", "Ethiopic Supplement")
+REGISTER_UNICODE_BLOCK_RANGE(0x13a0, 0x13ff, Cherokee,									"UnicodeBlock", "Cherokee", "Cherokee")
+REGISTER_UNICODE_BLOCK_RANGE(0x1400, 0x167f, UnifiedCanadianAboriginalSyllabics,		"UnicodeBlock", "UnifiedCanadianAboriginalSyllabics", "Unified Canadian Aboriginal Syllabics")
+REGISTER_UNICODE_BLOCK_RANGE(0x1680, 0x169f, Ogham,										"UnicodeBlock", "Ogham", "Ogham")
+REGISTER_UNICODE_BLOCK_RANGE(0x16a0, 0x16ff, Runic,										"UnicodeBlock", "Runic", "Runic")
+REGISTER_UNICODE_BLOCK_RANGE(0x1700, 0x171f, Tagalog,									"UnicodeBlock", "Tagalog", "Tagalog")
+REGISTER_UNICODE_BLOCK_RANGE(0x1720, 0x173f, Hanunoo,									"UnicodeBlock", "Hanunoo", "Hanunoo")
+REGISTER_UNICODE_BLOCK_RANGE(0x1740, 0x175f, Buhid,										"UnicodeBlock", "Buhid", "Buhid")
+REGISTER_UNICODE_BLOCK_RANGE(0x1760, 0x177f, Tagbanwa,									"UnicodeBlock", "Tagbanwa", "Tagbanwa")
+REGISTER_UNICODE_BLOCK_RANGE(0x1780, 0x17ff, Khmer,										"UnicodeBlock", "Khmer", "Khmer")
+REGISTER_UNICODE_BLOCK_RANGE(0x1800, 0x18af, Mongolian,									"UnicodeBlock", "Mongolian", "Mongolian")
+REGISTER_UNICODE_BLOCK_RANGE(0x18b0, 0x18ff, UnifiedCanadianAboriginalSyllabicsExtended,"UnicodeBlock", "UnifiedCanadianAboriginalSyllabicsExtended", "Unified Canadian Aboriginal Syllabics Extended")
+REGISTER_UNICODE_BLOCK_RANGE(0x1900, 0x194f, Limbu,										"UnicodeBlock", "Limbu", "Limbu")
+REGISTER_UNICODE_BLOCK_RANGE(0x1950, 0x197f, TaiLe,										"UnicodeBlock", "TaiLe", "Tai Le")
+REGISTER_UNICODE_BLOCK_RANGE(0x1980, 0x19Df, NewTaiLue,									"UnicodeBlock", "NewTaiLue", "New Tai Lue")
+REGISTER_UNICODE_BLOCK_RANGE(0x19e0, 0x19ff, KhmerSymbols,								"UnicodeBlock", "KhmerSymbols", "Khmer Symbols")
+REGISTER_UNICODE_BLOCK_RANGE(0x1a00, 0x1a1f, Buginese,									"UnicodeBlock", "Buginese", "Buginese")
+REGISTER_UNICODE_BLOCK_RANGE(0x1a20, 0x1aaf, TaiTham,									"UnicodeBlock", "TaiTham", "Tai Tham")
+REGISTER_UNICODE_BLOCK_RANGE(0x1ab0, 0x1aff, CombiningDiacriticalMarksExtended,			"UnicodeBlock", "CombiningDiacriticalMarksExtended", "Combining Diacritical Marks Extended")
+REGISTER_UNICODE_BLOCK_RANGE(0x1b00, 0x1b7f, Balinese,									"UnicodeBlock", "Balinese", "Balinese")
+REGISTER_UNICODE_BLOCK_RANGE(0x1b80, 0x1bbf, Sundanese,									"UnicodeBlock", "Sundanese", "Sundanese")
+REGISTER_UNICODE_BLOCK_RANGE(0x1bc0, 0x1bff, Batak,										"UnicodeBlock", "Batak", "Batak")
+REGISTER_UNICODE_BLOCK_RANGE(0x1c00, 0x1c4f, Lepcha,									"UnicodeBlock", "Lepcha", "Lepcha")
+REGISTER_UNICODE_BLOCK_RANGE(0x1c50, 0x1c7f, OlChiki,									"UnicodeBlock", "OlChiki", "Ol Chiki")
+REGISTER_UNICODE_BLOCK_RANGE(0x1c80, 0x1c87, CyrillicExtendedC,							"UnicodeBlock", "CyrillicExtendedC", "Cyrillic Extended C")
+REGISTER_UNICODE_BLOCK_RANGE(0x1cc0, 0x1ccf, SundaneseSupplement,						"UnicodeBlock", "SundaneseSupplement", "Sundanese Supplement")
+REGISTER_UNICODE_BLOCK_RANGE(0x1cd0, 0x1cff, VedicExtensions,							"UnicodeBlock", "VedicExtensions", "Vedic Extensions")
+REGISTER_UNICODE_BLOCK_RANGE(0x1d00, 0x1d7f, PhoneticExtensions,						"UnicodeBlock", "PhoneticExtensions", "Phonetic Extensions")
+REGISTER_UNICODE_BLOCK_RANGE(0x1d80, 0x1dbf, PhoneticExtensionsSupplement,				"UnicodeBlock", "PhoneticExtensionsSupplement", "Phonetic Extensions Supplement")
+REGISTER_UNICODE_BLOCK_RANGE(0x1dc0, 0x1dff, CombiningDiacriticalMarksSupplement,		"UnicodeBlock", "CombiningDiacriticalMarksSupplement", "Combining Diacritical Marks Supplement")
+REGISTER_UNICODE_BLOCK_RANGE(0x1e00, 0x1eff, LatinExtendedAdditional,					"UnicodeBlock", "LatinExtendedAdditional", "Latin Extended Additional")
+REGISTER_UNICODE_BLOCK_RANGE(0x1f00, 0x1fff, GreekExtended,								"UnicodeBlock", "GreekExtended", "Greek Extended")
+REGISTER_UNICODE_BLOCK_RANGE(0x2000, 0x206f, GeneralPunctuation,						"UnicodeBlock", "GeneralPunctuation", "General Punctuation")
+REGISTER_UNICODE_BLOCK_RANGE(0x2070, 0x209f, SuperscriptsAndSubscripts,					"UnicodeBlock", "SuperscriptsAndSubscripts", "Superscripts and Subscripts")
+REGISTER_UNICODE_BLOCK_RANGE(0x20a0, 0x20cf, CurrencySymbols,							"UnicodeBlock", "CurrencySymbols", "Currency Symbols")
+REGISTER_UNICODE_BLOCK_RANGE(0x20d0, 0x20ff, CombiningDiacriticalMarksForSymbols,		"UnicodeBlock", "CombiningDiacriticalMarksForSymbols", "Combining Diacritical Marks for Symbols")
+REGISTER_UNICODE_BLOCK_RANGE(0x2100, 0x214f, LetterlikeSymbols,							"UnicodeBlock", "LetterlikeSymbols", "Letterlike Symbols")
+REGISTER_UNICODE_BLOCK_RANGE(0x2150, 0x218f, NumberForms,								"UnicodeBlock", "NumberForms", "Number Forms")
+REGISTER_UNICODE_BLOCK_RANGE(0x2190, 0x21ff, Arrows,									"UnicodeBlock", "Arrows", "Arrows")
+REGISTER_UNICODE_BLOCK_RANGE(0x2200, 0x22ff, MathematicalOperators,						"UnicodeBlock", "MathematicalOperators", "Mathematical Operators")
+REGISTER_UNICODE_BLOCK_RANGE(0x2300, 0x23ff, MiscellaneousTechnical,					"UnicodeBlock", "MiscellaneousTechnical", "Miscellaneous Technical")
+REGISTER_UNICODE_BLOCK_RANGE(0x2400, 0x243f, ControlPictures,							"UnicodeBlock", "ControlPictures", "Control Pictures")
+REGISTER_UNICODE_BLOCK_RANGE(0x2440, 0x245f, OpticalCharacterRecognition,				"UnicodeBlock", "OpticalCharacterRecognition", "Optical Character Recognition")
+REGISTER_UNICODE_BLOCK_RANGE(0x2460, 0x24ff, EnclosedAlphanumerics,						"UnicodeBlock", "EnclosedAlphanumerics", "Enclosed Alphanumerics")
+REGISTER_UNICODE_BLOCK_RANGE(0x2500, 0x257f, BoxDrawing,								"UnicodeBlock", "BoxDrawing", "Box Drawing")
+REGISTER_UNICODE_BLOCK_RANGE(0x2580, 0x259f, BlockElements,								"UnicodeBlock", "BlockElements", "Block Elements")
+REGISTER_UNICODE_BLOCK_RANGE(0x25a0, 0x25ff, GeometricShapes,							"UnicodeBlock", "GeometricShapes", "Geometric Shapes")
+REGISTER_UNICODE_BLOCK_RANGE(0x2600, 0x26ff, MiscellaneousSymbols,						"UnicodeBlock", "MiscellaneousSymbols", "Miscellaneous Symbols")
+REGISTER_UNICODE_BLOCK_RANGE(0x2700, 0x27bf, Dingbats,									"UnicodeBlock", "Dingbats", "Dingbats")
+REGISTER_UNICODE_BLOCK_RANGE(0x27c0, 0x27ef, MiscellaneousMathematicalSymbolsA,			"UnicodeBlock", "MiscellaneousMathematicalSymbolsA", "Miscellaneous Mathematical Symbols - A")
+REGISTER_UNICODE_BLOCK_RANGE(0x27f0, 0x27ff, SupplementalArrowsA,						"UnicodeBlock", "SupplementalArrowsA", "Supplemental Arrows - A")
+REGISTER_UNICODE_BLOCK_RANGE(0x2800, 0x28ff, BraillePatterns,							"UnicodeBlock", "BraillePatterns", "Braille Patterns")
+REGISTER_UNICODE_BLOCK_RANGE(0x2900, 0x297f, SupplementalArrowsB,						"UnicodeBlock", "SupplementalArrowsB", "Supplemental Arrows - B")
+REGISTER_UNICODE_BLOCK_RANGE(0x2980, 0x29ff, MiscellaneousMathematicalSymbolsB,			"UnicodeBlock", "MiscellaneousMathematicalSymbolsB", "Miscellaneous Mathematical Symbols - B")
+REGISTER_UNICODE_BLOCK_RANGE(0x2a00, 0x2aff, SupplementalMathematicalOperators,			"UnicodeBlock", "SupplementalMathematicalOperators", "Supplemental Mathematical Operators")
+REGISTER_UNICODE_BLOCK_RANGE(0x2b00, 0x2bff, MiscellaneousSymbolsAndArrows,				"UnicodeBlock", "MiscellaneousSymbolsAndArrows", "Miscellaneous Symbols and Arrows")
+REGISTER_UNICODE_BLOCK_RANGE(0x2c00, 0x2c5f, Glagolitic,								"UnicodeBlock", "Glagolitic", "Glagolitic")
+REGISTER_UNICODE_BLOCK_RANGE(0x2c60, 0x2c7f, LatinExtendedC,							"UnicodeBlock", "LatinExtendedC", "Latin Extended - C")
+REGISTER_UNICODE_BLOCK_RANGE(0x2c80, 0x2cff, Coptic,									"UnicodeBlock", "Coptic", "Coptic")
+REGISTER_UNICODE_BLOCK_RANGE(0x2d00, 0x2d2f, GeorgianSupplement,						"UnicodeBlock", "GeorgianSupplement", "Georgian Supplement")
+REGISTER_UNICODE_BLOCK_RANGE(0x2d30, 0x2d7f, Tifinagh,									"UnicodeBlock", "Tifinagh", "Tifinagh")
+REGISTER_UNICODE_BLOCK_RANGE(0x2d80, 0x2ddf, EthiopicExtended,							"UnicodeBlock", "EthiopicExtended", "Ethiopic Extended")
+REGISTER_UNICODE_BLOCK_RANGE(0x2de0, 0x2dff, CyrillicExtendedA,							"UnicodeBlock", "CyrillicExtendedA", "Cyrillic Extended - A")
+REGISTER_UNICODE_BLOCK_RANGE(0x2e00, 0x2e7f, SupplementalPunctuation,					"UnicodeBlock", "SupplementalPunctuation", "Supplemental Punctuation")
+REGISTER_UNICODE_BLOCK_RANGE(0x2e80, 0x2eff, CJKRadicalsSupplement,						"UnicodeBlock", "CJKRadicalsSupplement", "CJK Radicals Supplement")
+REGISTER_UNICODE_BLOCK_RANGE(0x2f00, 0x2fdf, KangxiRadicals,							"UnicodeBlock", "KangxiRadicals", "Kangxi Radicals")
+REGISTER_UNICODE_BLOCK_RANGE(0x2ff0, 0x2fff, IdeographicDescriptionCharacters,			"UnicodeBlock", "IdeographicDescriptionCharacters", "Ideographic Description Characters")
+REGISTER_UNICODE_BLOCK_RANGE(0x3000, 0x303f, CJKSymbolsAndPunctuation,					"UnicodeBlock", "CJKSymbolsAndPunctuation", "CJK Symbols and Punctuation")
+REGISTER_UNICODE_BLOCK_RANGE(0x3040, 0x309f, Hiragana,									"UnicodeBlock", "Hiragana", "Hiragana")
+REGISTER_UNICODE_BLOCK_RANGE(0x30a0, 0x30ff, Katakana,									"UnicodeBlock", "Katakana", "Katakana")
+REGISTER_UNICODE_BLOCK_RANGE(0x3100, 0x312f, Bopomofo,									"UnicodeBlock", "Bopomofo", "Bopomofo")
+REGISTER_UNICODE_BLOCK_RANGE(0x3130, 0x318f, HangulCompatibilityJamo,					"UnicodeBlock", "HangulCompatibilityJamo", "Hangul Compatibility Jamo")
+REGISTER_UNICODE_BLOCK_RANGE(0x3190, 0x319f, Kanbun,									"UnicodeBlock", "Kanbun", "Kanbun")
+REGISTER_UNICODE_BLOCK_RANGE(0x31a0, 0x31bf, BopomofoExtended,							"UnicodeBlock", "BopomofoExtended", "Bopomofo Extended")
+REGISTER_UNICODE_BLOCK_RANGE(0x31c0, 0x31ef, CJKStrokes,								"UnicodeBlock", "CJKStrokes", "CJK Strokes")
+REGISTER_UNICODE_BLOCK_RANGE(0x31f0, 0x31ff, KatakanaPhoneticExtensions,				"UnicodeBlock", "KatakanaPhoneticExtensions", "Katakana Phonetic Extensions")
+REGISTER_UNICODE_BLOCK_RANGE(0x3200, 0x32ff, EnclosedCJKLettersAndMonths,				"UnicodeBlock", "EnclosedCJKLettersAndMonths", "Enclosed CJK Letters and Months")
+REGISTER_UNICODE_BLOCK_RANGE(0x3300, 0x33ff, CJKCompatibility,							"UnicodeBlock", "CJKCompatibility", "CJK Compatibility")
+REGISTER_UNICODE_BLOCK_RANGE(0x3400, 0x4dbf, CJKUnifiedIdeographsExtensionA,			"UnicodeBlock", "CJKUnifiedIdeographsExtensionA", "CJK Unified Ideographs Extension A")
+REGISTER_UNICODE_BLOCK_RANGE(0x4dc0, 0x4dff, YijingHexagramSymbols,						"UnicodeBlock", "YijingHexagramSymbols", "Yijing Hexagram Symbols")
+REGISTER_UNICODE_BLOCK_RANGE(0x4e00, 0x9fff, CJKUnifiedIdeographs,						"UnicodeBlock", "CJKUnifiedIdeographs", "CJK Unified Ideographs")
+REGISTER_UNICODE_BLOCK_RANGE(0xa000, 0xa48f, YiSyllables,								"UnicodeBlock", "YiSyllables", "Yi Syllables")
+REGISTER_UNICODE_BLOCK_RANGE(0xa490, 0xa4cf, YiRadicals,								"UnicodeBlock", "YiRadicals", "Yi Radicals")
+REGISTER_UNICODE_BLOCK_RANGE(0xa4d0, 0xa4ff, Lisu,										"UnicodeBlock", "Lisu", "Lisu")
+REGISTER_UNICODE_BLOCK_RANGE(0xa500, 0xa63f, Vai,										"UnicodeBlock", "Vai", "Vai")
+REGISTER_UNICODE_BLOCK_RANGE(0xa640, 0xa69f, CyrillicExtendedB,							"UnicodeBlock", "CyrillicExtendedB", "Cyrillic Extended - B")
+REGISTER_UNICODE_BLOCK_RANGE(0xa6a0, 0xa6ff, Bamum,										"UnicodeBlock", "Bamum", "Bamum")
+REGISTER_UNICODE_BLOCK_RANGE(0xa700, 0xa71f, ModifierToneLetters,						"UnicodeBlock", "ModifierToneLetters", "Modifier Tone Letters")
+REGISTER_UNICODE_BLOCK_RANGE(0xa720, 0xa7ff, LatinExtendedD,							"UnicodeBlock", "LatinExtendedD", "Latin Extended - D")
+REGISTER_UNICODE_BLOCK_RANGE(0xa800, 0xa82f, SylotiNagri,								"UnicodeBlock", "SylotiNagri", "Syloti Nagri")
+REGISTER_UNICODE_BLOCK_RANGE(0xa830, 0xa83f, CommonIndicNumberForms,					"UnicodeBlock", "CommonIndicNumberForms", "Common Indic Number Forms")
+REGISTER_UNICODE_BLOCK_RANGE(0xa840, 0xa87f, PhagsPa,									"UnicodeBlock", "PhagsPa", "Phags - pa")
+REGISTER_UNICODE_BLOCK_RANGE(0xa880, 0xa8df, Saurashtra,								"UnicodeBlock", "Saurashtra", "Saurashtra")
+REGISTER_UNICODE_BLOCK_RANGE(0xa8e0, 0xa8ff, DevanagariExtended,						"UnicodeBlock", "DevanagariExtended", "Devanagari Extended")
+REGISTER_UNICODE_BLOCK_RANGE(0xa900, 0xa92f, KayahLi,									"UnicodeBlock", "KayahLi", "Kayah Li")
+REGISTER_UNICODE_BLOCK_RANGE(0xa930, 0xa95f, Rejang,									"UnicodeBlock", "Rejang", "Rejang")
+REGISTER_UNICODE_BLOCK_RANGE(0xa960, 0xa97f, HangulJamoExtendedA,						"UnicodeBlock", "HangulJamoExtendedA", "Hangul Jamo Extended - A")
+REGISTER_UNICODE_BLOCK_RANGE(0xa980, 0xa9df, Javanese,									"UnicodeBlock", "Javanese", "Javanese")
+REGISTER_UNICODE_BLOCK_RANGE(0xa9e0, 0xa9ff, MyanmarExtendedB,							"UnicodeBlock", "MyanmarExtendedB", "Myanmar Extended - B")
+REGISTER_UNICODE_BLOCK_RANGE(0xaa00, 0xaa5f, Cham,										"UnicodeBlock", "Cham", "Cham")
+REGISTER_UNICODE_BLOCK_RANGE(0xaa60, 0xaa7f, MyanmarExtendedA,							"UnicodeBlock", "MyanmarExtendedA", "Myanmar Extended - A")
+REGISTER_UNICODE_BLOCK_RANGE(0xaa80, 0xaadf, TaiViet,									"UnicodeBlock", "TaiViet", "Tai Viet")
+REGISTER_UNICODE_BLOCK_RANGE(0xaae0, 0xaaff, MeeteiMayekExtensions,						"UnicodeBlock", "MeeteiMayekExtensions", "Meetei Mayek Extensions")
+REGISTER_UNICODE_BLOCK_RANGE(0xab00, 0xab2f, EthiopicExtendedA,							"UnicodeBlock", "EthiopicExtendedA", "Ethiopic Extended - A")
+REGISTER_UNICODE_BLOCK_RANGE(0xab30, 0xab6f, LatinExtendedE,							"UnicodeBlock", "LatinExtendedE", "Latin Extended - E")
+REGISTER_UNICODE_BLOCK_RANGE(0xab70, 0xabbf, CherokeeSupplement,						"UnicodeBlock", "CherokeeSupplement", "Cherokee Supplement")
+REGISTER_UNICODE_BLOCK_RANGE(0xabc0, 0xabff, MeeteiMayek,								"UnicodeBlock", "MeeteiMayek", "Meetei Mayek")
+REGISTER_UNICODE_BLOCK_RANGE(0xac00, 0xd7af, HangulSyllables,							"UnicodeBlock", "HangulSyllables", "Hangul Syllables")
+REGISTER_UNICODE_BLOCK_RANGE(0xd7b0, 0xd7ff, HangulJamoExtendedB,						"UnicodeBlock", "HangulJamoExtendedB", "Hangul Jamo Extended - B")
+REGISTER_UNICODE_BLOCK_RANGE(0xd800, 0xdb7f, HighSurrogates,							"UnicodeBlock", "HighSurrogates", "High Surrogates")
+REGISTER_UNICODE_BLOCK_RANGE(0xdb80, 0xdbff, HighPrivateUseSurrogates,					"UnicodeBlock", "HighPrivateUseSurrogates", "High Private Use Surrogates")
+REGISTER_UNICODE_BLOCK_RANGE(0xdc00, 0xdfff, LowSurrogates,								"UnicodeBlock", "LowSurrogates", "Low Surrogates")
+REGISTER_UNICODE_BLOCK_RANGE(0xe000, 0xf8ff, PrivateUseArea,							"UnicodeBlock", "PrivateUseArea", "Private Use Area")
+REGISTER_UNICODE_BLOCK_RANGE(0xf900, 0xfaff, CJKCompatibilityIdeographs,				"UnicodeBlock", "CJKCompatibilityIdeographs", "CJK Compatibility Ideographs")
+REGISTER_UNICODE_BLOCK_RANGE(0xfb00, 0xfb4f, AlphabeticPresentationForms,				"UnicodeBlock", "AlphabeticPresentationForms", "Alphabetic Presentation Forms")
+REGISTER_UNICODE_BLOCK_RANGE(0xfb50, 0xfdff, ArabicPresentationFormsA,					"UnicodeBlock", "ArabicPresentationFormsA", "Arabic Presentation Forms - A")
+REGISTER_UNICODE_BLOCK_RANGE(0xfe00, 0xfe0f, VariationSelectors,						"UnicodeBlock", "VariationSelectors", "Variation Selectors")
+REGISTER_UNICODE_BLOCK_RANGE(0xfe10, 0xfe1f, VerticalForms,								"UnicodeBlock", "VerticalForms", "Vertical Forms")
+REGISTER_UNICODE_BLOCK_RANGE(0xfe20, 0xfe2f, CombiningHalfMarks,						"UnicodeBlock", "CombiningHalfMarks", "Combining Half Marks")
+REGISTER_UNICODE_BLOCK_RANGE(0xfe30, 0xfe4f, CJKCompatibilityForms,						"UnicodeBlock", "CJKCompatibilityForms", "CJK Compatibility Forms")
+REGISTER_UNICODE_BLOCK_RANGE(0xfe50, 0xfe6f, SmallFormVariants,							"UnicodeBlock", "SmallFormVariants", "Small Form Variants")
+REGISTER_UNICODE_BLOCK_RANGE(0xfe70, 0xfeff, ArabicPresentationFormsB,					"UnicodeBlock", "ArabicPresentationFormsB", "Arabic Presentation Forms - B")
+REGISTER_UNICODE_BLOCK_RANGE(0xff00, 0xffef, HalfwidthAndFullwidthForms,				"UnicodeBlock", "HalfwidthAndFullwidthForms", "Halfwidth and Fullwidth Forms")
+REGISTER_UNICODE_BLOCK_RANGE(0xfff0, 0xffff, Specials,									"UnicodeBlock", "Specials", "Specials")
 
 // The following are outside the Basic Multilingual Plane, so may not be supported by UE on all platforms
-REGISTER_UNICODE_BLOCK_RANGE(0x10000, 0x1007f, LinearBSyllabary,						NSLOCTEXT("UnicodeBlock", "LinearBSyllabary", "Linear B Syllabary"))
-REGISTER_UNICODE_BLOCK_RANGE(0x10080, 0x100ff, LinearBIdeograms,						NSLOCTEXT("UnicodeBlock", "LinearBIdeograms", "Linear B Ideograms"))
-REGISTER_UNICODE_BLOCK_RANGE(0x10100, 0x1013f, AegeanNumbers,							NSLOCTEXT("UnicodeBlock", "AegeanNumbers", "Aegean Numbers"))
-REGISTER_UNICODE_BLOCK_RANGE(0x10140, 0x1018f, AncientGreekNumbers,						NSLOCTEXT("UnicodeBlock", "AncientGreekNumbers", "Ancient Greek Numbers"))
-REGISTER_UNICODE_BLOCK_RANGE(0x10190, 0x101cf, AncientSymbols,							NSLOCTEXT("UnicodeBlock", "AncientSymbols", "Ancient Symbols"))
-REGISTER_UNICODE_BLOCK_RANGE(0x101d0, 0x101ff, PhaistosDisc,							NSLOCTEXT("UnicodeBlock", "PhaistosDisc", "Phaistos Disc"))
-REGISTER_UNICODE_BLOCK_RANGE(0x10280, 0x1029f, Lycian,									NSLOCTEXT("UnicodeBlock", "Lycian", "Lycian"))
-REGISTER_UNICODE_BLOCK_RANGE(0x102a0, 0x102df, Carian,									NSLOCTEXT("UnicodeBlock", "Carian", "Carian"))
-REGISTER_UNICODE_BLOCK_RANGE(0x102e0, 0x102ff, CopticEpactNumbers,						NSLOCTEXT("UnicodeBlock", "CopticEpactNumbers", "Coptic Epact Numbers"))
-REGISTER_UNICODE_BLOCK_RANGE(0x10300, 0x1032f, OldItalic,								NSLOCTEXT("UnicodeBlock", "OldItalic", "Old Italic"))
-REGISTER_UNICODE_BLOCK_RANGE(0x10330, 0x1034f, Gothic,									NSLOCTEXT("UnicodeBlock", "Gothic", "Gothic"))
-REGISTER_UNICODE_BLOCK_RANGE(0x10350, 0x1037f, OldPermic,								NSLOCTEXT("UnicodeBlock", "OldPermic", "Old Permic"))
-REGISTER_UNICODE_BLOCK_RANGE(0x10380, 0x1039f, Ugaritic,								NSLOCTEXT("UnicodeBlock", "Ugaritic", "Ugaritic"))
-REGISTER_UNICODE_BLOCK_RANGE(0x103a0, 0x103df, OldPersian,								NSLOCTEXT("UnicodeBlock", "Old Persian", "Old Persian"))
-REGISTER_UNICODE_BLOCK_RANGE(0x10400, 0x1044f, Deseret,									NSLOCTEXT("UnicodeBlock", "Deseret", "Deseret"))
-REGISTER_UNICODE_BLOCK_RANGE(0x10450, 0x1047f, Shavian,									NSLOCTEXT("UnicodeBlock", "Shavian", "Shavian"))
-REGISTER_UNICODE_BLOCK_RANGE(0x10480, 0x104af, Osmanya,									NSLOCTEXT("UnicodeBlock", "Osmanya", "Osmanya"))
-REGISTER_UNICODE_BLOCK_RANGE(0x104b0, 0x104ff, Osage,									NSLOCTEXT("UnicodeBlock", "Osage", "Osage"))
-REGISTER_UNICODE_BLOCK_RANGE(0x10500, 0x1052f, Elbasan,									NSLOCTEXT("UnicodeBlock", "Elbasan", "Elbasan"))
-REGISTER_UNICODE_BLOCK_RANGE(0x10530, 0x1056f, CaucasianAlbanian,						NSLOCTEXT("UnicodeBlock", "CaucasianAlbanian", "Caucasian Albanian"))
-REGISTER_UNICODE_BLOCK_RANGE(0x10600, 0x1077f, LinearA,									NSLOCTEXT("UnicodeBlock", "LinearA", "Linear A"))
-REGISTER_UNICODE_BLOCK_RANGE(0x10800, 0x1083f, CypriotSyllabary,						NSLOCTEXT("UnicodeBlock", "CypriotSyllabary", "Cypriot Syllabary"))
-REGISTER_UNICODE_BLOCK_RANGE(0x10840, 0x1085f, ImperialAramaic,							NSLOCTEXT("UnicodeBlock", "ImperialAramaic", "Imperial Aramaic"))
-REGISTER_UNICODE_BLOCK_RANGE(0x10860, 0x1087f, Palmyrene,								NSLOCTEXT("UnicodeBlock", "Palmyrene", "Palmyrene"))
-REGISTER_UNICODE_BLOCK_RANGE(0x10880, 0x108af, Nabataean,								NSLOCTEXT("UnicodeBlock", "Nabataean", "Nabataean"))
-REGISTER_UNICODE_BLOCK_RANGE(0x108e0, 0x108ff, Hatran,									NSLOCTEXT("UnicodeBlock", "Hatran", "Hatran"))
-REGISTER_UNICODE_BLOCK_RANGE(0x10900, 0x1091f, Phoenician,								NSLOCTEXT("UnicodeBlock", "Phoenician", "Phoenician"))
-REGISTER_UNICODE_BLOCK_RANGE(0x10920, 0x1093f, Lydian,									NSLOCTEXT("UnicodeBlock", "Lydian", "Lydian"))
-REGISTER_UNICODE_BLOCK_RANGE(0x10980, 0x1099f, MeroiticHieroglyphs,						NSLOCTEXT("UnicodeBlock", "MeroiticHieroglyphs", "Meroitic Hieroglyphs"))
-REGISTER_UNICODE_BLOCK_RANGE(0x109a0, 0x109ff, MeroiticCursive,							NSLOCTEXT("UnicodeBlock", "MeroiticCursive", "Meroitic Cursive"))
-REGISTER_UNICODE_BLOCK_RANGE(0x10a00, 0x10a5f, Kharoshthi,								NSLOCTEXT("UnicodeBlock", "Kharoshthi", "Kharoshthi"))
-REGISTER_UNICODE_BLOCK_RANGE(0x10a60, 0x10a7f, OldSouthArabian,							NSLOCTEXT("UnicodeBlock", "OldSouthArabian", "Old South Arabian"))
-REGISTER_UNICODE_BLOCK_RANGE(0x10a80, 0x10a9f, OldNorthArabian,							NSLOCTEXT("UnicodeBlock", "OldNorthArabian", "Old North Arabian"))
-REGISTER_UNICODE_BLOCK_RANGE(0x10ac0, 0x10aff, Manichaean,								NSLOCTEXT("UnicodeBlock", "Manichaean", "Manichaean"))
-REGISTER_UNICODE_BLOCK_RANGE(0x10b00, 0x10b3f, Avestan,									NSLOCTEXT("UnicodeBlock", "Avestan", "Avestan"))
-REGISTER_UNICODE_BLOCK_RANGE(0x10b40, 0x10b5f, InscriptionalParthian,					NSLOCTEXT("UnicodeBlock", "InscriptionalParthian", "Inscriptional Parthian"))
-REGISTER_UNICODE_BLOCK_RANGE(0x10b60, 0x10b7f, InscriptionalPahlavi,					NSLOCTEXT("UnicodeBlock", "InscriptionalPahlavi", "Inscriptional Pahlavi"))
-REGISTER_UNICODE_BLOCK_RANGE(0x10b80, 0x10baf, PsalterPahlavi,							NSLOCTEXT("UnicodeBlock", "PsalterPahlavi", "Psalter Pahlavi"))
-REGISTER_UNICODE_BLOCK_RANGE(0x10c00, 0x10c4f, OldTurkic,								NSLOCTEXT("UnicodeBlock", "OldTurkic", "Old Turkic"))
-REGISTER_UNICODE_BLOCK_RANGE(0x10c80, 0x10cff, OldHungarian,							NSLOCTEXT("UnicodeBlock", "OldHungarian", "Old Hungarian"))
-REGISTER_UNICODE_BLOCK_RANGE(0x10e60, 0x10e7f, RumiNumeralSymbols,						NSLOCTEXT("UnicodeBlock", "RumiNumeralSymbols", "Rumi Numeral Symbols"))
-REGISTER_UNICODE_BLOCK_RANGE(0x11000, 0x1107f, Brahmi,									NSLOCTEXT("UnicodeBlock", "Brahmi", "Brahmi"))
-REGISTER_UNICODE_BLOCK_RANGE(0x11080, 0x110cf, Kaithi,									NSLOCTEXT("UnicodeBlock", "Kaithi", "Kaithi"))
-REGISTER_UNICODE_BLOCK_RANGE(0x110d0, 0x110ff, SoraSompeng,								NSLOCTEXT("UnicodeBlock", "SoraSompeng", "Sora Sompeng"))
-REGISTER_UNICODE_BLOCK_RANGE(0x11100, 0x1114f, Chakma,									NSLOCTEXT("UnicodeBlock", "Chakma", "Chakma"))
-REGISTER_UNICODE_BLOCK_RANGE(0x11150, 0x1117f, Mahajani,								NSLOCTEXT("UnicodeBlock", "Mahajani", "Mahajani"))
-REGISTER_UNICODE_BLOCK_RANGE(0x11180, 0x111df, Sharada,									NSLOCTEXT("UnicodeBlock", "Sharada", "Sharada"))
-REGISTER_UNICODE_BLOCK_RANGE(0x111e0, 0x111ff, SinhalaArchaicNumbers,					NSLOCTEXT("UnicodeBlock", "SinhalaArchaicNumbers", "Sinhala Archaic Numbers"))
-REGISTER_UNICODE_BLOCK_RANGE(0x11200, 0x1124f, Khojki,									NSLOCTEXT("UnicodeBlock", "Khojki", "Khojki"))
-REGISTER_UNICODE_BLOCK_RANGE(0x11280, 0x112af, Multani,									NSLOCTEXT("UnicodeBlock", "Multani", "Multani"))
-REGISTER_UNICODE_BLOCK_RANGE(0x112b0, 0x112ff, Khudawadi,								NSLOCTEXT("UnicodeBlock", "Khudawadi", "Khudawadi"))
-REGISTER_UNICODE_BLOCK_RANGE(0x11300, 0x1137f, Grantha,									NSLOCTEXT("UnicodeBlock", "Grantha", "Grantha"))
-REGISTER_UNICODE_BLOCK_RANGE(0x11400, 0x1147f, Newa,									NSLOCTEXT("UnicodeBlock", "Newa", "Newa"))
-REGISTER_UNICODE_BLOCK_RANGE(0x11480, 0x114df, Tirhuta,									NSLOCTEXT("UnicodeBlock", "Tirhuta", "Tirhuta"))
-REGISTER_UNICODE_BLOCK_RANGE(0x11580, 0x115ff, Siddham,									NSLOCTEXT("UnicodeBlock", "Siddham", "Siddham"))
-REGISTER_UNICODE_BLOCK_RANGE(0x11600, 0x1165f, Modi,									NSLOCTEXT("UnicodeBlock", "Modi", "Modi"))
-REGISTER_UNICODE_BLOCK_RANGE(0x11660, 0x1167f, MongolianSupplement,						NSLOCTEXT("UnicodeBlock", "MongolianSupplement", "Mongolian Supplement"))
-REGISTER_UNICODE_BLOCK_RANGE(0x11680, 0x116cf, Takri,									NSLOCTEXT("UnicodeBlock", "Takri", "Takri"))
-REGISTER_UNICODE_BLOCK_RANGE(0x11700, 0x1173f, Ahom,									NSLOCTEXT("UnicodeBlock", "Ahom", "Ahom"))
-REGISTER_UNICODE_BLOCK_RANGE(0x118a0, 0x118ff, WarangCiti,								NSLOCTEXT("UnicodeBlock", "WarangCiti", "Warang Citi"))
-REGISTER_UNICODE_BLOCK_RANGE(0x11ac0, 0x11aff, PauCinHau,								NSLOCTEXT("UnicodeBlock", "PauCinHau", "Pau Cin Hau"))
-REGISTER_UNICODE_BLOCK_RANGE(0x11c00, 0x11c6f, Bhaiksuki,								NSLOCTEXT("UnicodeBlock", "Bhaiksuki", "Bhaiksuki"))
-REGISTER_UNICODE_BLOCK_RANGE(0x11c70, 0x11cbf, Marchen,									NSLOCTEXT("UnicodeBlock", "Marchen", "Marchen"))
-REGISTER_UNICODE_BLOCK_RANGE(0x12000, 0x123ff, Cuneiform,								NSLOCTEXT("UnicodeBlock", "Cuneiform", "Cuneiform"))
-REGISTER_UNICODE_BLOCK_RANGE(0x12400, 0x1247f, CuneiformNumbersAndPunctuation,			NSLOCTEXT("UnicodeBlock", "CuneiformNumbersAndPunctuation", "Cuneiform Numbers and Punctuation"))
-REGISTER_UNICODE_BLOCK_RANGE(0x12480, 0x1254f, EarlyDynasticCuneiform,					NSLOCTEXT("UnicodeBlock", "EarlyDynasticCuneiform", "Early Dynastic Cuneiform"))
-REGISTER_UNICODE_BLOCK_RANGE(0x13000, 0x1342f, EgyptianHieroglyphs,						NSLOCTEXT("UnicodeBlock", "EgyptianHieroglyphs", "Egyptian Hieroglyphs"))
-REGISTER_UNICODE_BLOCK_RANGE(0x14400, 0x1467f, AnatolianHieroglyphs,					NSLOCTEXT("UnicodeBlock", "AnatolianHieroglyphs", "Anatolian Hieroglyphs"))
-REGISTER_UNICODE_BLOCK_RANGE(0x16800, 0x16a3f, BamumSupplement,							NSLOCTEXT("UnicodeBlock", "BamumSupplement", "Bamum Supplement"))
-REGISTER_UNICODE_BLOCK_RANGE(0x16a40, 0x16a6f, Mro,										NSLOCTEXT("UnicodeBlock", "Mro", "Mro"))
-REGISTER_UNICODE_BLOCK_RANGE(0x16ad0, 0x16aff, BassaVah,								NSLOCTEXT("UnicodeBlock", "BassaVah", "Bassa Vah"))
-REGISTER_UNICODE_BLOCK_RANGE(0x16b00, 0x16b8f, PahawhHmong,								NSLOCTEXT("UnicodeBlock", "PahawhHmong", "Pahawh Hmong"))
-REGISTER_UNICODE_BLOCK_RANGE(0x16f00, 0x16f9f, Miao,									NSLOCTEXT("UnicodeBlock", "Miao", "Miao"))
-REGISTER_UNICODE_BLOCK_RANGE(0x16fe0, 0x16fff, IdeographicSymbolsAndPunctuation,		NSLOCTEXT("UnicodeBlock", "IdeographicSymbolsAndPunctuation", "Ideographic Symbols and Punctuation"))
-REGISTER_UNICODE_BLOCK_RANGE(0x17000, 0x187ff, Tangut,									NSLOCTEXT("UnicodeBlock", "Tangut", "Tangut"))
-REGISTER_UNICODE_BLOCK_RANGE(0x18800, 0x18aff, TangutComponents,						NSLOCTEXT("UnicodeBlock", "TangutComponents", "Tangut Components"))
-REGISTER_UNICODE_BLOCK_RANGE(0x1b000, 0x1b0ff, KanaSupplement,							NSLOCTEXT("UnicodeBlock", "KanaSupplement", "Kana Supplement"))
-REGISTER_UNICODE_BLOCK_RANGE(0x1bc00, 0x1bc9f, Duployan,								NSLOCTEXT("UnicodeBlock", "Duployan", "Duployan"))
-REGISTER_UNICODE_BLOCK_RANGE(0x1bca0, 0x1bcaf, ShorthandFormatControls,					NSLOCTEXT("UnicodeBlock", "ShorthandFormatControls", "Shorthand Format Controls"))
-REGISTER_UNICODE_BLOCK_RANGE(0x1d000, 0x1d0ff, ByzantineMusicalSymbols,					NSLOCTEXT("UnicodeBlock", "ByzantineMusicalSymbols", "Byzantine Musical Symbols"))
-REGISTER_UNICODE_BLOCK_RANGE(0x1d100, 0x1d1ff, MusicalSymbols,							NSLOCTEXT("UnicodeBlock", "MusicalSymbols", "Musical Symbols"))
-REGISTER_UNICODE_BLOCK_RANGE(0x1d200, 0x1d24f, AncientGreekMusicalNotation,				NSLOCTEXT("UnicodeBlock", "AncientGreekMusicalNotation", "Ancient Greek Musical Notation"))
-REGISTER_UNICODE_BLOCK_RANGE(0x1d300, 0x1d35f, TaiXuanJingSymbols,						NSLOCTEXT("UnicodeBlock", "TaiXuanJingSymbols", "Tai Xuan Jing Symbols"))
-REGISTER_UNICODE_BLOCK_RANGE(0x1d360, 0x1d37f, CountingRodNumerals,						NSLOCTEXT("UnicodeBlock", "CountingRodNumerals", "Counting Rod Numerals"))
-REGISTER_UNICODE_BLOCK_RANGE(0x1d400, 0x1d7ff, MathematicalAlphanumericSymbols,			NSLOCTEXT("UnicodeBlock", "MathematicalAlphanumericSymbols", "Mathematical Alphanumeric Symbols"))
-REGISTER_UNICODE_BLOCK_RANGE(0x1d800, 0x1daaf, SuttonSignWriting,						NSLOCTEXT("UnicodeBlock", "SuttonSignWriting", "Sutton SignWriting"))
-REGISTER_UNICODE_BLOCK_RANGE(0x1e000, 0x1e02f, GlagoliticSupplement,					NSLOCTEXT("UnicodeBlock", "GlagoliticSupplement", "Glagolitic Supplement"))
-REGISTER_UNICODE_BLOCK_RANGE(0x1e800, 0x1e8df, MendeKikakui,							NSLOCTEXT("UnicodeBlock", "MendeKikakui", "Mende Kikakui"))
-REGISTER_UNICODE_BLOCK_RANGE(0x1e900, 0x1e95f, Adlam,									NSLOCTEXT("UnicodeBlock", "Adlam", "Adlam"))
-REGISTER_UNICODE_BLOCK_RANGE(0x1ee00, 0x1eeff, ArabicMathematicalAlphabeticSymbols,		NSLOCTEXT("UnicodeBlock", "ArabicMathematicalAlphabeticSymbols", "Arabic Mathematical Alphabetic Symbols"))
-REGISTER_UNICODE_BLOCK_RANGE(0x1f000, 0x1f02f, MahjongTiles,							NSLOCTEXT("UnicodeBlock", "MahjongTiles", "Mahjong Tiles"))
-REGISTER_UNICODE_BLOCK_RANGE(0x1f030, 0x1f09f, DominoTiles,								NSLOCTEXT("UnicodeBlock", "DominoTiles", "Domino Tiles"))
-REGISTER_UNICODE_BLOCK_RANGE(0x1f0a0, 0x1f0ff, PlayingCards,							NSLOCTEXT("UnicodeBlock", "PlayingCards", "Playing Cards"))
-REGISTER_UNICODE_BLOCK_RANGE(0x1f100, 0x1f1ff, EnclosedAlphanumericSupplement,			NSLOCTEXT("UnicodeBlock", "EnclosedAlphanumericSupplement", "Enclosed Alphanumeric Supplement"))
-REGISTER_UNICODE_BLOCK_RANGE(0x1f200, 0x1f2ff, EnclosedIdeographicSupplement,			NSLOCTEXT("UnicodeBlock", "EnclosedIdeographicSupplement", "Enclosed Ideographic Supplement"))
-REGISTER_UNICODE_BLOCK_RANGE(0x1f300, 0x1f5ff, MiscellaneousSymbolsAndPictographs,		NSLOCTEXT("UnicodeBlock", "MiscellaneousSymbolsAndPictographs", "Miscellaneous Symbols and Pictographs"))
-REGISTER_UNICODE_BLOCK_RANGE(0x1f600, 0x1f64f, EmoticonsEmoji,							NSLOCTEXT("UnicodeBlock", "EmoticonsEmoji", "Emoticons (Emoji)"))
-REGISTER_UNICODE_BLOCK_RANGE(0x1f650, 0x1f67f, OrnamentalDingbats,						NSLOCTEXT("UnicodeBlock", "OrnamentalDingbats", "Ornamental Dingbats"))
-REGISTER_UNICODE_BLOCK_RANGE(0x1f680, 0x1f6ff, TransportAndMapSymbols,					NSLOCTEXT("UnicodeBlock", "TransportAndMapSymbols", "Transport and Map Symbols"))
-REGISTER_UNICODE_BLOCK_RANGE(0x1f700, 0x1f77f, AlchemicalSymbols,						NSLOCTEXT("UnicodeBlock", "AlchemicalSymbols", "Alchemical Symbols"))
-REGISTER_UNICODE_BLOCK_RANGE(0x1f780, 0x1f7ff, GeometricShapesExtended,					NSLOCTEXT("UnicodeBlock", "GeometricShapesExtended", "Geometric Shapes Extended"))
-REGISTER_UNICODE_BLOCK_RANGE(0x1f800, 0x1f8ff, SupplementalArrowsC,						NSLOCTEXT("UnicodeBlock", "SupplementalArrowsC", "Supplemental Arrows - C"))
-REGISTER_UNICODE_BLOCK_RANGE(0x1f900, 0x1f9ff, SupplementalSymbolsAndPictographs,		NSLOCTEXT("UnicodeBlock", "SupplementalSymbolsAndPictographs", "Supplemental Symbols and Pictographs"))
-REGISTER_UNICODE_BLOCK_RANGE(0x20000, 0x2a6d6, CJKUnifiedIdeographsExtensionB,			NSLOCTEXT("UnicodeBlock", "CJKUnifiedIdeographsExtensionB", "CJK Unified Ideographs Extension B"))
-REGISTER_UNICODE_BLOCK_RANGE(0x2a700, 0x2b734, CJKUnifiedIdeographsExtensionC,			NSLOCTEXT("UnicodeBlock", "CJKUnifiedIdeographsExtensionC", "CJK Unified Ideographs Extension C"))
-REGISTER_UNICODE_BLOCK_RANGE(0x2b740, 0x2b81d, CJKUnifiedIdeographsExtensionD,			NSLOCTEXT("UnicodeBlock", "CJKUnifiedIdeographsExtensionD", "CJK Unified Ideographs Extension D"))
-REGISTER_UNICODE_BLOCK_RANGE(0x2b820, 0x2cea1, CJKUnifiedIdeographsExtensionE,			NSLOCTEXT("UnicodeBlock", "CJKUnifiedIdeographsExtensionE", "CJK Unified Ideographs Extension E"))
-REGISTER_UNICODE_BLOCK_RANGE(0x2f800, 0x2fa1f, CJKCompatibilityIdeographsSupplement,	NSLOCTEXT("UnicodeBlock", "CJKCompatibilityIdeographsSupplement", "CJK Compatibility Ideographs Supplement"))
-REGISTER_UNICODE_BLOCK_RANGE(0xe0000, 0xe007f, Tags,									NSLOCTEXT("UnicodeBlock", "Tags", "Tags"))
-REGISTER_UNICODE_BLOCK_RANGE(0xe0100, 0xe01ef, VariationSelectorsSupplement,			NSLOCTEXT("UnicodeBlock", "VariationSelectorsSupplement", "Variation Selectors Supplement"))
+REGISTER_UNICODE_BLOCK_RANGE(0x10000, 0x1007f, LinearBSyllabary,						"UnicodeBlock", "LinearBSyllabary", "Linear B Syllabary")
+REGISTER_UNICODE_BLOCK_RANGE(0x10080, 0x100ff, LinearBIdeograms,						"UnicodeBlock", "LinearBIdeograms", "Linear B Ideograms")
+REGISTER_UNICODE_BLOCK_RANGE(0x10100, 0x1013f, AegeanNumbers,							"UnicodeBlock", "AegeanNumbers", "Aegean Numbers")
+REGISTER_UNICODE_BLOCK_RANGE(0x10140, 0x1018f, AncientGreekNumbers,						"UnicodeBlock", "AncientGreekNumbers", "Ancient Greek Numbers")
+REGISTER_UNICODE_BLOCK_RANGE(0x10190, 0x101cf, AncientSymbols,							"UnicodeBlock", "AncientSymbols", "Ancient Symbols")
+REGISTER_UNICODE_BLOCK_RANGE(0x101d0, 0x101ff, PhaistosDisc,							"UnicodeBlock", "PhaistosDisc", "Phaistos Disc")
+REGISTER_UNICODE_BLOCK_RANGE(0x10280, 0x1029f, Lycian,									"UnicodeBlock", "Lycian", "Lycian")
+REGISTER_UNICODE_BLOCK_RANGE(0x102a0, 0x102df, Carian,									"UnicodeBlock", "Carian", "Carian")
+REGISTER_UNICODE_BLOCK_RANGE(0x102e0, 0x102ff, CopticEpactNumbers,						"UnicodeBlock", "CopticEpactNumbers", "Coptic Epact Numbers")
+REGISTER_UNICODE_BLOCK_RANGE(0x10300, 0x1032f, OldItalic,								"UnicodeBlock", "OldItalic", "Old Italic")
+REGISTER_UNICODE_BLOCK_RANGE(0x10330, 0x1034f, Gothic,									"UnicodeBlock", "Gothic", "Gothic")
+REGISTER_UNICODE_BLOCK_RANGE(0x10350, 0x1037f, OldPermic,								"UnicodeBlock", "OldPermic", "Old Permic")
+REGISTER_UNICODE_BLOCK_RANGE(0x10380, 0x1039f, Ugaritic,								"UnicodeBlock", "Ugaritic", "Ugaritic")
+REGISTER_UNICODE_BLOCK_RANGE(0x103a0, 0x103df, OldPersian,								"UnicodeBlock", "Old Persian", "Old Persian")
+REGISTER_UNICODE_BLOCK_RANGE(0x10400, 0x1044f, Deseret,									"UnicodeBlock", "Deseret", "Deseret")
+REGISTER_UNICODE_BLOCK_RANGE(0x10450, 0x1047f, Shavian,									"UnicodeBlock", "Shavian", "Shavian")
+REGISTER_UNICODE_BLOCK_RANGE(0x10480, 0x104af, Osmanya,									"UnicodeBlock", "Osmanya", "Osmanya")
+REGISTER_UNICODE_BLOCK_RANGE(0x104b0, 0x104ff, Osage,									"UnicodeBlock", "Osage", "Osage")
+REGISTER_UNICODE_BLOCK_RANGE(0x10500, 0x1052f, Elbasan,									"UnicodeBlock", "Elbasan", "Elbasan")
+REGISTER_UNICODE_BLOCK_RANGE(0x10530, 0x1056f, CaucasianAlbanian,						"UnicodeBlock", "CaucasianAlbanian", "Caucasian Albanian")
+REGISTER_UNICODE_BLOCK_RANGE(0x10600, 0x1077f, LinearA,									"UnicodeBlock", "LinearA", "Linear A")
+REGISTER_UNICODE_BLOCK_RANGE(0x10800, 0x1083f, CypriotSyllabary,						"UnicodeBlock", "CypriotSyllabary", "Cypriot Syllabary")
+REGISTER_UNICODE_BLOCK_RANGE(0x10840, 0x1085f, ImperialAramaic,							"UnicodeBlock", "ImperialAramaic", "Imperial Aramaic")
+REGISTER_UNICODE_BLOCK_RANGE(0x10860, 0x1087f, Palmyrene,								"UnicodeBlock", "Palmyrene", "Palmyrene")
+REGISTER_UNICODE_BLOCK_RANGE(0x10880, 0x108af, Nabataean,								"UnicodeBlock", "Nabataean", "Nabataean")
+REGISTER_UNICODE_BLOCK_RANGE(0x108e0, 0x108ff, Hatran,									"UnicodeBlock", "Hatran", "Hatran")
+REGISTER_UNICODE_BLOCK_RANGE(0x10900, 0x1091f, Phoenician,								"UnicodeBlock", "Phoenician", "Phoenician")
+REGISTER_UNICODE_BLOCK_RANGE(0x10920, 0x1093f, Lydian,									"UnicodeBlock", "Lydian", "Lydian")
+REGISTER_UNICODE_BLOCK_RANGE(0x10980, 0x1099f, MeroiticHieroglyphs,						"UnicodeBlock", "MeroiticHieroglyphs", "Meroitic Hieroglyphs")
+REGISTER_UNICODE_BLOCK_RANGE(0x109a0, 0x109ff, MeroiticCursive,							"UnicodeBlock", "MeroiticCursive", "Meroitic Cursive")
+REGISTER_UNICODE_BLOCK_RANGE(0x10a00, 0x10a5f, Kharoshthi,								"UnicodeBlock", "Kharoshthi", "Kharoshthi")
+REGISTER_UNICODE_BLOCK_RANGE(0x10a60, 0x10a7f, OldSouthArabian,							"UnicodeBlock", "OldSouthArabian", "Old South Arabian")
+REGISTER_UNICODE_BLOCK_RANGE(0x10a80, 0x10a9f, OldNorthArabian,							"UnicodeBlock", "OldNorthArabian", "Old North Arabian")
+REGISTER_UNICODE_BLOCK_RANGE(0x10ac0, 0x10aff, Manichaean,								"UnicodeBlock", "Manichaean", "Manichaean")
+REGISTER_UNICODE_BLOCK_RANGE(0x10b00, 0x10b3f, Avestan,									"UnicodeBlock", "Avestan", "Avestan")
+REGISTER_UNICODE_BLOCK_RANGE(0x10b40, 0x10b5f, InscriptionalParthian,					"UnicodeBlock", "InscriptionalParthian", "Inscriptional Parthian")
+REGISTER_UNICODE_BLOCK_RANGE(0x10b60, 0x10b7f, InscriptionalPahlavi,					"UnicodeBlock", "InscriptionalPahlavi", "Inscriptional Pahlavi")
+REGISTER_UNICODE_BLOCK_RANGE(0x10b80, 0x10baf, PsalterPahlavi,							"UnicodeBlock", "PsalterPahlavi", "Psalter Pahlavi")
+REGISTER_UNICODE_BLOCK_RANGE(0x10c00, 0x10c4f, OldTurkic,								"UnicodeBlock", "OldTurkic", "Old Turkic")
+REGISTER_UNICODE_BLOCK_RANGE(0x10c80, 0x10cff, OldHungarian,							"UnicodeBlock", "OldHungarian", "Old Hungarian")
+REGISTER_UNICODE_BLOCK_RANGE(0x10e60, 0x10e7f, RumiNumeralSymbols,						"UnicodeBlock", "RumiNumeralSymbols", "Rumi Numeral Symbols")
+REGISTER_UNICODE_BLOCK_RANGE(0x11000, 0x1107f, Brahmi,									"UnicodeBlock", "Brahmi", "Brahmi")
+REGISTER_UNICODE_BLOCK_RANGE(0x11080, 0x110cf, Kaithi,									"UnicodeBlock", "Kaithi", "Kaithi")
+REGISTER_UNICODE_BLOCK_RANGE(0x110d0, 0x110ff, SoraSompeng,								"UnicodeBlock", "SoraSompeng", "Sora Sompeng")
+REGISTER_UNICODE_BLOCK_RANGE(0x11100, 0x1114f, Chakma,									"UnicodeBlock", "Chakma", "Chakma")
+REGISTER_UNICODE_BLOCK_RANGE(0x11150, 0x1117f, Mahajani,								"UnicodeBlock", "Mahajani", "Mahajani")
+REGISTER_UNICODE_BLOCK_RANGE(0x11180, 0x111df, Sharada,									"UnicodeBlock", "Sharada", "Sharada")
+REGISTER_UNICODE_BLOCK_RANGE(0x111e0, 0x111ff, SinhalaArchaicNumbers,					"UnicodeBlock", "SinhalaArchaicNumbers", "Sinhala Archaic Numbers")
+REGISTER_UNICODE_BLOCK_RANGE(0x11200, 0x1124f, Khojki,									"UnicodeBlock", "Khojki", "Khojki")
+REGISTER_UNICODE_BLOCK_RANGE(0x11280, 0x112af, Multani,									"UnicodeBlock", "Multani", "Multani")
+REGISTER_UNICODE_BLOCK_RANGE(0x112b0, 0x112ff, Khudawadi,								"UnicodeBlock", "Khudawadi", "Khudawadi")
+REGISTER_UNICODE_BLOCK_RANGE(0x11300, 0x1137f, Grantha,									"UnicodeBlock", "Grantha", "Grantha")
+REGISTER_UNICODE_BLOCK_RANGE(0x11400, 0x1147f, Newa,									"UnicodeBlock", "Newa", "Newa")
+REGISTER_UNICODE_BLOCK_RANGE(0x11480, 0x114df, Tirhuta,									"UnicodeBlock", "Tirhuta", "Tirhuta")
+REGISTER_UNICODE_BLOCK_RANGE(0x11580, 0x115ff, Siddham,									"UnicodeBlock", "Siddham", "Siddham")
+REGISTER_UNICODE_BLOCK_RANGE(0x11600, 0x1165f, Modi,									"UnicodeBlock", "Modi", "Modi")
+REGISTER_UNICODE_BLOCK_RANGE(0x11660, 0x1167f, MongolianSupplement,						"UnicodeBlock", "MongolianSupplement", "Mongolian Supplement")
+REGISTER_UNICODE_BLOCK_RANGE(0x11680, 0x116cf, Takri,									"UnicodeBlock", "Takri", "Takri")
+REGISTER_UNICODE_BLOCK_RANGE(0x11700, 0x1173f, Ahom,									"UnicodeBlock", "Ahom", "Ahom")
+REGISTER_UNICODE_BLOCK_RANGE(0x118a0, 0x118ff, WarangCiti,								"UnicodeBlock", "WarangCiti", "Warang Citi")
+REGISTER_UNICODE_BLOCK_RANGE(0x11ac0, 0x11aff, PauCinHau,								"UnicodeBlock", "PauCinHau", "Pau Cin Hau")
+REGISTER_UNICODE_BLOCK_RANGE(0x11c00, 0x11c6f, Bhaiksuki,								"UnicodeBlock", "Bhaiksuki", "Bhaiksuki")
+REGISTER_UNICODE_BLOCK_RANGE(0x11c70, 0x11cbf, Marchen,									"UnicodeBlock", "Marchen", "Marchen")
+REGISTER_UNICODE_BLOCK_RANGE(0x12000, 0x123ff, Cuneiform,								"UnicodeBlock", "Cuneiform", "Cuneiform")
+REGISTER_UNICODE_BLOCK_RANGE(0x12400, 0x1247f, CuneiformNumbersAndPunctuation,			"UnicodeBlock", "CuneiformNumbersAndPunctuation", "Cuneiform Numbers and Punctuation")
+REGISTER_UNICODE_BLOCK_RANGE(0x12480, 0x1254f, EarlyDynasticCuneiform,					"UnicodeBlock", "EarlyDynasticCuneiform", "Early Dynastic Cuneiform")
+REGISTER_UNICODE_BLOCK_RANGE(0x13000, 0x1342f, EgyptianHieroglyphs,						"UnicodeBlock", "EgyptianHieroglyphs", "Egyptian Hieroglyphs")
+REGISTER_UNICODE_BLOCK_RANGE(0x14400, 0x1467f, AnatolianHieroglyphs,					"UnicodeBlock", "AnatolianHieroglyphs", "Anatolian Hieroglyphs")
+REGISTER_UNICODE_BLOCK_RANGE(0x16800, 0x16a3f, BamumSupplement,							"UnicodeBlock", "BamumSupplement", "Bamum Supplement")
+REGISTER_UNICODE_BLOCK_RANGE(0x16a40, 0x16a6f, Mro,										"UnicodeBlock", "Mro", "Mro")
+REGISTER_UNICODE_BLOCK_RANGE(0x16ad0, 0x16aff, BassaVah,								"UnicodeBlock", "BassaVah", "Bassa Vah")
+REGISTER_UNICODE_BLOCK_RANGE(0x16b00, 0x16b8f, PahawhHmong,								"UnicodeBlock", "PahawhHmong", "Pahawh Hmong")
+REGISTER_UNICODE_BLOCK_RANGE(0x16f00, 0x16f9f, Miao,									"UnicodeBlock", "Miao", "Miao")
+REGISTER_UNICODE_BLOCK_RANGE(0x16fe0, 0x16fff, IdeographicSymbolsAndPunctuation,		"UnicodeBlock", "IdeographicSymbolsAndPunctuation", "Ideographic Symbols and Punctuation")
+REGISTER_UNICODE_BLOCK_RANGE(0x17000, 0x187ff, Tangut,									"UnicodeBlock", "Tangut", "Tangut")
+REGISTER_UNICODE_BLOCK_RANGE(0x18800, 0x18aff, TangutComponents,						"UnicodeBlock", "TangutComponents", "Tangut Components")
+REGISTER_UNICODE_BLOCK_RANGE(0x1b000, 0x1b0ff, KanaSupplement,							"UnicodeBlock", "KanaSupplement", "Kana Supplement")
+REGISTER_UNICODE_BLOCK_RANGE(0x1bc00, 0x1bc9f, Duployan,								"UnicodeBlock", "Duployan", "Duployan")
+REGISTER_UNICODE_BLOCK_RANGE(0x1bca0, 0x1bcaf, ShorthandFormatControls,					"UnicodeBlock", "ShorthandFormatControls", "Shorthand Format Controls")
+REGISTER_UNICODE_BLOCK_RANGE(0x1d000, 0x1d0ff, ByzantineMusicalSymbols,					"UnicodeBlock", "ByzantineMusicalSymbols", "Byzantine Musical Symbols")
+REGISTER_UNICODE_BLOCK_RANGE(0x1d100, 0x1d1ff, MusicalSymbols,							"UnicodeBlock", "MusicalSymbols", "Musical Symbols")
+REGISTER_UNICODE_BLOCK_RANGE(0x1d200, 0x1d24f, AncientGreekMusicalNotation,				"UnicodeBlock", "AncientGreekMusicalNotation", "Ancient Greek Musical Notation")
+REGISTER_UNICODE_BLOCK_RANGE(0x1d300, 0x1d35f, TaiXuanJingSymbols,						"UnicodeBlock", "TaiXuanJingSymbols", "Tai Xuan Jing Symbols")
+REGISTER_UNICODE_BLOCK_RANGE(0x1d360, 0x1d37f, CountingRodNumerals,						"UnicodeBlock", "CountingRodNumerals", "Counting Rod Numerals")
+REGISTER_UNICODE_BLOCK_RANGE(0x1d400, 0x1d7ff, MathematicalAlphanumericSymbols,			"UnicodeBlock", "MathematicalAlphanumericSymbols", "Mathematical Alphanumeric Symbols")
+REGISTER_UNICODE_BLOCK_RANGE(0x1d800, 0x1daaf, SuttonSignWriting,						"UnicodeBlock", "SuttonSignWriting", "Sutton SignWriting")
+REGISTER_UNICODE_BLOCK_RANGE(0x1e000, 0x1e02f, GlagoliticSupplement,					"UnicodeBlock", "GlagoliticSupplement", "Glagolitic Supplement")
+REGISTER_UNICODE_BLOCK_RANGE(0x1e800, 0x1e8df, MendeKikakui,							"UnicodeBlock", "MendeKikakui", "Mende Kikakui")
+REGISTER_UNICODE_BLOCK_RANGE(0x1e900, 0x1e95f, Adlam,									"UnicodeBlock", "Adlam", "Adlam")
+REGISTER_UNICODE_BLOCK_RANGE(0x1ee00, 0x1eeff, ArabicMathematicalAlphabeticSymbols,		"UnicodeBlock", "ArabicMathematicalAlphabeticSymbols", "Arabic Mathematical Alphabetic Symbols")
+REGISTER_UNICODE_BLOCK_RANGE(0x1f000, 0x1f02f, MahjongTiles,							"UnicodeBlock", "MahjongTiles", "Mahjong Tiles")
+REGISTER_UNICODE_BLOCK_RANGE(0x1f030, 0x1f09f, DominoTiles,								"UnicodeBlock", "DominoTiles", "Domino Tiles")
+REGISTER_UNICODE_BLOCK_RANGE(0x1f0a0, 0x1f0ff, PlayingCards,							"UnicodeBlock", "PlayingCards", "Playing Cards")
+REGISTER_UNICODE_BLOCK_RANGE(0x1f100, 0x1f1ff, EnclosedAlphanumericSupplement,			"UnicodeBlock", "EnclosedAlphanumericSupplement", "Enclosed Alphanumeric Supplement")
+REGISTER_UNICODE_BLOCK_RANGE(0x1f200, 0x1f2ff, EnclosedIdeographicSupplement,			"UnicodeBlock", "EnclosedIdeographicSupplement", "Enclosed Ideographic Supplement")
+REGISTER_UNICODE_BLOCK_RANGE(0x1f300, 0x1f5ff, MiscellaneousSymbolsAndPictographs,		"UnicodeBlock", "MiscellaneousSymbolsAndPictographs", "Miscellaneous Symbols and Pictographs")
+REGISTER_UNICODE_BLOCK_RANGE(0x1f600, 0x1f64f, EmoticonsEmoji,							"UnicodeBlock", "EmoticonsEmoji", "Emoticons (Emoji)")
+REGISTER_UNICODE_BLOCK_RANGE(0x1f650, 0x1f67f, OrnamentalDingbats,						"UnicodeBlock", "OrnamentalDingbats", "Ornamental Dingbats")
+REGISTER_UNICODE_BLOCK_RANGE(0x1f680, 0x1f6ff, TransportAndMapSymbols,					"UnicodeBlock", "TransportAndMapSymbols", "Transport and Map Symbols")
+REGISTER_UNICODE_BLOCK_RANGE(0x1f700, 0x1f77f, AlchemicalSymbols,						"UnicodeBlock", "AlchemicalSymbols", "Alchemical Symbols")
+REGISTER_UNICODE_BLOCK_RANGE(0x1f780, 0x1f7ff, GeometricShapesExtended,					"UnicodeBlock", "GeometricShapesExtended", "Geometric Shapes Extended")
+REGISTER_UNICODE_BLOCK_RANGE(0x1f800, 0x1f8ff, SupplementalArrowsC,						"UnicodeBlock", "SupplementalArrowsC", "Supplemental Arrows - C")
+REGISTER_UNICODE_BLOCK_RANGE(0x1f900, 0x1f9ff, SupplementalSymbolsAndPictographs,		"UnicodeBlock", "SupplementalSymbolsAndPictographs", "Supplemental Symbols and Pictographs")
+REGISTER_UNICODE_BLOCK_RANGE(0x20000, 0x2a6d6, CJKUnifiedIdeographsExtensionB,			"UnicodeBlock", "CJKUnifiedIdeographsExtensionB", "CJK Unified Ideographs Extension B")
+REGISTER_UNICODE_BLOCK_RANGE(0x2a700, 0x2b734, CJKUnifiedIdeographsExtensionC,			"UnicodeBlock", "CJKUnifiedIdeographsExtensionC", "CJK Unified Ideographs Extension C")
+REGISTER_UNICODE_BLOCK_RANGE(0x2b740, 0x2b81d, CJKUnifiedIdeographsExtensionD,			"UnicodeBlock", "CJKUnifiedIdeographsExtensionD", "CJK Unified Ideographs Extension D")
+REGISTER_UNICODE_BLOCK_RANGE(0x2b820, 0x2cea1, CJKUnifiedIdeographsExtensionE,			"UnicodeBlock", "CJKUnifiedIdeographsExtensionE", "CJK Unified Ideographs Extension E")
+REGISTER_UNICODE_BLOCK_RANGE(0x2f800, 0x2fa1f, CJKCompatibilityIdeographsSupplement,	"UnicodeBlock", "CJKCompatibilityIdeographsSupplement", "CJK Compatibility Ideographs Supplement")
+REGISTER_UNICODE_BLOCK_RANGE(0xe0000, 0xe007f, Tags,									"UnicodeBlock", "Tags", "Tags")
+REGISTER_UNICODE_BLOCK_RANGE(0xe0100, 0xe01ef, VariationSelectorsSupplement,			"UnicodeBlock", "VariationSelectorsSupplement", "Variation Selectors Supplement")
