[SlateCore]
* Reduced compile time of UMGCoreStyle.cpp from 32s to 6s with msvc by introducing anonymous lambda scopes

diff --git a/Engine/Source/Runtime/SlateCore/Private/Styling/UMGCoreStyle.cpp b/Engine/Source/Runtime/SlateCore/Private/Styling/UMGCoreStyle.cpp
--- a/Engine/Source/Runtime/SlateCore/Private/Styling/UMGCoreStyle.cpp
+++ b/Engine/Source/Runtime/SlateCore/Private/Styling/UMGCoreStyle.cpp
@@ -17,6 +17,18 @@
 #include "Styling/StyleColors.h"
 #include "Styling/CoreStyle.h"
 
+// This is super annoying but huge functions can take a very long time to compile on msvc
+// .. and by sprinkling these splits inside functions we cut this file's compile time in half
+#if PLATFORM_COMPILER_CLANG
+	#define UE_COMPILE_SCOPE_BEGIN
+	#define UE_COMPILE_SCOPE_SPLIT
+	#define UE_COMPILE_SCOPE_END
+#else
+	#define UE_COMPILE_SCOPE_BEGIN [&](){
+	#define UE_COMPILE_SCOPE_SPLIT }();[&](){
+	#define UE_COMPILE_SCOPE_END }();
+#endif
+
 // This is to fix the issue that SlateStyleMacros like IMAGE_BRUSH look for RootToContentDir but Style->RootToContentDir is how the core style is set up
 #define RootToContentDir Style->RootToContentDir
 
@@ -297,6 +309,8 @@ TSharedRef<ISlateStyle> FUMGCoreStyle::Create()
 		.SetNormalPadding(FMargin(0.0f, 0.0f, 0.0f, 1.0f))
 		.SetPressedPadding(FMargin(0.0f, 1.0f, 0.0f, 0.0f));
 
+	UE_COMPILE_SCOPE_BEGIN
+
 	// Convenient transparent/invisible elements
 	{
 		Style->Set("NoBrush", new FSlateNoResource());
@@ -369,6 +383,8 @@ TSharedRef<ISlateStyle> FUMGCoreStyle::Create()
 		Style->Set( "GenericCommands.Rename", new IMAGE_BRUSH( "Icons/Edit/Icon_Edit_Rename_16x", Icon16x16 ) );
 	}
 
+	UE_COMPILE_SCOPE_END
+
 	// SVerticalBox Drag& Drop Icon
 	Style->Set("VerticalBoxDragIndicator", new IMAGE_BRUSH("Common/VerticalBoxDragIndicator", FVector2f(6.f, 45.f)));
 	Style->Set("VerticalBoxDragIndicatorShort", new IMAGE_BRUSH("Common/VerticalBoxDragIndicatorShort", FVector2f(6.f, 15.f)));
@@ -411,6 +427,8 @@ TSharedRef<ISlateStyle> FUMGCoreStyle::Create()
 		Style->Set( "InvertedForeground", InvertedForeground );
 	}
 
+	UE_COMPILE_SCOPE_BEGIN
+
 	// SComboButton and SComboBox defaults...
 	{
 		FComboButtonStyle ComboButton = FComboButtonStyle()
@@ -511,6 +529,8 @@ TSharedRef<ISlateStyle> FUMGCoreStyle::Create()
 		Style->Set( "RadioButton", BasicRadioButtonStyle );
 	}
 
+	UE_COMPILE_SCOPE_END
+
 	// SEditableText defaults...
 	{
 		FSlateBrush* SelectionBackground = new BOX_BRUSH( "Common/EditableTextSelectionBackground",  FMargin(4.f/16.f) );
@@ -541,6 +561,8 @@ TSharedRef<ISlateStyle> FUMGCoreStyle::Create()
 		// "NormalFont".
 	}
 
+	UE_COMPILE_SCOPE_BEGIN
+
 	static const FEditableTextBoxStyle DarkEditableTextBoxStyle = FEditableTextBoxStyle()
 		.SetTextStyle(NormalText)
 		.SetBackgroundImageNormal(BOX_BRUSH("Common/TextBox_Dark", FMargin(4.0f / 16.0f)))
@@ -648,6 +670,8 @@ TSharedRef<ISlateStyle> FUMGCoreStyle::Create()
 		Style->Set("Hyperlink", Hyperlink);
 	}
 
+	UE_COMPILE_SCOPE_SPLIT
+
 	// SProgressBar defaults...
 	{
 		Style->Set( "ProgressBar", FProgressBarStyle()
@@ -726,6 +750,8 @@ TSharedRef<ISlateStyle> FUMGCoreStyle::Create()
 			);
 	}
 
+	UE_COMPILE_SCOPE_SPLIT
+
 	// SNumericEntryBox defaults...
 	{
 		Style->Set( "NumericEntrySpinBox", FSpinBoxStyle()
@@ -792,6 +818,8 @@ TSharedRef<ISlateStyle> FUMGCoreStyle::Create()
 		Style->Set("ColorGradingWheel.Selector", new IMAGE_BRUSH("Common/Circle", FVector2f(8.f, 8.f)));
 	}
 
+	UE_COMPILE_SCOPE_SPLIT
+
 	// SSplitter
 	{
 		Style->Set( "Splitter", FSplitterStyle()
@@ -882,6 +910,8 @@ TSharedRef<ISlateStyle> FUMGCoreStyle::Create()
 			);
 	}
 
+	UE_COMPILE_SCOPE_SPLIT
+
 	// MultiBox
 	{
 		Style->Set( "MultiBox.GenericToolBarIcon", new IMAGE_BRUSH( "Icons/Icon_generic_toolbar", Icon40x40 ) );
@@ -997,6 +1027,8 @@ TSharedRef<ISlateStyle> FUMGCoreStyle::Create()
 		Style->Set("ToolBar.IconSize", Icon20x20);
 	}
 
+	UE_COMPILE_SCOPE_SPLIT
+
 	// MenuBar
 	{
 		Style->Set( "Menu.Background",	new BOX_BRUSH( "Old/Menu_Background", FMargin(8.0f/64.0f) ) );
@@ -1113,6 +1145,8 @@ TSharedRef<ISlateStyle> FUMGCoreStyle::Create()
 		Style->Set("WindowMenuBar.MenuBar.Padding", FMargin(12, 2));
 	}
 
+	UE_COMPILE_SCOPE_SPLIT
+
 	// SExpandableButton defaults...
 	{
 		Style->Set( "ExpandableButton.Background", new BOX_BRUSH( "Common/Button", 8.0f/32.0f ) );
@@ -1253,6 +1287,8 @@ TSharedRef<ISlateStyle> FUMGCoreStyle::Create()
 		Style->Set( "Docking.Cross.HoveredTint", FLinearColor::White );
 	}
 
+	UE_COMPILE_SCOPE_SPLIT
+
 	// SScrollBox defaults...
 	{
 		Style->Set( "ScrollBox", FScrollBoxStyle()
@@ -1382,6 +1418,8 @@ TSharedRef<ISlateStyle> FUMGCoreStyle::Create()
 		Style->Set("SyntaxHighlight.NodeAttributeValue", FTextBlockStyle(SmallMonospacedText).SetColorAndOpacity(FLinearColor(FColor(0xffb46100)))); // orange
 	}
 
+	UE_COMPILE_SCOPE_END
+
 	return Style;
 }
 
