[Viewport Toolbar] //UE5/Main - AudioQA Editor Win64 Test=UE.EditorAutomation(RunTest=Editor.Python.AudioQA) - test_create_open_close_delete_audio_assets completed with warnings

diff --git a/Engine/Source/Runtime/SlateCore/Private/Styling/SlateStyleSet.cpp b/Engine/Source/Runtime/SlateCore/Private/Styling/SlateStyleSet.cpp
--- a/Engine/Source/Runtime/SlateCore/Private/Styling/SlateStyleSet.cpp
+++ b/Engine/Source/Runtime/SlateCore/Private/Styling/SlateStyleSet.cpp
@@ -189,10 +189,11 @@ float FSlateStyleSet::GetFloat(const FName PropertyName, const ANSICHAR* Specifi
 		}
 	}
 
-	if (Result == nullptr) 
+	// In case the required property is not available, and no DefaultValue was provided, let's log a Warning
+	if (Result == nullptr && DefaultValue == FStyleDefaults::GetFloat())
 	{
 		const ISlateStyle* ReportingStyle = RequestingStyle != nullptr ? RequestingStyle : this;
-		ReportingStyle->LogMissingResource(EStyleMessageSeverity::Warning, FText::Format(NSLOCTEXT("SlateStyleSet", "UknownSlateFloat", "Unable to find float property '{0}' in style."), FText::FromName(PropertyName)), PropertyName);
+		ReportingStyle->LogMissingResource(EStyleMessageSeverity::Warning, FText::Format(NSLOCTEXT("SlateStyleSet", "UnknownSlateFloat", "Unable to find float property '{0}' in style."), FText::FromName(PropertyName)), PropertyName);
 	}
 
 	return Result ? *Result : DefaultValue;
@@ -210,10 +211,11 @@ FVector2D FSlateStyleSet::GetVector(const FName PropertyName, const ANSICHAR* Sp
 		}
 	}
 
-	if (Result == nullptr)
+	// In case the required property is not available, and no DefaultValue was provided, let's log a Warning
+	if (Result == nullptr && DefaultValue == FStyleDefaults::GetVector2D())
 	{
 		const ISlateStyle* ReportingStyle = RequestingStyle != nullptr ? RequestingStyle : this;
-		ReportingStyle->LogMissingResource(EStyleMessageSeverity::Warning, FText::Format(NSLOCTEXT("SlateStyleSet", "UknownSlateVector", "Unable to find vector property '{0}' in style."), FText::FromName(PropertyName)), PropertyName);
+		ReportingStyle->LogMissingResource(EStyleMessageSeverity::Warning, FText::Format(NSLOCTEXT("SlateStyleSet", "UnknownSlateVector", "Unable to find vector property '{0}' in style."), FText::FromName(PropertyName)), PropertyName);
 	}
 
 	return Result ? FVector2D(*Result) : DefaultValue;
@@ -232,10 +234,11 @@ const FLinearColor& FSlateStyleSet::GetColor(const FName PropertyName, const ANS
 		}
 	}
 
-	if ( Result == nullptr) 
+	// In case the required property is not available, and no DefaultValue was provided, let's log a Warning
+	if (Result == nullptr && DefaultValue == FStyleDefaults::GetColor())
 	{
 		const ISlateStyle* ReportingStyle = RequestingStyle != nullptr ? RequestingStyle : this;
-		ReportingStyle->LogMissingResource(EStyleMessageSeverity::Warning, FText::Format(NSLOCTEXT("SlateStyleSet", "UknownColor", "Unable to find Color '{0}'."), FText::FromName(LookupName)), LookupName);
+		ReportingStyle->LogMissingResource(EStyleMessageSeverity::Warning, FText::Format(NSLOCTEXT("SlateStyleSet", "UnknownColor", "Unable to find Color '{0}'."), FText::FromName(LookupName)), LookupName);
 	}
 
 	return Result ? *Result : DefaultValue;
@@ -260,10 +263,11 @@ const FSlateColor FSlateStyleSet::GetSlateColor(const FName PropertyName, const
 		}
 	}
 
-	if ( Result == nullptr && LinearResult == nullptr)
+	// In case the required property is not available, and no DefaultValue was provided, let's log a Warning
+	if (Result == nullptr && LinearResult == nullptr && DefaultValue == FStyleDefaults::GetSlateColor())
 	{
 		const ISlateStyle* ReportingStyle = RequestingStyle != nullptr ? RequestingStyle : this;
-		ReportingStyle->LogMissingResource(EStyleMessageSeverity::Warning, FText::Format(NSLOCTEXT("SlateStyleSet", "UknownSlateColor", "Unable to find SlateColor '{0}'."), FText::FromName(StyleName)), StyleName);
+		ReportingStyle->LogMissingResource(EStyleMessageSeverity::Warning, FText::Format(NSLOCTEXT("SlateStyleSet", "UnknownSlateColor", "Unable to find SlateColor '{0}'."), FText::FromName(StyleName)), StyleName);
 	}
 
 	return Result ? *Result : LinearResult ? *LinearResult : DefaultValue;
@@ -282,10 +286,11 @@ const FMargin& FSlateStyleSet::GetMargin(const FName PropertyName, const ANSICHA
 		}
 	}
 
-	if ( Result == nullptr )
+	// In case the required property is not available, and no DefaultValue was provided, let's log a Warning
+	if (Result == nullptr && DefaultValue == FStyleDefaults::GetMargin())
 	{
 		const ISlateStyle* ReportingStyle = RequestingStyle != nullptr ? RequestingStyle : this;
-		ReportingStyle->LogMissingResource(EStyleMessageSeverity::Warning, FText::Format(NSLOCTEXT("SlateStyleSet", "UknownMargin", "Unable to find Margin '{0}'."), FText::FromName(StyleName)), StyleName);
+		ReportingStyle->LogMissingResource(EStyleMessageSeverity::Warning, FText::Format(NSLOCTEXT("SlateStyleSet", "UnknownMargin", "Unable to find Margin '{0}'."), FText::FromName(StyleName)), StyleName);
 	}
 
 	return Result ? *Result : DefaultValue;
@@ -318,7 +323,7 @@ const FSlateBrush* FSlateStyleSet::GetBrush(const FName PropertyName, const ANSI
 	if ( Result == nullptr )
 	{
 		const ISlateStyle* ReportingStyle = RequestingStyle != nullptr ? RequestingStyle : this;
-		ReportingStyle->LogMissingResource(EStyleMessageSeverity::Warning, FText::Format(NSLOCTEXT("SlateStyleSet", "UknownBrush", "Unable to find Brush '{0}'."), FText::FromName(StyleName)), StyleName);
+		ReportingStyle->LogMissingResource(EStyleMessageSeverity::Warning, FText::Format(NSLOCTEXT("SlateStyleSet", "UnknownBrush", "Unable to find Brush '{0}'."), FText::FromName(StyleName)), StyleName);
 	}
 
 	return Result ? Result : GetDefaultBrush();
@@ -413,7 +418,7 @@ const FSlateSound& FSlateStyleSet::GetSound(const FName PropertyName, const ANSI
 	if ( Result == nullptr )
 	{
 		const ISlateStyle* ReportingStyle = RequestingStyle != nullptr ? RequestingStyle : this;
-		ReportingStyle->LogMissingResource(EStyleMessageSeverity::Warning, FText::Format(NSLOCTEXT("SlateStyleSet", "UknownSound", "Unable to find Sound '{0}'."), FText::FromName(StyleName)), StyleName);
+		ReportingStyle->LogMissingResource(EStyleMessageSeverity::Warning, FText::Format(NSLOCTEXT("SlateStyleSet", "UnknownSound", "Unable to find Sound '{0}'."), FText::FromName(StyleName)), StyleName);
 	}
 
 	return Result ? *Result : FStyleDefaults::GetSound();
