Fixes an edge case in font geometry preprocessing for distance field rendering, where preprocessed geometry was not recognized as modified if the same set of contours remained but unique vertices were added or removed. Adds a secondary check to filter out unnecessary preprocessed geometry, reducing memory footprint.
#tests Tested in editor and in a standalone sandbox

Updated font stats:
Preprocessed geometry for font Noto Sans SC is 57972 bytes, containing 86 paths and 5 windings out of 30907 total glyphs. (NotoSansSC-Black)
Preprocessed geometry for font Noto Sans SC is 134886 bytes, containing 206 paths and 4 windings out of 30907 total glyphs. (NotoSansSC-Regular)
Preprocessed geometry for font Noto Sans TC is 42273 bytes, containing 61 paths and 5 windings out of 20826 total glyphs. (NotoSansTC-Black)
Preprocessed geometry for font Noto Sans TC is 88382 bytes, containing 134 paths and 4 windings out of 20826 total glyphs. (NotoSansTC-Regular)
Preprocessed geometry for font Noto Sans KR is 19519 bytes, containing 41 paths and 8 windings out of 24737 total glyphs.
Preprocessed geometry for font Noto Sans JP is 13858 bytes, containing 23 paths and 4 windings out of 17850 total glyphs.
Preprocessed geometry for font Noto Sans is 19089 bytes, containing 60 paths and 27 windings out of 3246 total glyphs. (NotoSans-Regular)
Preprocessed geometry for font Noto Sans is 16605 bytes, containing 45 paths and 5 windings out of 3249 total glyphs. (NotoSans-Italic)
Preprocessed geometry for font Noto Sans is 19143 bytes, containing 60 paths and 27 windings out of 3246 total glyphs. (NotoSans-Bold)
Preprocessed geometry for font Noto Sans Arabic is 390138 bytes, containing 591 paths and 12 windings out of 1677 total glyphs. (NotoSansArabic-Regular)
Preprocessed geometry for font Noto Sans Devanagari is 32075 bytes, containing 73 paths and 0 windings out of 954 total glyphs.
Preprocessed geometry for font Tazugane for Epic Games is 6459 bytes, containing 6 paths and 0 windings out of 7734 total glyphs.
Preprocessed geometry for font NIS_JYAUN_P is 3133 bytes, containing 8 paths and 0 windings out of 5002 total glyphs.
Preprocessed geometry for font Burbank Big Condensed is 1 bytes, containing 0 paths and 0 windings out of 690 total glyphs.
Preprocessed geometry for font Heading Now is 30 bytes, containing 0 paths and 1 windings out of 1200 total glyphs. (HeadingNow_74Regular)
Preprocessed geometry for font RL Small is 1611 bytes, containing 6 paths and 13 windings out of 527 total glyphs. (RLSmall-Regular)
In conclusion, for all font faces the size increase is at most 2%, except Noto Sans Arabic (+100%) and Devanagari (+8%)

diff --git a/Engine/Source/Runtime/SlateCore/Private/Fonts/FontGeometryPreprocessing.cpp b/Engine/Source/Runtime/SlateCore/Private/Fonts/FontGeometryPreprocessing.cpp
--- a/Engine/Source/Runtime/SlateCore/Private/Fonts/FontGeometryPreprocessing.cpp
+++ b/Engine/Source/Runtime/SlateCore/Private/Fonts/FontGeometryPreprocessing.cpp
@@ -20,8 +20,72 @@ using namespace SlateMsdfgen;
 // Corners with an angle greater than 3 radians (~171 degrees) won't be treated as corners.
 static constexpr double SdfCornerAngleThreshold = 3.0;
 
+static void ScanlineShape(const msdfgen::Shape& Shape, double Y, TArray<msdfgen::Scanline::Intersection>& OutIntersections)
+{
+	double X[3];
+	int DY[3];
+	OutIntersections.Reset();
+	for (const msdfgen::Contour& Contour : Shape.contours)
+	{
+		for (const msdfgen::EdgeHolder& Edge : Contour.edges)
+		{
+			int N = Edge->scanlineIntersections(X, DY, Y);
+			for (int I = 0; I < N; ++I) {
+				msdfgen::Scanline::Intersection Intersection = { X[I], DY[I] };
+				OutIntersections.Add(Intersection);
+			}
+		}
+	}
+	OutIntersections.Sort([](const msdfgen::Scanline::Intersection& A, const msdfgen::Scanline::Intersection& B) { return A.x < B.x; });
+}
+
+static bool ScanlineEquivalence(const msdfgen::Shape& A, const msdfgen::Shape& B)
+{
+	constexpr double VertexYEpsilon = .0001;
+	TSet<double> VertY;
+	for (const msdfgen::Contour& Contour : A.contours)
+	{
+		for (const msdfgen::EdgeHolder& Edge : Contour.edges)
+		{
+			VertY.Add(Edge->controlPoints()[0].y);
+		}
+	}
+	for (const msdfgen::Contour& Contour : B.contours)
+	{
+		for (const msdfgen::EdgeHolder& Edge : Contour.edges)
+		{
+			VertY.Add(Edge->controlPoints()[0].y);
+		}
+	}
+	TArray<msdfgen::Scanline::Intersection> AIntersections, BIntersections;
+	for (double BaseY : VertY)
+	{
+		for (int Sign = -1; Sign <= 1; Sign += 2)
+		{
+			// Test scanline slightly above and below vertex Y.
+			const double Y = BaseY+Sign*VertexYEpsilon;
+			ScanlineShape(A, Y, AIntersections);
+			ScanlineShape(B, Y, BIntersections);
+			if (AIntersections.Num() != BIntersections.Num())
+			{
+				return false;
+			}
+			for (SIZE_T I = 0; I < AIntersections.Num(); ++I)
+			{
+				if (AIntersections[I].direction != BIntersections[I].direction)
+				{
+					return false;
+				}
+			}
+		}
+	}
+	return true;
+}
+
 bool PreprocessFontGeometry(FPreprocessedFontGeometry& OutPreprocessedFontGeometry, FT_Face InFreeTypeFace)
 {
+	int StatNumFullGeometry = 0;
+	int StatNumWindings = 0;
 	int NumCorrectWindings = 0;
 	int NumReverseWindings = 0;
 	/** Maps winding directions of each glyph in font:
@@ -69,6 +133,7 @@ bool PreprocessFontGeometry(FPreprocessedFontGeometry& OutPreprocessedFontGeomet
 			else // Inconsistent windings case - list of contours to reverse must be stored either way
 			{
 				bPreprocessed = true;
+				++StatNumWindings;
 			}
 		}
 
@@ -78,9 +143,23 @@ bool PreprocessFontGeometry(FPreprocessedFontGeometry& OutPreprocessedFontGeomet
 		FPreprocessedGlyphGeometry::FGeometryFingerprint ResolvedGeometry(MsdfgenShape);
 		if (ResolvedGeometry != OriginalGeometry)
 		{
-			msdfgen::edgeColoringInkTrap(MsdfgenShape, SdfCornerAngleThreshold);
-			PreprocessedGlyphGeometry = FPreprocessedGlyphGeometry(MsdfgenShape);
-			bPreprocessed = true;
+			/* In many cases, the resolved geometry is slightly different but not in a consequential way.
+			 * To save memory, we will do a secondary check by inspecting scanlines near vertices and compare them between resolved and unresolved shape.
+			 * If these scanlines are equivalent, it can be assumed the rasterization output will be the same.
+			 */
+			// Reconstruct unresolved shape
+			msdfgen::Shape UnresolvedMsdfgenShape;
+			msdfgen::readFreetypeOutline(UnresolvedMsdfgenShape, &InFreeTypeFace->glyph->outline, 1);
+			// Apply potential contour winding reversals from previous step
+			PreprocessedGlyphGeometry.View().UpdateWindings(UnresolvedMsdfgenShape);
+
+			if (!ScanlineEquivalence(MsdfgenShape, UnresolvedMsdfgenShape))
+			{
+				msdfgen::edgeColoringInkTrap(MsdfgenShape, SdfCornerAngleThreshold);
+				PreprocessedGlyphGeometry = FPreprocessedGlyphGeometry(MsdfgenShape);
+				bPreprocessed = true;
+				++StatNumFullGeometry;
+			}
 		}
 
 		if (bPreprocessed)
@@ -116,8 +195,12 @@ bool PreprocessFontGeometry(FPreprocessedFontGeometry& OutPreprocessedFontGeomet
 		if (GlyphWindingDirections[GlyphIndex] == AtypicalWinding)
 		{
 			OutPreprocessedFontGeometry.AddGlyph(GlyphIndex, GlyphWindingReversalView);
+			++StatNumWindings;
 		}
 	}
+
+	UE_LOG(LogSlate, Log, TEXT("Preprocessed geometry for font %hs is %d bytes, containing %d paths and %d windings out of %d total glyphs."), InFreeTypeFace->family_name, int(OutPreprocessedFontGeometry.GetDataSize()), StatNumFullGeometry, StatNumWindings, int(InFreeTypeFace->num_glyphs));
+
 	return true;
 }
 
diff --git a/Engine/Source/Runtime/SlateCore/Private/Fonts/PreprocessedFontGeometry.cpp b/Engine/Source/Runtime/SlateCore/Private/Fonts/PreprocessedFontGeometry.cpp
--- a/Engine/Source/Runtime/SlateCore/Private/Fonts/PreprocessedFontGeometry.cpp
+++ b/Engine/Source/Runtime/SlateCore/Private/Fonts/PreprocessedFontGeometry.cpp
@@ -287,10 +287,12 @@ bool FPreprocessedGlyphGeometry::FGeometryFingerprint::OneWayMatch(const FGeomet
 		int32 MatchedBContourIndex = -1;
 		for (const TPair<int32, int32>& AVertex : AContour)
 		{
+			bool bVertexMatched = false;
 			for (int32 BContourIndex = 0; BContourIndex < B.Contours.Num(); ++BContourIndex)
 			{
 				if (B.Contours[BContourIndex].Contains(AVertex))
 				{
+					bVertexMatched = true;
 					if (BContourIndex != MatchedBContourIndex)
 					{
 						if (MatchedBContourIndex < 0)
@@ -309,6 +311,10 @@ bool FPreprocessedGlyphGeometry::FGeometryFingerprint::OneWayMatch(const FGeomet
 					}
 				}
 			}
+			if (!bVertexMatched)
+			{
+				return false;
+			}
 		}
 	}
 	return true;
@@ -451,6 +457,11 @@ SIZE_T FPreprocessedFontGeometry::GetAllocatedSize() const
 	return Glyphs.GetAllocatedSize() + ContourData.GetAllocatedSize() + CoordinateData.GetAllocatedSize();
 }
 
+SIZE_T FPreprocessedFontGeometry::GetDataSize() const
+{
+	return sizeof(bool) + (sizeof(int32)+sizeof(FGlyphHeader)) * Glyphs.Num() + ContourData.NumBytes() + CoordinateData.NumBytes();
+}
+
 FArchive &operator<<(FArchive &Ar, FPreprocessedFontGeometry &FontGeometry)
 {
 	return Ar
diff --git a/Engine/Source/Runtime/SlateCore/Private/Fonts/PreprocessedFontGeometry.h b/Engine/Source/Runtime/SlateCore/Private/Fonts/PreprocessedFontGeometry.h
--- a/Engine/Source/Runtime/SlateCore/Private/Fonts/PreprocessedFontGeometry.h
+++ b/Engine/Source/Runtime/SlateCore/Private/Fonts/PreprocessedFontGeometry.h
@@ -225,6 +225,9 @@ public:
 	/** Returns the object's total allocated size */
 	SIZE_T GetAllocatedSize() const;
 
+	/** Returns the object's actual data size */
+	SIZE_T GetDataSize() const;
+
 	/** Serializes the font geometry object to/from an archive */
 	friend FArchive& operator<<(FArchive& Ar, FPreprocessedFontGeometry& FontGeometry);
 
