[Slate] Makes it possible to toggle distance field text rendering using CVar at runtime
Renames CVar SlateSdfText.Enable to UI.SlateSDFText.Enable (the old name is supported as an alias)

#tests Tested in editor & PIE
#rnx <Release Notes exclude>

diff --git a/Engine/Source/Runtime/Engine/Private/FontFace.cpp b/Engine/Source/Runtime/Engine/Private/FontFace.cpp
--- a/Engine/Source/Runtime/Engine/Private/FontFace.cpp
+++ b/Engine/Source/Runtime/Engine/Private/FontFace.cpp
@@ -308,7 +308,7 @@ FString UFontFace::GetCookedFilename() const
 void UFontFace::UpdateDeviceRasterizationSettings()
 {
 	DeviceRasterizationSettings = FFontRasterizationSettings();
-	if (bEnableDistanceFieldRendering && GetDefault<UUserInterfaceSettings>()->bEnableDistanceFieldFontRasterization && IsSlateSdfTextFeatureEnabled())
+	if (bEnableDistanceFieldRendering && GetDefault<UUserInterfaceSettings>()->bEnableDistanceFieldFontRasterization)
 	{
 		DeviceRasterizationSettings.Mode = GetDeviceFontFaceDistanceFieldRasterizationMode();
 		if (PlatformRasterizationModeOverrides.IsSet())
diff --git a/Engine/Source/Runtime/SlateCore/Private/Fonts/CompositeFont.cpp b/Engine/Source/Runtime/SlateCore/Private/Fonts/CompositeFont.cpp
--- a/Engine/Source/Runtime/SlateCore/Private/Fonts/CompositeFont.cpp
+++ b/Engine/Source/Runtime/SlateCore/Private/Fonts/CompositeFont.cpp
@@ -1,6 +1,7 @@
 // Copyright Epic Games, Inc. All Rights Reserved.
 
 #include "Fonts/CompositeFont.h"
+#include "Fonts/FontCache.h"
 #include "SlateGlobals.h"
 #include "UObject/EditorObjectVersion.h"
 #include "UObject/FortniteMainBranchObjectVersion.h"
@@ -258,7 +259,7 @@ const UObject* FFontData::GetFontFaceAsset() const
 
 FFontRasterizationSettings FFontData::GetFontRasterizationSettings() const
 {
-	if (FontFaceAsset)
+	if (FontFaceAsset && IsSlateSdfTextFeatureEnabled())
 	{
 		const IFontFaceInterface* FontFace = CastChecked<const IFontFaceInterface>(FontFaceAsset);
 		return FontFace->GetRasterizationSettings();
diff --git a/Engine/Source/Runtime/SlateCore/Private/Fonts/FontCache.cpp b/Engine/Source/Runtime/SlateCore/Private/Fonts/FontCache.cpp
--- a/Engine/Source/Runtime/SlateCore/Private/Fonts/FontCache.cpp
+++ b/Engine/Source/Runtime/SlateCore/Private/Fonts/FontCache.cpp
@@ -73,11 +73,33 @@ FAutoConsoleVariableRef CVarFlushFontCache(
 	TEXT("Flush the font cache."));
 #endif
 
-static TAutoConsoleVariable<bool> CVarSlateSdfTextEnable(TEXT("SlateSdfText.Enable"), true, TEXT("Enables MSDF-based text rendering in Slate"));
+static bool bSlateSdfTextEnable = true;
+
+void OnSlateSdfTextEnablementChanged(IConsoleVariable* Var);
+
+FAutoConsoleVariableRef CVarSlateSdfTextEnable(
+	TEXT("UI.SlateSDFText.Enable"),
+	bSlateSdfTextEnable,
+	TEXT("Enables signed distance field-based text rendering in Slate"),
+	FConsoleVariableDelegate::CreateStatic(OnSlateSdfTextEnablementChanged));
+
+
+FAutoConsoleVariableShadow ShadowCVarSlateSdfTextEnable(TEXT("SlateSdfText.Enable"), TEXT("UI.SlateSDFText.Enable"));
 
 bool IsSlateSdfTextFeatureEnabled()
 {
-	return CVarSlateSdfTextEnable.GetValueOnAnyThread();
+	return bSlateSdfTextEnable;
+}
+
+void OnSlateSdfTextEnablementChanged(IConsoleVariable*)
+{
+	if (FSlateApplicationBase::IsInitialized())
+	{
+		if (const FSlateRenderer* Renderer = FSlateApplicationBase::Get().GetRenderer())
+		{
+			Renderer->GetFontCache()->RequestFlushCache(TEXT("Value of UI.SlateSDFText.Enable changed."));
+		}
+	}
 }
 
 ETextShapingMethod GetDefaultTextShapingMethod()
