Extended toolbar style for combo button content: min, max size and horizontal alignment

diff --git a/Engine/Source/Runtime/Slate/Private/Framework/MultiBox/SToolBarComboButtonBlock.cpp b/Engine/Source/Runtime/Slate/Private/Framework/MultiBox/SToolBarComboButtonBlock.cpp
--- a/Engine/Source/Runtime/Slate/Private/Framework/MultiBox/SToolBarComboButtonBlock.cpp
+++ b/Engine/Source/Runtime/Slate/Private/Framework/MultiBox/SToolBarComboButtonBlock.cpp
@@ -166,12 +166,15 @@ void SToolBarComboButtonBlock::BuildMultiBlockWidget(const ISlateStyle* StyleSet
 			// Label text
 			+ SHorizontalBox::Slot()
 			.AutoWidth()
+			.MinWidth(ToolBarStyle.ComboContentMinWidth)
+			.MaxWidth(ToolBarStyle.ComboContentMaxWidth)
 			.Padding(TAttribute<FMargin>::CreateLambda(
 				[Label, ToolBarStyle]() -> FMargin
 				{
 					return Label.Get().IsEmpty() ? FMargin(0) : ToolBarStyle.LabelPadding;
 				}
 			))
+			.HAlign(ToolBarStyle.ComboContentHorizontalAlignment)
 			.VAlign(VAlign_Center)	// Center the label text horizontally
 			[
 				SNew(STextBlock)
diff --git a/Engine/Source/Runtime/SlateCore/Private/Styling/ToolBarStyle.cpp b/Engine/Source/Runtime/SlateCore/Private/Styling/ToolBarStyle.cpp
--- a/Engine/Source/Runtime/SlateCore/Private/Styling/ToolBarStyle.cpp
+++ b/Engine/Source/Runtime/SlateCore/Private/Styling/ToolBarStyle.cpp
@@ -33,6 +33,9 @@ FToolBarStyle::FToolBarStyle()
 	, WrapButtonIndex(-1)
 	, IconSize(16, 16)
 	, bShowLabels(true)
+	, ComboContentMinWidth(0)
+	, ComboContentMaxWidth(0)
+	, ComboContentHorizontalAlignment(HAlign_Fill)
 {}
 
 FToolBarStyle::FToolBarStyle(const FToolBarStyle&) = default;
diff --git a/Engine/Source/Runtime/SlateCore/Public/Styling/ToolBarStyle.h b/Engine/Source/Runtime/SlateCore/Public/Styling/ToolBarStyle.h
--- a/Engine/Source/Runtime/SlateCore/Public/Styling/ToolBarStyle.h
+++ b/Engine/Source/Runtime/SlateCore/Public/Styling/ToolBarStyle.h
@@ -162,6 +162,20 @@ struct FToolBarStyle : public FSlateWidgetStyle
 	float ButtonContentFillWidth = 1.0f;
 	FToolBarStyle& SetButtonContentFillWidth(float InButtonContentFillWidth) { ButtonContentFillWidth = InButtonContentFillWidth; return *this; }
 
+	/** Min width that label text block slot in combo buttons should have. */
+	UPROPERTY(EditAnywhere, BlueprintReadWrite, Category = Appearance)
+	float ComboContentMinWidth;
+	FToolBarStyle& SetComboLabelMinWidth(float InComboContentMinWidth) { ComboContentMinWidth = InComboContentMinWidth; return *this; }
+	
+	/** Max width that label text block slot in combo buttons should have. 0 means no max. */
+	UPROPERTY(EditAnywhere, BlueprintReadWrite, Category = Appearance)
+	float ComboContentMaxWidth;
+	FToolBarStyle& SetComboLabelMaxWidth(float InComboContentMaxWidth) { ComboContentMaxWidth = InComboContentMaxWidth; return *this; }
+	
+	UPROPERTY(EditAnywhere, BlueprintReadWrite, Category = Appearance)
+	TEnumAsByte<EHorizontalAlignment> ComboContentHorizontalAlignment;
+	FToolBarStyle& SetComboContentHorizontalAlignment(EHorizontalAlignment InAlignment) { ComboContentHorizontalAlignment = InAlignment; return *this; }
+
 	UPROPERTY(EditAnywhere, BlueprintReadWrite, Category = Appearance)
 	FMargin IconPaddingWithVisibleLabel;
 	FToolBarStyle& SetIconPaddingWithVisibleLabel(const FMargin& InIconPaddingWithVisibleLabel) { IconPaddingWithVisibleLabel = InIconPaddingWithVisibleLabel; return *this; }
