[Slate] Adds debug visualization of text rendered with different techniques (MSDF, SDF, bitmap) - CVar UI.SlateSDFText.DebugVisual

#tests Tested in editor & PIE
#rnx <Release Notes exclude>

diff --git a/Engine/Source/Runtime/SlateCore/Private/Rendering/ElementBatcher.cpp b/Engine/Source/Runtime/SlateCore/Private/Rendering/ElementBatcher.cpp
--- a/Engine/Source/Runtime/SlateCore/Private/Rendering/ElementBatcher.cpp
+++ b/Engine/Source/Runtime/SlateCore/Private/Rendering/ElementBatcher.cpp
@@ -16,6 +16,8 @@
 #include "Widgets/SWidgetUtils.h"
 #include "Rendering/DrawElementPayloads.h"
 
+#define ENABLE_SDF_DEBUG_VISUAL (!UE_BUILD_SHIPPING)
+
 CSV_DECLARE_CATEGORY_MODULE_EXTERN(SLATECORE_API, Slate);
 
 DECLARE_CYCLE_STAT(TEXT("Slate RT: Create Batches"), STAT_SlateRTCreateBatches, STATGROUP_Slate);
@@ -33,6 +35,27 @@ DEFINE_STAT(STAT_SlateInvalidation_RecachedElements);
 
 int32 GSlateFeathering = 0;
 
+#if ENABLE_SDF_DEBUG_VISUAL
+enum class ESdfTextDebugVisual
+{
+	Disabled = 0,
+	Raw = 1,
+	Color = 2,
+	Boxes = 3
+};
+
+static int32 SdfTextDebugVisual = static_cast<int32>(ESdfTextDebugVisual::Disabled);
+
+FAutoConsoleVariableRef CVarSlateSdfTextDebugVisual(
+	TEXT("UI.SlateSDFText.DebugVisual"),
+	SdfTextDebugVisual,
+	TEXT("Slate SDF text debug visuals - 0: Disabled, 1: Raw SDF, 2: Text color, 3: Colored boxes."));
+
+static const FColor DebugVisualBitmapTextColor(0x80, 0x80, 0x80); // Gray
+static const FColor DebugVisualMsdfTextColor(0xff, 0x00, 0xff); // Magenta
+static const FColor DebugVisualSdfTextColor(0x00, 0xff, 0xff); // Cyan
+static const FColor DebugVisualApproximateSdfTextColor(0xff, 0x80, 0x00); // Orange
+#endif // ENABLE_SDF_DEBUG_VISUAL
 
 FSlateElementBatch::FSlateElementBatch(const FSlateShaderResource* InShaderResource, const FShaderParams& InShaderParams, ESlateShader ShaderType, ESlateDrawPrimitive PrimitiveType, ESlateDrawEffect InDrawEffects, ESlateBatchDrawFlag InBatchFlags, const FSlateDrawElement& InDrawElement, int32 InstanceCount, uint32 InstanceOffset, ISlateUpdatableInstanceBuffer* InstanceData)
 	: BatchKey(InShaderParams, ShaderType, PrimitiveType, InDrawEffects, InBatchFlags, InDrawElement.GetClippingHandle(), InstanceCount, InstanceOffset, InstanceData, InDrawElement.GetSceneIndex())
@@ -3284,8 +3307,10 @@ int32 FSlateElementBatcher::BuildShapedTextSequence(const FShapedTextBuildContex
 
 	FSlateRenderTransform RenderTransform = *Context.RenderTransform;
 
+	FColor FontTint = Context.FontTint;
 	FColor Tint = FColor::White;
 
+	const UObject* FontMaterial = Context.FontMaterial;
 	ETextOverflowDirection OverflowDirection = Context.OverflowDirection;
 	ETextOverflowPolicy OverflowPolicy = Context.OverflowPolicy;
 
@@ -3320,6 +3345,23 @@ int32 FSlateElementBatcher::BuildShapedTextSequence(const FShapedTextBuildContex
 	// True by default, will be checked only after the middle skip and will become false as soon as we find a Visible character after the ellipsis, used to skip whitespaces after the ellipsis
 	bool bIsWhitespaceAfterMiddleSkip = true;
 
+#if ENABLE_SDF_DEBUG_VISUAL
+	TArray<TPair<FMargin, FColor> > DebugVisualBoxes;
+
+	switch (SdfTextDebugVisual)
+	{
+		case static_cast<int32>(ESdfTextDebugVisual::Disabled):
+		default:
+			// No effect
+			break;
+		case static_cast<int32>(ESdfTextDebugVisual::Raw):
+		case static_cast<int32>(ESdfTextDebugVisual::Color):
+			FontTint = FColor::White;
+			FontMaterial = nullptr;
+			break;
+	}
+#endif // ENABLE_SDF_DEBUG_VISUAL
+
 	// Calculate the Glyph Indexes to skip to fit the ellipsis and add eventually the offset to the LineX to align the text correctly for the RightToLeft direction
 	if (OverflowPolicy == ETextOverflowPolicy::MiddleEllipsis && bWillBeClipped)
 	{
@@ -3456,14 +3498,13 @@ int32 FSlateElementBatcher::BuildShapedTextSequence(const FShapedTextBuildContex
 					FontAtlasTexture = SlateFontTexture->GetSlateTexture();
 					check(FontAtlasTexture);
 
-					FontShaderResource = ResourceManager.GetFontShaderResource(FontTextureIndex, FontAtlasTexture, Context.FontMaterial);
+					FontShaderResource = ResourceManager.GetFontShaderResource(FontTextureIndex, FontAtlasTexture, FontMaterial);
 					check(FontShaderResource);
 
 					InvTextureSizeX = 1.0f / FontAtlasTexture->GetWidth();
 					InvTextureSizeY = 1.0f / FontAtlasTexture->GetHeight();
 
 					const ESlateFontAtlasContentType ContentType = SlateFontTexture->GetContentType();
-					Tint = ContentType == ESlateFontAtlasContentType::Color ? FColor::White : Context.FontTint;
 
 					ESlateShader ShaderType = ESlateShader::Default;
 					switch (ContentType)
@@ -3487,6 +3528,41 @@ int32 FSlateElementBatcher::BuildShapedTextSequence(const FShapedTextBuildContex
 					}
 					check(ShaderType != ESlateShader::Default);
 
+#if ENABLE_SDF_DEBUG_VISUAL
+					switch (SdfTextDebugVisual)
+					{
+						case static_cast<int32>(ESdfTextDebugVisual::Disabled):
+						default:
+							break;
+						case static_cast<int32>(ESdfTextDebugVisual::Raw):
+							if (bIsSdfGlyph)
+							{
+								ShaderType = ContentType == ESlateFontAtlasContentType::Alpha ? ESlateShader::GrayscaleFont : ESlateShader::ColorFont;
+							}
+							break;
+						case static_cast<int32>(ESdfTextDebugVisual::Color):
+							switch (GlyphToRender.FontFaceData->RasterizationMode)
+							{
+								case EFontRasterizationMode::Bitmap:
+								default:
+									FontTint = DebugVisualBitmapTextColor;
+									break;
+								case EFontRasterizationMode::Msdf:
+									FontTint = DebugVisualMsdfTextColor;
+									break;
+								case EFontRasterizationMode::Sdf:
+									FontTint = DebugVisualSdfTextColor;
+									break;
+								case EFontRasterizationMode::SdfApproximation:
+									FontTint = DebugVisualApproximateSdfTextColor;
+									break;
+							}
+							break;
+					}
+#endif // ENABLE_SDF_DEBUG_VISUAL
+
+					Tint = ContentType == ESlateFontAtlasContentType::Color ? FColor::White : FontTint;
+
 					FShaderParams ShaderParams;
 					if (bIsSdfGlyph)
 					{
@@ -3687,7 +3763,7 @@ int32 FSlateElementBatcher::BuildShapedTextSequence(const FShapedTextBuildContex
 			const uint32 IndexStart = RenderBatch->GetNumVertices();
 
 			float Ut = 0.0f, Vt = 0.0f, UtMax = 0.0f, VtMax = 0.0f;
-			if (Context.FontMaterial)
+			if (FontMaterial)
 			{
 				float DistAlpha = (float)GlyphIndex / NumGlyphs;
 				float DistAlphaNext = (float)(GlyphIndex + 1) / NumGlyphs;
@@ -3743,6 +3819,31 @@ int32 FSlateElementBatcher::BuildShapedTextSequence(const FShapedTextBuildContex
 
 			// Reset whitespace advance to 0, this is a visible character
 			PreviousWhitespaceAdvance = 0;
+
+#if ENABLE_SDF_DEBUG_VISUAL
+			if (SdfTextDebugVisual == static_cast<int32>(ESdfTextDebugVisual::Boxes))
+			{
+				FColor BoxColor;
+				switch (GlyphToRender.FontFaceData->RasterizationMode)
+				{
+					case EFontRasterizationMode::Bitmap:
+					default:
+						BoxColor = DebugVisualBitmapTextColor;
+						BoxColor.A = 0x60;
+						break;
+					case EFontRasterizationMode::Msdf:
+						BoxColor = DebugVisualMsdfTextColor;
+						break;
+					case EFontRasterizationMode::Sdf:
+						BoxColor = DebugVisualSdfTextColor;
+						break;
+					case EFontRasterizationMode::SdfApproximation:
+						BoxColor = DebugVisualApproximateSdfTextColor;
+						break;
+				}
+				DebugVisualBoxes.Add(TPair<FMargin, FColor>(FMargin(UpperLeft.X, UpperLeft.Y, LowerRight.X, LowerRight.Y), BoxColor));
+			}
+#endif // ENABLE_SDF_DEBUG_VISUAL
 		}
 		else if (!bIsVisible)
 		{
@@ -3801,6 +3902,44 @@ int32 FSlateElementBatcher::BuildShapedTextSequence(const FShapedTextBuildContex
 
 		BuildShapedTextSequence<Rounding>(EllipsisContext);
 	}
+	
+#if ENABLE_SDF_DEBUG_VISUAL
+	if (SdfTextDebugVisual == static_cast<int32>(ESdfTextDebugVisual::Boxes) && !DebugVisualBoxes.IsEmpty())
+	{
+		RenderBatch = &CreateRenderBatch(Context.LayerId,
+			FShaderParams(),
+			nullptr,
+			ESlateDrawPrimitive::LineList,
+			ESlateShader::Default,
+			Context.DrawElement->GetDrawEffects(),
+			ESlateBatchDrawFlag::None,
+			*Context.DrawElement);
+		RenderBatch->ReserveVertices(6*DebugVisualBoxes.Num());
+		RenderBatch->ReserveIndices(8*DebugVisualBoxes.Num());
+		for (const TPair<FMargin, FColor>& Box : DebugVisualBoxes)
+		{
+			FColor TransparentColor = Box.Value;
+			TransparentColor.A = 0;
+			RenderBatch->AddVertex(FSlateVertex::Make<Rounding>(RenderTransform, FVector2f(Box.Key.Left, .5f*(Box.Key.Top + Box.Key.Bottom)), FVector4f(), FVector2f(0.0f, 0.0f), TransparentColor));
+			RenderBatch->AddVertex(FSlateVertex::Make<Rounding>(RenderTransform, FVector2f(Box.Key.Left, Box.Key.Bottom), FVector4f(), FVector2f(0.0f, 0.0f), Box.Value));
+			RenderBatch->AddVertex(FSlateVertex::Make<Rounding>(RenderTransform, FVector2f(Box.Key.Right, Box.Key.Bottom), FVector4f(), FVector2f(0.0f, 0.0f), TransparentColor));
+			RenderBatch->AddVertex(FSlateVertex::Make<Rounding>(RenderTransform, FVector2f(Box.Key.Right, .5f*(Box.Key.Top + Box.Key.Bottom)), FVector4f(), FVector2f(0.0f, 0.0f), TransparentColor));
+			RenderBatch->AddVertex(FSlateVertex::Make<Rounding>(RenderTransform, FVector2f(Box.Key.Right, Box.Key.Top), FVector4f(), FVector2f(0.0f, 0.0f), Box.Value));
+			RenderBatch->AddVertex(FSlateVertex::Make<Rounding>(RenderTransform, FVector2f(Box.Key.Left, Box.Key.Top), FVector4f(), FVector2f(0.0f, 0.0f), TransparentColor));
+		}
+		for (int32 I = 0; I < DebugVisualBoxes.Num(); ++I)
+		{
+			RenderBatch->AddIndex(static_cast<SlateIndex>(6*I+0));
+			RenderBatch->AddIndex(static_cast<SlateIndex>(6*I+1));
+			RenderBatch->AddIndex(static_cast<SlateIndex>(6*I+1));
+			RenderBatch->AddIndex(static_cast<SlateIndex>(6*I+2));
+			RenderBatch->AddIndex(static_cast<SlateIndex>(6*I+3));
+			RenderBatch->AddIndex(static_cast<SlateIndex>(6*I+4));
+			RenderBatch->AddIndex(static_cast<SlateIndex>(6*I+4));
+			RenderBatch->AddIndex(static_cast<SlateIndex>(6*I+5));
+		}
+	}
+#endif // ENABLE_SDF_DEBUG_VISUAL
 
 	return GlyphsRendered;
 }
