ColorPicker: New eye dropper icon
#rnx

diff --git a/Engine/Source/Runtime/AppFramework/Private/Widgets/Colors/SEyeDropperButton.cpp b/Engine/Source/Runtime/AppFramework/Private/Widgets/Colors/SEyeDropperButton.cpp
--- a/Engine/Source/Runtime/AppFramework/Private/Widgets/Colors/SEyeDropperButton.cpp
+++ b/Engine/Source/Runtime/AppFramework/Private/Widgets/Colors/SEyeDropperButton.cpp
@@ -28,7 +28,7 @@ void SEyeDropperButton::Construct(const FArguments& InArgs)
 			+ SOverlay::Slot()
 			[
 				SNew(SImage)
-				.Image(FCoreStyle::Get().GetBrush("ColorPicker.EyeDropper"))
+				.Image(FCoreStyle::Get().GetBrush("ColorPicker.EyeDropperLarge"))
 				.ToolTipText(LOCTEXT("EyeDropperButton_ToolTip", "Activates the eye-dropper for selecting a colored pixel from any window."))
 				.ColorAndOpacity(this, &SEyeDropperButton::GetDropperImageColor)
 			]
diff --git a/Engine/Source/Runtime/SlateCore/Private/Styling/StarshipCoreStyle.cpp b/Engine/Source/Runtime/SlateCore/Private/Styling/StarshipCoreStyle.cpp
--- a/Engine/Source/Runtime/SlateCore/Private/Styling/StarshipCoreStyle.cpp
+++ b/Engine/Source/Runtime/SlateCore/Private/Styling/StarshipCoreStyle.cpp
@@ -1970,6 +1970,7 @@ void FStarshipCoreStyle::SetupColorPickerStyles(TSharedRef<FStyle>& Style)
 
 		Style->Set("ColorPicker.AlphaBackground", new IMAGE_BRUSH("Starship/Common/Checker", Icon16x16, FLinearColor::White, ESlateBrushTileType::Both));
 		Style->Set("ColorPicker.EyeDropper", new IMAGE_BRUSH("Icons/eyedropper_16px", Icon16x16));
+		Style->Set("ColorPicker.EyeDropperLarge", new IMAGE_BRUSH_SVG("Starship/Common/EyeDropper_20", Icon20x20));
 		Style->Set("ColorPicker.Mode", new IMAGE_BRUSH("Common/ColorPicker_Mode_16x", Icon16x16));
 		Style->Set("ColorPicker.ModeWheel", new IMAGE_BRUSH("Common/ColorWheel_20x", Icon20x20));
 		Style->Set("ColorPicker.ModeSpectrum", new IMAGE_BRUSH("Common/ColorSpectrum_20x", Icon20x20));
