[Slate] Added TResizingWidgetSlotMixin, for re-use of size related slot properties

diff --git a/Engine/Source/Runtime/Slate/Private/Widgets/Layout/SScrollBox.cpp b/Engine/Source/Runtime/Slate/Private/Widgets/Layout/SScrollBox.cpp
--- a/Engine/Source/Runtime/Slate/Private/Widgets/Layout/SScrollBox.cpp
+++ b/Engine/Source/Runtime/Slate/Private/Widgets/Layout/SScrollBox.cpp
@@ -12,31 +12,13 @@
 void SScrollBox::FSlot::Construct(const FChildren& SlotOwner, FSlotArguments&& InArgs)
 {
 	TBasicLayoutWidgetSlot<FSlot>::Construct(SlotOwner, MoveTemp(InArgs));
-	if (InArgs._MinSize.IsSet())
-	{
-		SetMinSize(MoveTemp(InArgs._MinSize));
-	}
-	if (InArgs._MaxSize.IsSet())
-	{
-		SetMaxSize(MoveTemp(InArgs._MaxSize));
-	}
-	if (InArgs._SizeParam.IsSet())
-	{
-		SetSizeParam(MoveTemp(InArgs._SizeParam.GetValue()));
-	}
+	TResizingWidgetSlotMixin<FSlot>::ConstructMixin(SlotOwner, MoveTemp(InArgs));
 }
 
 void SScrollBox::FSlot::RegisterAttributes(FSlateWidgetSlotAttributeInitializer& AttributeInitializer)
 {
 	TBasicLayoutWidgetSlot<FSlot>::RegisterAttributes(AttributeInitializer);
-	SLATE_ADD_SLOT_ATTRIBUTE_DEFINITION_WITH_NAME(FSlot, AttributeInitializer, "Slot.MinSize", MinSize, EInvalidateWidgetReason::Layout);
-	SLATE_ADD_SLOT_ATTRIBUTE_DEFINITION_WITH_NAME(FSlot, AttributeInitializer, "Slot.MaxSize", MaxSize, EInvalidateWidgetReason::Layout);
-	SLATE_ADD_SLOT_ATTRIBUTE_DEFINITION_WITH_NAME(FSlot, AttributeInitializer, "Slot.SizeValue", SizeValue, EInvalidateWidgetReason::Layout)
-		.UpdatePrerequisite("Slot.MinSize")
-		.UpdatePrerequisite("Slot.MaxSize");
-	SLATE_ADD_SLOT_ATTRIBUTE_DEFINITION_WITH_NAME(FSlot, AttributeInitializer, "Slot.ShrinkSizeValue", ShrinkSizeValue, EInvalidateWidgetReason::Layout)
-		.UpdatePrerequisite("Slot.MinSize")
-		.UpdatePrerequisite("Slot.MaxSize");
+	TResizingWidgetSlotMixin<FSlot>::RegisterAttributesMixin(AttributeInitializer);
 }
 
 SScrollBox::FSlot::FSlotArguments SScrollBox::Slot()
diff --git a/Engine/Source/Runtime/Slate/Public/Widgets/Layout/SScrollBox.h b/Engine/Source/Runtime/Slate/Public/Widgets/Layout/SScrollBox.h
--- a/Engine/Source/Runtime/Slate/Public/Widgets/Layout/SScrollBox.h
+++ b/Engine/Source/Runtime/Slate/Public/Widgets/Layout/SScrollBox.h
@@ -72,14 +72,10 @@ class SScrollBox : public SCompoundWidget
 {
 public:
 	/** A Slot that provides layout options for the contents of a scrollable box. */
-	class FSlot : public TBasicLayoutWidgetSlot<FSlot>
+	class FSlot : public TBasicLayoutWidgetSlot<FSlot>, public TResizingWidgetSlotMixin<FSlot>
 	{
 	public:
-		SLATE_SLOT_BEGIN_ARGS(FSlot, TBasicLayoutWidgetSlot<FSlot>)
-			SLATE_ARGUMENT(TOptional<FSizeParam>, SizeParam)
-			TAttribute<float> _MinSize;
-			TAttribute<float> _MaxSize;
-
+		SLATE_SLOT_BEGIN_ARGS_OneMixin(FSlot, TBasicLayoutWidgetSlot<FSlot>, TResizingWidgetSlotMixin<FSlot>)
 			/** The widget's DesiredSize will be used as the space required. */
 			FSlotArguments& AutoSize()
 			{
@@ -93,12 +89,12 @@ public:
 				_SizeParam = FStretch(MoveTemp(InStretchCoefficient));
 				return Me();
 			}
-			
+
 			/**
 			 * The widget's content size is adjusted proportionally to fit the available space.
 			 * The slots size starts at DesiredSize, and a slot with coefficient of 2 will get adjusted twice as much as slot with coefficient 1 to fit the available space.
 			 * @param InStretchCoefficient Stretch coefficient for this slot.
-			 * @param InShrinkStretchCoefficient If specified, this stretch coefficient is used when the slots is shrinking below desired size. Otherwise InStretchCoefficient is used for both shrink and grow.  
+			 * @param InShrinkStretchCoefficient If specified, this stretch coefficient is used when the slots is shrinking below desired size. Otherwise InStretchCoefficient is used for both shrink and grow.
 			 */
 			FSlotArguments& FillContentSize(TAttribute<float> InStretchCoefficient, TAttribute<float> InShrinkStretchCoefficient = TAttribute<float>())
 			{
@@ -124,109 +120,13 @@ public:
 		/** Default values for a slot. */
 		FSlot()
 			: TBasicLayoutWidgetSlot<FSlot>(HAlign_Fill, VAlign_Fill)
-			, SizeRule(FSizeParam::SizeRule_Auto)
-			, SizeValue(*this, 1.f)
-			, ShrinkSizeValue(*this, 1.f)
-			, MinSize(*this, 0.0f)
-			, MaxSize(*this, 0.0f)
-		{ }
-
-		SLATE_API void Construct(const FChildren& SlotOwner, FSlotArguments&& InArgs);
-		static SLATE_API void RegisterAttributes(FSlateWidgetSlotAttributeInitializer& AttributeInitializer);
-
-		/** Get the space rule this slot should occupy along scrollbox's direction. */
-		FSizeParam::ESizeRule GetSizeRule() const
-		{
-			return SizeRule;
-		}
-
-		/** Get the space rule value this slot should occupy along scrollbox's direction. */
-		float GetSizeValue() const
-		{
-			return SizeValue.Get();
-		}
-
-		/**
-		 * Get the size parameter for the space rule, used when the slot size is shrinking below desired size.
-		 * Used for size rule SizeRule_StretchContent.
-		 */
-		float GetShrinkSizeValue() const
+			, TResizingWidgetSlotMixin<SScrollBox::FSlot>()
 		{
-			return ShrinkSizeValue.Get();
+			this->SizeRule = FSizeParam::SizeRule_Auto;
 		}
 
-		/** Get the min size the slot can be.*/
-		float GetMinSize() const
-		{
-			return MinSize.Get();
-		}
-		
-		/** Get the max size the slot can be.*/
-		float GetMaxSize() const
-		{
-			return MaxSize.Get();
-		}
-
-		/** Set the size Param of the slot, It could be a FStretch, FStretchContent, or a FAuto. */
-		void SetSizeParam(FSizeParam InSizeParam)
-		{
-			SizeRule = InSizeParam.SizeRule;
-
-			// ShrinkSizeValue is only used for StretchContent.
-			// If ShrinkValue is not set, make it equal to the Value. 
-			if (SizeRule == FSizeParam::SizeRule_StretchContent)
-			{
-				if (InSizeParam.ShrinkValue.IsSet())
-				{
-					ShrinkSizeValue.Assign(*this, MoveTemp(InSizeParam.ShrinkValue));
-				}
-				else
-				{
-					ShrinkSizeValue.Assign(*this, InSizeParam.Value); // Make copy, let SizeValue use move.
-				}
-			}
-			else
-			{
-				ShrinkSizeValue.Set(*this, 1.f); // Reset
-			}
-
-			SizeValue.Assign(*this, MoveTemp(InSizeParam.Value));
-		}
-
-		/** The widget's DesiredSize will be used as the space required. */
-		void SetSizeToAuto()
-		{
-			SetSizeParam(FAuto());
-		}
-
-		/** The available space will be distributed proportionately. */
-		void SetSizeToStretch(TAttribute<float> StretchCoefficient)
-		{
-			SetSizeParam(FStretch(MoveTemp(StretchCoefficient)));
-		}
-
-		/**
-		 * The widget's content size is adjusted proportionally to fit the available space.
-		 * The slots size starts at DesiredSize, and a slot with coefficient of 2 will get adjusted twice as much as slot with coefficient 1 to fit the available space.
-		 * @param InStretchCoefficient Stretch coefficient for this slot.
-		 * @param InShrinkStretchCoefficient If specified, this stretch coefficient is used when the slots is shrinking below desired size. Otherwise InStretchCoefficient is used for both shrink and grow.  
-		 */
-		void SetSizeToStretchContent(TAttribute<float> InStretchCoefficient, TAttribute<float> InShrinkStretchCoefficient = TAttribute<float>())
-		{
-			SetSizeParam(FStretchContent(MoveTemp(InStretchCoefficient), MoveTemp(InShrinkStretchCoefficient)));
-		}
-
-		/** Set the min size in SlateUnit this slot can be. */
-		void SetMinSize(TAttribute<float> InMinSize)
-		{
-			MinSize.Assign(*this, MoveTemp(InMinSize));
-		}
-
-		/** Set the max size in SlateUnit this slot can be. */
-		void SetMaxSize(TAttribute<float> InMaxSize)
-		{
-			MaxSize.Assign(*this, MoveTemp(InMaxSize));
-		}
+		SLATE_API void Construct(const FChildren& SlotOwner, FSlotArguments&& InArgs);
+		static SLATE_API void RegisterAttributes(FSlateWidgetSlotAttributeInitializer& AttributeInitializer);
 
 	private:
 		/**
@@ -237,24 +137,8 @@ public:
 		 * peers' SizeRule_Auto requirements have been satisfied.
 		 */
 
-		 /** The sizing rule to use. */
-		FSizeParam::ESizeRule SizeRule;
-
 		/** Flag indicating if the ShrinkSizeValue value is set. */
 		bool bIsShrinkSizeValueSet;
-		
-		/** The actual value this size parameter stores. */
-		typename TBasicLayoutWidgetSlot<FSlot>::template TSlateSlotAttribute<float> SizeValue;
-
-		/** The actual value this size parameter stores, used for shrinking (negative if not defined, use SizeValue). */
-		typename TBasicLayoutWidgetSlot<FSlot>::template TSlateSlotAttribute<float> ShrinkSizeValue;
-
-		/** The min size that this slot can be */
-		typename TBasicLayoutWidgetSlot<FSlot>::template TSlateSlotAttribute<float> MinSize;
-
-		/** The max size that this slot can be (0 if no max) */
-		typename TBasicLayoutWidgetSlot<FSlot>::template TSlateSlotAttribute<float> MaxSize;
-
 	};
 
 	SLATE_BEGIN_ARGS(SScrollBox)
diff --git a/Engine/Source/Runtime/SlateCore/Public/Layout/BasicLayoutWidgetSlot.h b/Engine/Source/Runtime/SlateCore/Public/Layout/BasicLayoutWidgetSlot.h
--- a/Engine/Source/Runtime/SlateCore/Public/Layout/BasicLayoutWidgetSlot.h
+++ b/Engine/Source/Runtime/SlateCore/Public/Layout/BasicLayoutWidgetSlot.h
@@ -7,6 +7,7 @@
 #include "Layout/FlowDirection.h"
 #include "Layout/Margin.h"
 #include "Misc/Optional.h"
+#include "Types/SlateStructs.h"
 #include "Widgets/DeclarativeSyntaxSupport.h"
 
 #include <type_traits>
@@ -391,6 +392,193 @@ protected:
 };
 
 
+/** Mixin to add resizing functionality to a base slot. */
+template <typename MixedIntoType>
+class TResizingWidgetSlotMixin
+{
+public:
+	TResizingWidgetSlotMixin()
+		: SizeRule(FSizeParam::SizeRule_Stretch)
+		, SizeValue(static_cast<MixedIntoType&>(*this), 1.f)
+		, ShrinkSizeValue(static_cast<MixedIntoType&>(*this), 1.f)
+		, MinSize(static_cast<MixedIntoType&>(*this), 0.0f)
+		, MaxSize(static_cast<MixedIntoType&>(*this), 0.0f)
+	{}
+
+	struct FSlotArgumentsMixin
+	{
+	private:
+		friend TResizingWidgetSlotMixin;
+
+	public:
+		typename MixedIntoType::FSlotArguments& SizeParam(TOptional<FSizeParam> InSizeParam)
+		{
+			_SizeParam = InSizeParam;
+			return static_cast<typename MixedIntoType::FSlotArguments&>(*this);
+		}
+
+	protected:
+		TOptional<FSizeParam> _SizeParam;
+		TAttribute<float> _MinSize;
+		TAttribute<float> _MaxSize;
+	};
+
+protected:
+	void ConstructMixin(const FChildren& SlotOwner, FSlotArgumentsMixin&& InArgs)
+	{
+		if (InArgs._MinSize.IsSet())
+		{
+			SetMinSize(MoveTemp(InArgs._MinSize));
+		}
+
+		if (InArgs._MaxSize.IsSet())
+		{
+			SetMaxSize(MoveTemp(InArgs._MaxSize));
+		}
+
+		if (InArgs._SizeParam.IsSet())
+		{
+			SetSizeParam(MoveTemp(InArgs._SizeParam.GetValue()));
+		}
+	}
+
+	static void RegisterAttributesMixin(FSlateWidgetSlotAttributeInitializer& AttributeInitializer)
+	{
+		SLATE_ADD_SLOT_ATTRIBUTE_DEFINITION_WITH_NAME(MixedIntoType, AttributeInitializer, "Slot.MinSize", MinSize, EInvalidateWidgetReason::Layout);
+		SLATE_ADD_SLOT_ATTRIBUTE_DEFINITION_WITH_NAME(MixedIntoType, AttributeInitializer, "Slot.MaxSize", MaxSize, EInvalidateWidgetReason::Layout);
+		SLATE_ADD_SLOT_ATTRIBUTE_DEFINITION_WITH_NAME(MixedIntoType, AttributeInitializer, "Slot.SizeValue", SizeValue, EInvalidateWidgetReason::Layout)
+			.UpdatePrerequisite("Slot.MinSize")
+			.UpdatePrerequisite("Slot.MaxSize");
+		SLATE_ADD_SLOT_ATTRIBUTE_DEFINITION_WITH_NAME(MixedIntoType, AttributeInitializer, "Slot.ShrinkSizeValue", ShrinkSizeValue, EInvalidateWidgetReason::Layout)
+			.UpdatePrerequisite("Slot.MinSize")
+			.UpdatePrerequisite("Slot.MaxSize");
+	}
+
+public:
+	/** Get the space rule this slot should occupy along panel's direction. */
+	FSizeParam::ESizeRule GetSizeRule() const
+	{
+		return SizeRule;
+	}
+
+	/**
+	 * Get the size parameter for the space rule.
+	 * Used for size rule SizeRule_Stretch and SizeRule_StretchContent.
+	 */
+	float GetSizeValue() const
+	{
+		return SizeValue.Get();
+	}
+
+	/**
+	 * Get the size parameter for the space rule, used when the slot size is shrinking below desired size.
+	 * Used for size rule SizeRule_StretchContent.
+	 */
+	float GetShrinkSizeValue() const
+	{
+		return ShrinkSizeValue.Get();
+	}
+
+	/** Set the size Param of the slot, It could be a FStretch, FStretchContent, or a FAuto. */
+	void SetSizeParam(FSizeParam InSizeParam)
+	{
+		SizeRule = InSizeParam.SizeRule;
+
+		// ShrinkSizeValue is only used for StretchContent.
+		// If ShrinkValue is not set, make it equal to the Value.
+		if (SizeRule == FSizeParam::SizeRule_StretchContent)
+		{
+			if (InSizeParam.ShrinkValue.IsSet())
+			{
+				ShrinkSizeValue.Assign(static_cast<MixedIntoType&>(*this), MoveTemp(InSizeParam.ShrinkValue));
+			}
+			else
+			{
+				ShrinkSizeValue.Assign(static_cast<MixedIntoType&>(*this), InSizeParam.Value); // Make copy, let SizeValue use move.
+			}
+		}
+		else
+		{
+			ShrinkSizeValue.Set(static_cast<MixedIntoType&>(*this), 1.f); // Reset
+		}
+
+		SizeValue.Assign(static_cast<MixedIntoType&>(*this), MoveTemp(InSizeParam.Value));
+	}
+
+	/**
+	 * The widget's DesiredSize will be used as the space required.
+	 */
+	void SetSizeToAuto()
+	{
+		SetSizeParam(FAuto());
+	}
+
+	/**
+	 * The available space will be distributed proportionately to each slots stretch coefficient.
+	 * A slot with coefficient of 2 will get assigned twice as much available space as slot with coefficient 1.
+	 * @param InStretchCoefficient Stretch coefficient for this slot.
+	 */
+	void SetSizeToStretch(TAttribute<float> InStretchCoefficient)
+	{
+		SetSizeParam(FStretch(MoveTemp(InStretchCoefficient)));
+	}
+
+	/**
+	 * The widget's content size is adjusted proportionally to fit the available space.
+	 * The slots size starts at DesiredSize, and a slot with coefficient of 2 will get adjusted twice as much as slot with coefficient 1 to fit the available space.
+	 * @param InStretchCoefficient Stretch coefficient for this slot.
+	 * @param InShrinkStretchCoefficient If specified, this stretch coefficient is used when the slots is shrinking below desired size. Otherwise InStretchCoefficient is used for both shrink and grow.
+	 */
+	void SetSizeToStretchContent(TAttribute<float> InStretchCoefficient, TAttribute<float> InShrinkStretchCoefficient = TAttribute<float>())
+	{
+		SetSizeParam(FStretchContent(MoveTemp(InStretchCoefficient), MoveTemp(InShrinkStretchCoefficient)));
+	}
+
+	/** Get the min size the slot can be.*/
+	float GetMinSize() const
+	{
+		return MinSize.Get();
+	}
+
+	/** Set the min size in SlateUnit this slot can be. */
+	void SetMinSize(TAttribute<float> InMinSize)
+	{
+		MinSize.Assign(static_cast<MixedIntoType&>(*this), MoveTemp(InMinSize));
+	}
+
+	/** Get the max size the slot can be.*/
+	float GetMaxSize() const
+	{
+		return MaxSize.Get();
+	}
+
+	/** Set the max size in SlateUnit this slot can be. */
+	void SetMaxSize(TAttribute<float> InMaxSize)
+	{
+		MaxSize.Assign(static_cast<MixedIntoType&>(*this), MoveTemp(InMaxSize));
+	}
+
+protected:
+	/** The sizing rule to use, see ESizeRule for more info how the different rules work. */
+	FSizeParam::ESizeRule SizeRule;
+
+	template <typename InObjectType>
+	using TSlateSlotAttribute = ::SlateAttributePrivate::TSlateContainedAttribute<InObjectType, ::SlateAttributePrivate::FSlateAttributeNoInvalidationReason, TSlateAttributeComparePredicate<>>;
+
+	/** The actual value this size parameter stores. */
+	TSlateSlotAttribute<float> SizeValue;
+
+	/** The actual value this size parameter stores, used for shrinking (negative if not defined, use SizeValue). */
+	TSlateSlotAttribute<float> ShrinkSizeValue;
+
+	/** The min size that this slot can be */
+	TSlateSlotAttribute<float> MinSize;
+
+	/** The max size that this slot can be (0 if no max) */
+	TSlateSlotAttribute<float> MaxSize;
+};
+
+
 /** A templated basic slot that can be used by layout. */
 template <typename SlotType>
 class TBasicLayoutWidgetSlot : public TWidgetSlotWithAttributeSupport<SlotType>
diff --git a/Engine/Source/Runtime/SlateCore/Public/Widgets/SBoxPanel.h b/Engine/Source/Runtime/SlateCore/Public/Widgets/SBoxPanel.h
--- a/Engine/Source/Runtime/SlateCore/Public/Widgets/SBoxPanel.h
+++ b/Engine/Source/Runtime/SlateCore/Public/Widgets/SBoxPanel.h
@@ -33,176 +33,32 @@ protected:
 	 * child should be arranged on the screen.
 	 */
 	template<typename SlotType>
-	class TSlot : public TBasicLayoutWidgetSlot<SlotType>
+	class TSlot : public TBasicLayoutWidgetSlot<SlotType>, public TResizingWidgetSlotMixin<SlotType>
 	{
 	public:
-		SLATE_SLOT_BEGIN_ARGS(TSlot, TBasicLayoutWidgetSlot<SlotType>)
-			SLATE_ARGUMENT(TOptional<FSizeParam>, SizeParam)
-			TAttribute<float> _MinSize;
-			TAttribute<float> _MaxSize;
+		SLATE_SLOT_BEGIN_ARGS_OneMixin(TSlot, TBasicLayoutWidgetSlot<SlotType>, TResizingWidgetSlotMixin<SlotType>)
 		SLATE_SLOT_END_ARGS()
 
 	protected:
 		/** Default values for a slot. */
 		TSlot()
 			: TBasicLayoutWidgetSlot<SlotType>(HAlign_Fill, VAlign_Fill)
-			, SizeRule(FSizeParam::SizeRule_Stretch)
-			, SizeValue(*this, 1.f)
-			, ShrinkSizeValue(*this, 1.f)
-			, MinSize(*this, 0.0f)
-			, MaxSize(*this, 0.0f)
-		{ }
+			, TResizingWidgetSlotMixin<SlotType>()
+		{
+		}
 
 	public:
 		void Construct(const FChildren& SlotOwner, FSlotArguments&& InArgs)
 		{
 			TBasicLayoutWidgetSlot<SlotType>::Construct(SlotOwner, MoveTemp(InArgs));
-			if (InArgs._MinSize.IsSet())
-			{
-				SetMinSize(MoveTemp(InArgs._MinSize));
-			}
-			if (InArgs._MaxSize.IsSet())
-			{
-				SetMaxSize(MoveTemp(InArgs._MaxSize));
-			}
-			if (InArgs._SizeParam.IsSet())
-			{
-				SetSizeParam(MoveTemp(InArgs._SizeParam.GetValue()));
-			}
+			TResizingWidgetSlotMixin<SlotType>::ConstructMixin(SlotOwner, MoveTemp(InArgs));
 		}
 
 		static void RegisterAttributes(FSlateWidgetSlotAttributeInitializer& AttributeInitializer)
 		{
 			TBasicLayoutWidgetSlot<SlotType>::RegisterAttributes(AttributeInitializer);
-			SLATE_ADD_SLOT_ATTRIBUTE_DEFINITION_WITH_NAME(TSlot<SlotType>, AttributeInitializer, "Slot.MinSize", MinSize, EInvalidateWidgetReason::Layout);
-			SLATE_ADD_SLOT_ATTRIBUTE_DEFINITION_WITH_NAME(TSlot<SlotType>, AttributeInitializer, "Slot.MaxSize", MaxSize, EInvalidateWidgetReason::Layout);
-			SLATE_ADD_SLOT_ATTRIBUTE_DEFINITION_WITH_NAME(TSlot<SlotType>, AttributeInitializer, "Slot.SizeValue", SizeValue, EInvalidateWidgetReason::Layout)
-				.UpdatePrerequisite("Slot.MinSize")
-				.UpdatePrerequisite("Slot.MaxSize");
-			SLATE_ADD_SLOT_ATTRIBUTE_DEFINITION_WITH_NAME(TSlot<SlotType>, AttributeInitializer, "Slot.ShrinkSizeValue", ShrinkSizeValue, EInvalidateWidgetReason::Layout)
-				.UpdatePrerequisite("Slot.MinSize")
-				.UpdatePrerequisite("Slot.MaxSize");
-		}
-
-		/** Get the space rule this slot should occupy along panel's direction. */
-		FSizeParam::ESizeRule GetSizeRule() const
-		{
-			return SizeRule;
-		}
-
-		/**
-		 * Get the size parameter for the space rule.
-		 * Used for size rule SizeRule_Stretch and SizeRule_StretchContent. 
-		 */
-		float GetSizeValue() const
-		{
-			return SizeValue.Get();
-		}
-
-		/**
-		 * Get the size parameter for the space rule, used when the slot size is shrinking below desired size.
-		 * Used for size rule SizeRule_StretchContent.
-		 */
-		float GetShrinkSizeValue() const
-		{
-			return ShrinkSizeValue.Get();
-		}
-
-		/** Get the min size the slot can be.*/
-		float GetMinSize() const
-		{
-			return MinSize.Get();
-		}
-
-		/** Get the max size the slot can be.*/
-		float GetMaxSize() const
-		{
-			return MaxSize.Get();
+			TResizingWidgetSlotMixin<SlotType>::RegisterAttributesMixin(AttributeInitializer);
 		}
-
-	public:
-		/** Set the size Param of the slot, It could be a FStretch, FStretchContent, or a FAuto. */
-		void SetSizeParam(FSizeParam InSizeParam)
-		{
-			SizeRule = InSizeParam.SizeRule;
-
-			// ShrinkSizeValue is only used for StretchContent.
-			// If ShrinkValue is not set, make it equal to the Value. 
-			if (SizeRule == FSizeParam::SizeRule_StretchContent)
-			{
-				if (InSizeParam.ShrinkValue.IsSet())
-				{
-					ShrinkSizeValue.Assign(*this, MoveTemp(InSizeParam.ShrinkValue));
-				}
-				else
-				{
-					ShrinkSizeValue.Assign(*this, InSizeParam.Value); // Make copy, let SizeValue use move.
-				}
-			}
-			else
-			{
-				ShrinkSizeValue.Set(*this, 1.f); // Reset
-			}
-			
-			SizeValue.Assign(*this, MoveTemp(InSizeParam.Value));
-		}
-
-		/**
-		 * The widget's DesiredSize will be used as the space required.
-		 */
-		void SetSizeToAuto()
-		{
-			SetSizeParam(FAuto());
-		}
-
-		/**
-		 * The available space will be distributed proportionately to each slots stretch coefficient.
-		 * A slot with coefficient of 2 will get assigned twice as much available space as slot with coefficient 1. 
-		 * @param InStretchCoefficient Stretch coefficient for this slot.  
-		 */
-		void SetSizeToStretch(TAttribute<float> InStretchCoefficient)
-		{
-			SetSizeParam(FStretch(MoveTemp(InStretchCoefficient)));
-		}
-
-		/**
-		 * The widget's content size is adjusted proportionally to fit the available space.
-		 * The slots size starts at DesiredSize, and a slot with coefficient of 2 will get adjusted twice as much as slot with coefficient 1 to fit the available space.
-		 * @param InStretchCoefficient Stretch coefficient for this slot.
-		 * @param InShrinkStretchCoefficient If specified, this stretch coefficient is used when the slots is shrinking below desired size. Otherwise InStretchCoefficient is used for both shrink and grow.  
-		 */
-		void SetSizeToStretchContent(TAttribute<float> InStretchCoefficient, TAttribute<float> InShrinkStretchCoefficient = TAttribute<float>())
-		{
-			SetSizeParam(FStretchContent(MoveTemp(InStretchCoefficient), MoveTemp(InShrinkStretchCoefficient)));
-		}
-
-		/** Set the min size in SlateUnit this slot can be. */
-		void SetMinSize(TAttribute<float> InMinSize)
-		{
-			MinSize.Assign(*this, MoveTemp(InMinSize));
-		}
-
-		/** Set the max size in SlateUnit this slot can be. */
-		void SetMaxSize(TAttribute<float> InMaxSize)
-		{
-			MaxSize.Assign(*this, MoveTemp(InMaxSize));
-		}
-
-	private:
-		/** The sizing rule to use, see ESizeRule for more info how the different rules work. */
-		FSizeParam::ESizeRule SizeRule;
-
-		/** The actual value this size parameter stores. */
-		typename TBasicLayoutWidgetSlot<SlotType>::template TSlateSlotAttribute<float> SizeValue;
-
-		/** The actual value this size parameter stores, used for shrinking (negative if not defined, use SizeValue). */
-		typename TBasicLayoutWidgetSlot<SlotType>::template TSlateSlotAttribute<float> ShrinkSizeValue;
-
-		/** The min size that this slot can be */
-		typename TBasicLayoutWidgetSlot<SlotType>::template TSlateSlotAttribute<float> MinSize;
-
-		/** The max size that this slot can be (0 if no max) */
-		typename TBasicLayoutWidgetSlot<SlotType>::template TSlateSlotAttribute<float> MaxSize;
 	};
 
 public:
