Prevent constructing FInvalidationContext with an FWidgetStyle r-value.

diff --git a/Engine/Source/Runtime/SlateCore/Public/FastUpdate/SlateInvalidationRoot.h b/Engine/Source/Runtime/SlateCore/Public/FastUpdate/SlateInvalidationRoot.h
--- a/Engine/Source/Runtime/SlateCore/Public/FastUpdate/SlateInvalidationRoot.h
+++ b/Engine/Source/Runtime/SlateCore/Public/FastUpdate/SlateInvalidationRoot.h
@@ -39,6 +39,9 @@ struct FSlateInvalidationContext
 	{
 	} 
 
+	// Prevent construction with a temporary FWidgetStyle.
+	FSlateInvalidationContext(FSlateWindowElementList&, FWidgetStyle&&) = delete;
+
 	FSlateRect CullingRect;
 	UE::Slate::FDeprecateVector2DResult ViewOffset;
 	const FPaintArgs* PaintArgs;
