Backing out this change that could impact licensees.

[Backout] - CL39520023
#rnx
Original CL Desc
-----------------------------------------------------------------
Prevent constructing FInvalidationContext with an FWidgetStyle r-value.

diff --git a/Engine/Source/Runtime/SlateCore/Public/FastUpdate/SlateInvalidationRoot.h b/Engine/Source/Runtime/SlateCore/Public/FastUpdate/SlateInvalidationRoot.h
--- a/Engine/Source/Runtime/SlateCore/Public/FastUpdate/SlateInvalidationRoot.h
+++ b/Engine/Source/Runtime/SlateCore/Public/FastUpdate/SlateInvalidationRoot.h
@@ -39,9 +39,6 @@ struct FSlateInvalidationContext
 	{
 	} 
 
-	// Prevent construction with a temporary FWidgetStyle.
-	FSlateInvalidationContext(FSlateWindowElementList&, FWidgetStyle&&) = delete;
-
 	FSlateRect CullingRect;
 	UE::Slate::FDeprecateVector2DResult ViewOffset;
 	const FPaintArgs* PaintArgs;
