Disabling this visualizer in test builds due to potential UI perf regression.
#RNX

diff --git a/Engine/Source/Runtime/SlateCore/Private/Rendering/ElementBatcher.cpp b/Engine/Source/Runtime/SlateCore/Private/Rendering/ElementBatcher.cpp
--- a/Engine/Source/Runtime/SlateCore/Private/Rendering/ElementBatcher.cpp
+++ b/Engine/Source/Runtime/SlateCore/Private/Rendering/ElementBatcher.cpp
@@ -16,7 +16,7 @@
 #include "Widgets/SWidgetUtils.h"
 #include "Rendering/DrawElementPayloads.h"
 
-#define ENABLE_SDF_DEBUG_VISUAL (!UE_BUILD_SHIPPING)
+#define ENABLE_SDF_DEBUG_VISUAL (!UE_BUILD_SHIPPING && !UE_BUILD_TEST)
 
 CSV_DECLARE_CATEGORY_MODULE_EXTERN(SLATECORE_API, Slate);
 
