Slate: FPointerEvent constructed with a Gesture is now considered a TouchEvent

#rnx

diff --git a/Engine/Source/Runtime/SlateCore/Public/Input/Events.h b/Engine/Source/Runtime/SlateCore/Public/Input/Events.h
--- a/Engine/Source/Runtime/SlateCore/Public/Input/Events.h
+++ b/Engine/Source/Runtime/SlateCore/Public/Input/Events.h
@@ -926,7 +926,7 @@ public:
 		, PointerIndex(0)
 		, TouchpadIndex(0)
 		, Force(1.0f)
-		, bIsTouchEvent(false)
+		, bIsTouchEvent(InGestureType != EGestureEvent::None)
 		, GestureType(InGestureType)
 		, WheelOrGestureDelta(InGestureDelta)
 		, bIsDirectionInvertedFromDevice(bInIsDirectionInvertedFromDevice)
