[Slate] Added MinSize, MaxSize access to TResizingWidgetSlotMixin SlotArgs

diff --git a/Engine/Source/Runtime/SlateCore/Public/Layout/BasicLayoutWidgetSlot.h b/Engine/Source/Runtime/SlateCore/Public/Layout/BasicLayoutWidgetSlot.h
--- a/Engine/Source/Runtime/SlateCore/Public/Layout/BasicLayoutWidgetSlot.h
+++ b/Engine/Source/Runtime/SlateCore/Public/Layout/BasicLayoutWidgetSlot.h
@@ -417,6 +417,18 @@ public:
 			return static_cast<typename MixedIntoType::FSlotArguments&>(*this);
 		}
 
+		typename MixedIntoType::FSlotArguments& MinSize(TAttribute<float> InMinSize)
+		{
+			_MinSize = MoveTemp(InMinSize);
+			return static_cast<typename MixedIntoType::FSlotArguments&>(*this);
+		}
+
+		typename MixedIntoType::FSlotArguments& MaxSize(TAttribute<float> InMaxSize)
+		{
+			_MinSize = MoveTemp(InMaxSize);
+			return static_cast<typename MixedIntoType::FSlotArguments&>(*this);
+		}
+
 	protected:
 		TOptional<FSizeParam> _SizeParam;
 		TAttribute<float> _MinSize;
