[Slate] ResizingWidgetSlotMixin MaxSize setter fix

diff --git a/Engine/Source/Runtime/SlateCore/Public/Layout/BasicLayoutWidgetSlot.h b/Engine/Source/Runtime/SlateCore/Public/Layout/BasicLayoutWidgetSlot.h
--- a/Engine/Source/Runtime/SlateCore/Public/Layout/BasicLayoutWidgetSlot.h
+++ b/Engine/Source/Runtime/SlateCore/Public/Layout/BasicLayoutWidgetSlot.h
@@ -425,7 +425,7 @@ public:
 
 		typename MixedIntoType::FSlotArguments& MaxSize(TAttribute<float> InMaxSize)
 		{
-			_MinSize = MoveTemp(InMaxSize);
+			_MaxSize = MoveTemp(InMaxSize);
 			return static_cast<typename MixedIntoType::FSlotArguments&>(*this);
 		}
 
