Allowed Slate attribute bindings to be explicitly initialized on assignment.

When attributes are bound they are not automatically updated to the current value until later. This can cause issues with dynamicaly created widget as they can have a frame where they use the default value before getting the value from the binding next frame. In the case of the color block this resulted in a white flash. However the existing code indicated that there's also a need for this behavior in certain cases, so to allow for both a new flag enum has been added so widget can individually choose how to respond to attribute binding.

[RN] Allowed Slate attribute bindings to be explicitly initialized on assignment.

diff --git a/Engine/Source/Runtime/Slate/Private/Widgets/Colors/SColorBlock.cpp b/Engine/Source/Runtime/Slate/Private/Widgets/Colors/SColorBlock.cpp
--- a/Engine/Source/Runtime/Slate/Private/Widgets/Colors/SColorBlock.cpp
+++ b/Engine/Source/Runtime/Slate/Private/Widgets/Colors/SColorBlock.cpp
@@ -40,15 +40,15 @@ SColorBlock::~SColorBlock() = default;
  */
 void SColorBlock::Construct(const FArguments& InArgs)
 {
-	Color.Assign(*this, InArgs._Color);
-	AlphaBackgroundBrush.Assign(*this, InArgs._AlphaBackgroundBrush);
-	GradientCornerRadius.Assign(*this, InArgs._CornerRadius, FVector4(0.0f));
-	ColorBlockSize.Assign(*this, InArgs._Size);
+	Color.Assign(*this, InArgs._Color, ESlateAttributeBindAction::Initialize);
+	AlphaBackgroundBrush.Assign(*this, InArgs._AlphaBackgroundBrush, ESlateAttributeBindAction::Initialize);
+	GradientCornerRadius.Assign(*this, InArgs._CornerRadius, FVector4(0.0f), ESlateAttributeBindAction::Initialize);
+	ColorBlockSize.Assign(*this, InArgs._Size, ESlateAttributeBindAction::Initialize);
 	MouseButtonDownHandler = InArgs._OnMouseButtonDown;
-	AlphaDisplayMode.Assign(*this, InArgs._AlphaDisplayMode);
-	ColorIsHSV.Assign(*this, InArgs._ColorIsHSV);
-	ShowBackgroundForAlpha.Assign(*this, InArgs._ShowBackgroundForAlpha);
-	bUseSRGB.Assign(*this, InArgs._UseSRGB);
+	AlphaDisplayMode.Assign(*this, InArgs._AlphaDisplayMode, ESlateAttributeBindAction::Initialize);
+	ColorIsHSV.Assign(*this, InArgs._ColorIsHSV, ESlateAttributeBindAction::Initialize);
+	ShowBackgroundForAlpha.Assign(*this, InArgs._ShowBackgroundForAlpha, ESlateAttributeBindAction::Initialize);
+	bUseSRGB.Assign(*this, InArgs._UseSRGB, ESlateAttributeBindAction::Initialize);
 }
 
 int32 SColorBlock::OnPaint(const FPaintArgs& Args, const FGeometry& AllottedGeometry, const FSlateRect& MyCullingRect, FSlateWindowElementList& OutDrawElements, int32 LayerId, const FWidgetStyle& InWidgetStyle, bool bParentEnabled) const
diff --git a/Engine/Source/Runtime/SlateCore/Public/Types/Attributes/SlateAttributeBase.inl b/Engine/Source/Runtime/SlateCore/Public/Types/Attributes/SlateAttributeBase.inl
--- a/Engine/Source/Runtime/SlateCore/Public/Types/Attributes/SlateAttributeBase.inl
+++ b/Engine/Source/Runtime/SlateCore/Public/Types/Attributes/SlateAttributeBase.inl
@@ -38,13 +38,6 @@ namespace SlateAttributePrivate
 		static bool IdenticalTo(const ISlateAttributeContainer& Container, const ObjectType& Lhs, const ObjectType& Rhs) { return FComparePredicate::IdenticalTo(Container.GetContainerWidget(), Lhs, Rhs); }
 
 	private:
-		void UpdateNowOnBind(ContainerType& Widget)
-		{
-#if UE_SLATE_WITH_ATTRIBUTE_INITIALIZATION_ON_BIND
-			ProtectedUpdateNow(Widget, InAttributeType);
-#endif
-		}
-
 		void VerifyOwningWidget(const SWidget& Widget) const
 		{
 #if UE_SLATE_WITH_MEMBER_ATTRIBUTE_DEBUGGING
@@ -354,12 +347,12 @@ namespace SlateAttributePrivate
 		 * @see Set
 		 * @return true if the getter was assigned or if the value is considered different and an invalidation occurred.
 		 */
-		bool Assign(ContainerType& Widget, const TAttribute<ObjectType>& OtherAttribute)
+		bool Assign(ContainerType& Widget, const TAttribute<ObjectType>& OtherAttribute, ESlateAttributeBindAction Actions = UE_SLATE_WITH_ATTRIBUTE_DEFAULT_INITIALIZATION_ACTION)
 		{
 			VerifyOwningWidget(Widget);
 			if (OtherAttribute.IsBound())
 			{
-				return AssignBinding(Widget, OtherAttribute.GetBinding());
+				return AssignBinding(Widget, OtherAttribute.GetBinding(), Actions);
 			}
 			else if (OtherAttribute.IsSet())
 			{
@@ -372,12 +365,12 @@ namespace SlateAttributePrivate
 			}
 		}
 
-		bool Assign(ContainerType& Widget, TAttribute<ObjectType>&& OtherAttribute)
+		bool Assign(ContainerType& Widget, TAttribute<ObjectType>&& OtherAttribute, ESlateAttributeBindAction Actions = UE_SLATE_WITH_ATTRIBUTE_DEFAULT_INITIALIZATION_ACTION)
 		{
 			VerifyOwningWidget(Widget);
 			if (OtherAttribute.IsBound())
 			{
-				return AssignBinding(Widget, MoveTemp(OtherAttribute.Steal().template Get<FGetter>()));
+				return AssignBinding(Widget, MoveTemp(OtherAttribute.Steal().template Get<FGetter>()), Actions);
 			}
 			else if (OtherAttribute.IsSet())
 			{
@@ -390,12 +383,12 @@ namespace SlateAttributePrivate
 			}
 		}
 
-		bool Assign(ContainerType& Widget, const TAttribute<ObjectType>& OtherAttribute, const ObjectType& DefaultValue)
+		bool Assign(ContainerType& Widget, const TAttribute<ObjectType>& OtherAttribute, const ObjectType& DefaultValue, ESlateAttributeBindAction Actions = UE_SLATE_WITH_ATTRIBUTE_DEFAULT_INITIALIZATION_ACTION)
 		{
 			VerifyOwningWidget(Widget);
 			if (OtherAttribute.IsBound())
 			{
-				return AssignBinding(Widget, OtherAttribute.GetBinding());
+				return AssignBinding(Widget, OtherAttribute.GetBinding(), Actions);
 			}
 			else if (OtherAttribute.IsSet())
 			{
@@ -407,12 +400,12 @@ namespace SlateAttributePrivate
 			}
 		}
 
-		bool Assign(ContainerType& Widget, TAttribute<ObjectType>&& OtherAttribute, const ObjectType& DefaultValue)
+		bool Assign(ContainerType& Widget, TAttribute<ObjectType>&& OtherAttribute, const ObjectType& DefaultValue, ESlateAttributeBindAction Actions = UE_SLATE_WITH_ATTRIBUTE_DEFAULT_INITIALIZATION_ACTION)
 		{
 			VerifyOwningWidget(Widget);
 			if (OtherAttribute.IsBound())
 			{
-				return AssignBinding(Widget, MoveTemp(OtherAttribute.Steal().template Get<FGetter>()));
+				return AssignBinding(Widget, MoveTemp(OtherAttribute.Steal().template Get<FGetter>()), Actions);
 			}
 			else if (OtherAttribute.IsSet())
 			{
@@ -424,12 +417,12 @@ namespace SlateAttributePrivate
 			}
 		}
 
-		bool Assign(ContainerType& Widget, const TAttribute<ObjectType>& OtherAttribute, ObjectType&& DefaultValue)
+		bool Assign(ContainerType& Widget, const TAttribute<ObjectType>& OtherAttribute, ObjectType&& DefaultValue, ESlateAttributeBindAction Actions = UE_SLATE_WITH_ATTRIBUTE_DEFAULT_INITIALIZATION_ACTION)
 		{
 			VerifyOwningWidget(Widget);
 			if (OtherAttribute.IsBound())
 			{
-				return AssignBinding(Widget, OtherAttribute.GetBinding());
+				return AssignBinding(Widget, OtherAttribute.GetBinding(), Actions);
 			}
 			else if (OtherAttribute.IsSet())
 			{
@@ -441,12 +434,12 @@ namespace SlateAttributePrivate
 			}
 		}
 
-		bool Assign(ContainerType& Widget, TAttribute<ObjectType>&& OtherAttribute, ObjectType&& DefaultValue)
+		bool Assign(ContainerType& Widget, TAttribute<ObjectType>&& OtherAttribute, ObjectType&& DefaultValue, ESlateAttributeBindAction Actions = UE_SLATE_WITH_ATTRIBUTE_DEFAULT_INITIALIZATION_ACTION)
 		{
 			VerifyOwningWidget(Widget);
 			if (OtherAttribute.IsBound())
 			{
-				return AssignBinding(Widget, MoveTemp(OtherAttribute.Steal().template Get<FGetter>()));
+				return AssignBinding(Widget, MoveTemp(OtherAttribute.Steal().template Get<FGetter>()), Actions);
 			}
 			else if (OtherAttribute.IsSet())
 			{
@@ -516,37 +509,43 @@ namespace SlateAttributePrivate
 		}
 
 	private:
-		void ConstructWrapper(ContainerType& Widget, const FGetter& Getter)
+		void ConstructWrapper(ContainerType& Widget, const FGetter& Getter, ESlateAttributeBindAction Actions = UE_SLATE_WITH_ATTRIBUTE_DEFAULT_INITIALIZATION_ACTION)
 		{
 			TUniquePtr<ISlateAttributeGetter> Wrapper = MakeUniqueGetter(*this, Getter);
 			ProtectedRegisterAttribute(Widget, InAttributeType, MoveTemp(Wrapper));
-			UpdateNowOnBind(Widget);
+			if (EnumHasAllFlags(Actions, ESlateAttributeBindAction::Initialize))
+			{
+				ProtectedUpdateNow(Widget, InAttributeType);
+			}
 		}
 
-		void ConstructWrapper(ContainerType& Widget, FGetter&& Getter)
+		void ConstructWrapper(ContainerType& Widget, FGetter&& Getter, ESlateAttributeBindAction Actions = UE_SLATE_WITH_ATTRIBUTE_DEFAULT_INITIALIZATION_ACTION)
 		{
 			TUniquePtr<ISlateAttributeGetter> Wrapper = MakeUniqueGetter(*this, MoveTemp(Getter));
 			ProtectedRegisterAttribute(Widget, InAttributeType, MoveTemp(Wrapper));
-			UpdateNowOnBind(Widget);
+			if (EnumHasAllFlags(Actions, ESlateAttributeBindAction::Initialize))
+			{
+				ProtectedUpdateNow(Widget, InAttributeType);
+			}
 		}
 
-		bool AssignBinding(ContainerType& Widget, const FGetter& Getter)
+		bool AssignBinding(ContainerType& Widget, const FGetter& Getter, ESlateAttributeBindAction Actions = UE_SLATE_WITH_ATTRIBUTE_DEFAULT_INITIALIZATION_ACTION)
 		{
 			const FDelegateHandle PreviousGetterHandle = ProtectedFindGetterHandle(Widget, InAttributeType);
 			if (PreviousGetterHandle != Getter.GetHandle())
 			{
-				ConstructWrapper(Widget, Getter);
+				ConstructWrapper(Widget, Getter, Actions);
 				return true;
 			}
 			return false;
 		}
 
-		bool AssignBinding(ContainerType& Widget, FGetter&& Getter)
+		bool AssignBinding(ContainerType& Widget, FGetter&& Getter, ESlateAttributeBindAction Actions = UE_SLATE_WITH_ATTRIBUTE_DEFAULT_INITIALIZATION_ACTION)
 		{
 			const FDelegateHandle PreviousGetterHandle = ProtectedFindGetterHandle(Widget, InAttributeType);
 			if (PreviousGetterHandle != Getter.GetHandle())
 			{
-				ConstructWrapper(Widget, MoveTemp(Getter));
+				ConstructWrapper(Widget, MoveTemp(Getter), Actions);
 				return true;
 			}
 			return false;
diff --git a/Engine/Source/Runtime/SlateCore/Public/Types/SlateAttribute.h b/Engine/Source/Runtime/SlateCore/Public/Types/SlateAttribute.h
--- a/Engine/Source/Runtime/SlateCore/Public/Types/SlateAttribute.h
+++ b/Engine/Source/Runtime/SlateCore/Public/Types/SlateAttribute.h
@@ -4,6 +4,7 @@
 
 #include "CoreMinimal.h"
 #include "Misc/Attribute.h"
+#include "Misc/EnumClassFlags.h"
 #include "Templates/EqualTo.h"
 #include "Widgets/InvalidateWidgetReason.h"
 
@@ -163,6 +164,17 @@
 
 class SWidget;
 
+enum class ESlateAttributeBindAction
+{
+	None = 0, // Do nothing when an attribute is bound.
+	Initialize = 1, // Initialize the attribute value from the binding.
+};
+ENUM_CLASS_FLAGS(ESlateAttributeBindAction)
+#if UE_SLATE_WITH_ATTRIBUTE_INITIALIZATION_ON_BIND
+#	define UE_SLATE_WITH_ATTRIBUTE_DEFAULT_INITIALIZATION_ACTION ESlateAttributeBindAction::Initialize
+#else
+#	define UE_SLATE_WITH_ATTRIBUTE_DEFAULT_INITIALIZATION_ACTION ESlateAttributeBindAction::None
+#endif
 
 /** Base struct of all SlateAttribute type. */
 struct FSlateAttributeBase
