Add tracking of whether virtual textures are being sampled in slate batches.
This will be used to determine if we need to bind a virtual feedback buffer for the pass that renders the batches.

diff --git a/Engine/Source/Runtime/SlateCore/Private/Rendering/ElementBatcher.cpp b/Engine/Source/Runtime/SlateCore/Private/Rendering/ElementBatcher.cpp
--- a/Engine/Source/Runtime/SlateCore/Private/Rendering/ElementBatcher.cpp
+++ b/Engine/Source/Runtime/SlateCore/Private/Rendering/ElementBatcher.cpp
@@ -80,6 +80,7 @@ FSlateBatchData::FSlateBatchData()
 	, NumLayers(0)
 	, NumBatches(0)
 	, bIsStencilBufferRequired(false)
+	, bRequiresVirtualTextureFeedback(false)
 {}
 
 FSlateBatchData::~FSlateBatchData() = default;
@@ -98,6 +99,7 @@ void FSlateBatchData::ResetData()
 	NumLayers = 0;
 
 	bIsStencilBufferRequired = false;
+	bRequiresVirtualTextureFeedback = false;
 }
 
 #define MAX_VERT_ARRAY_RECYCLE (200)
@@ -108,6 +110,11 @@ bool FSlateBatchData::IsStencilClippingRequired() const
 	return bIsStencilBufferRequired;
 }
 
+bool FSlateBatchData::IsVirtualTextureFeedbackRequired() const
+{
+	return bRequiresVirtualTextureFeedback;
+}
+
 FSlateRenderBatch& FSlateBatchData::AddRenderBatch(int32 InLayer, const FShaderParams& InShaderParams, const FSlateShaderResource* InResource, ESlateDrawPrimitive InPrimitiveType, ESlateShader InShaderType, ESlateDrawEffect InDrawEffects, ESlateBatchDrawFlag InDrawFlags, int8 SceneIndex)
 {
 	return RenderBatches.Emplace_GetRef(InLayer, InShaderParams, InResource, InPrimitiveType, InShaderType, InDrawEffects, InDrawFlags, SceneIndex, &UncachedSourceBatchVertices, &UncachedSourceBatchIndices, UncachedSourceBatchVertices.Num(), UncachedSourceBatchIndices.Num());
@@ -237,6 +244,11 @@ void FSlateBatchData::MergeRenderBatches()
 				bIsStencilBufferRequired |= CurBatch.ClippingState->GetClippingMethod() == EClippingMethod::Stencil;
 			}
 
+			if (const FSlateShaderResource* ShaderResource = CurBatch.GetShaderResource())
+			{
+				bRequiresVirtualTextureFeedback |= ShaderResource->RequiresVirtualTextureFeedback();
+			}
+
 #if 1  // Do batching at all?
 
 			if (CurBatch.bIsMergable)
diff --git a/Engine/Source/Runtime/SlateCore/Public/Rendering/ElementBatcher.h b/Engine/Source/Runtime/SlateCore/Public/Rendering/ElementBatcher.h
--- a/Engine/Source/Runtime/SlateCore/Public/Rendering/ElementBatcher.h
+++ b/Engine/Source/Runtime/SlateCore/Public/Rendering/ElementBatcher.h
@@ -164,6 +164,9 @@ public:
 	/** True if stencil buffer / clipping is needed. */
 	SLATECORE_API bool IsStencilClippingRequired() const;
 
+	/** True if virtual texture feedback is needed. */
+	SLATECORE_API bool IsVirtualTextureFeedbackRequired() const;
+
 	int32 GetFirstRenderBatchIndex() const { return FirstRenderBatchIndex; }
 	/** @return Total number of batched layers */
 	int32 GetNumLayers() const { return NumLayers; }
@@ -229,8 +232,10 @@ private:
 	/** Number of final render batches.  it is not the same as RenderBatches.Num(); */
 	int32 NumBatches;
 
-	/** */
+	/** true if slate rendering pass requires a stencil buffer */
 	bool bIsStencilBufferRequired;
+	/** true if slate rendering pass requires a virtual texture feedback buffer */
+	bool bRequiresVirtualTextureFeedback;
 };
 
 
diff --git a/Engine/Source/Runtime/SlateCore/Public/Textures/SlateShaderResource.h b/Engine/Source/Runtime/SlateCore/Public/Textures/SlateShaderResource.h
--- a/Engine/Source/Runtime/SlateCore/Public/Textures/SlateShaderResource.h
+++ b/Engine/Source/Runtime/SlateCore/Public/Textures/SlateShaderResource.h
@@ -75,6 +75,11 @@ public:
 	 */
 	virtual ESlatePostRT GetUsedSlatePostBuffers() const { return ESlatePostRT::None; }
 
+	/**
+	 * Does this resource require virtual texture feedback?
+	*/
+	virtual bool RequiresVirtualTextureFeedback() const { return false; }
+
 #if SLATE_CHECK_UOBJECT_RENDER_RESOURCES
 	virtual void CheckForStaleResources() const { }
 	bool Debug_IsDestroyed() const { return DestroyState != 0x21; }
diff --git a/Engine/Source/Runtime/SlateRHIRenderer/Private/SlateMaterialResource.cpp b/Engine/Source/Runtime/SlateRHIRenderer/Private/SlateMaterialResource.cpp
--- a/Engine/Source/Runtime/SlateRHIRenderer/Private/SlateMaterialResource.cpp
+++ b/Engine/Source/Runtime/SlateRHIRenderer/Private/SlateMaterialResource.cpp
@@ -73,6 +73,7 @@ FSlateMaterialResource::FSlateMaterialResource(const UMaterialInterface& InMater
 	, Width(FMath::RoundToInt(InImageSize.X))
 	, Height(FMath::RoundToInt(InImageSize.Y))
 	, CachedSlatePostBuffers(ESlatePostRT::None)
+	, CachedUsesVirtualTextures(false)
 {
 #if SLATE_CHECK_UOBJECT_RENDER_RESOURCES
 	SlateMaterialResource::CheckInvalidUMaterial(InMaterialResource, NAME_None);
@@ -112,6 +113,8 @@ FSlateMaterialResource::FSlateMaterialResource(const UMaterialInterface& InMater
 						CachedSlatePostBuffers |= SlatePostBitflag;
 					}
 				}
+
+				CachedUsesVirtualTextures |= (OutUsedTexture && OutUsedTexture->IsCurrentlyVirtualTextured());
 			}
 		}
 	}
@@ -135,6 +138,11 @@ ESlatePostRT FSlateMaterialResource::GetUsedSlatePostBuffers() const
 	return CachedSlatePostBuffers;
 }
 
+bool FSlateMaterialResource::RequiresVirtualTextureFeedback() const
+{
+	return CachedUsesVirtualTextures;
+}
+
 bool FSlateMaterialResource::IsResourceValid() const
 {
 	if (MaterialProxy && MaterialObject)
diff --git a/Engine/Source/Runtime/SlateRHIRenderer/Private/SlateMaterialResource.h b/Engine/Source/Runtime/SlateRHIRenderer/Private/SlateMaterialResource.h
--- a/Engine/Source/Runtime/SlateRHIRenderer/Private/SlateMaterialResource.h
+++ b/Engine/Source/Runtime/SlateRHIRenderer/Private/SlateMaterialResource.h
@@ -22,6 +22,7 @@ public:
 	virtual uint32 GetHeight() const override { return Height; }
 	virtual ESlateShaderResource::Type GetType() const override { return ESlateShaderResource::Material; }
 	virtual ESlatePostRT GetUsedSlatePostBuffers() const override;
+	virtual bool RequiresVirtualTextureFeedback() const override;
 	virtual bool IsResourceValid() const override;
 	//~ End FSlateShaderResource Interface.
 
@@ -66,6 +67,9 @@ private:
 	/** Cached slate SlatePostRT assets / buffers in use */
 	ESlatePostRT CachedSlatePostBuffers;
 
+	/** Cached result of whether we use virtual textures */
+	bool CachedUsesVirtualTextures;
+
 private:
 #if SLATE_CHECK_UOBJECT_RENDER_RESOURCES
 	void UpdateMaterialName();
diff --git a/Engine/Source/Runtime/SlateRHIRenderer/Private/SlateUTextureResource.cpp b/Engine/Source/Runtime/SlateRHIRenderer/Private/SlateUTextureResource.cpp
--- a/Engine/Source/Runtime/SlateRHIRenderer/Private/SlateUTextureResource.cpp
+++ b/Engine/Source/Runtime/SlateRHIRenderer/Private/SlateUTextureResource.cpp
@@ -6,11 +6,17 @@
 FSlateBaseUTextureResource::FSlateBaseUTextureResource(UTexture* InTexture)
 	: TextureObject(InTexture)
 	, CachedSlatePostBuffers(ESlatePostRT::None)
+	, CachedIsVirtualTexture(false)
 {
 #if SLATE_CHECK_UOBJECT_RENDER_RESOURCES
 	ObjectWeakPtr = InTexture;
 	UpdateDebugName();
 #endif
+
+	if (TextureObject&& TextureObject->IsCurrentlyVirtualTextured())
+	{
+		CachedIsVirtualTexture = true;
+	}
 }
 
 FSlateBaseUTextureResource::~FSlateBaseUTextureResource()
@@ -37,6 +43,11 @@ ESlatePostRT FSlateBaseUTextureResource::GetUsedSlatePostBuffers() const
 	return CachedSlatePostBuffers;
 }
 
+bool FSlateBaseUTextureResource::RequiresVirtualTextureFeedback() const
+{
+	return CachedIsVirtualTexture;
+}
+
 #if SLATE_CHECK_UOBJECT_RENDER_RESOURCES
 void FSlateBaseUTextureResource::UpdateDebugName()
 {
diff --git a/Engine/Source/Runtime/SlateRHIRenderer/Private/SlateUTextureResource.h b/Engine/Source/Runtime/SlateRHIRenderer/Private/SlateUTextureResource.h
--- a/Engine/Source/Runtime/SlateRHIRenderer/Private/SlateUTextureResource.h
+++ b/Engine/Source/Runtime/SlateRHIRenderer/Private/SlateUTextureResource.h
@@ -22,6 +22,7 @@ public:
 	virtual uint32 GetHeight() const;
 	virtual ESlateShaderResource::Type GetType() const;
 	virtual ESlatePostRT GetUsedSlatePostBuffers() const override;
+	virtual bool RequiresVirtualTextureFeedback() const override;
 
 	/** Gets the RHI resource used for rendering and updates the last render time for texture streaming */
 	FTextureRHIRef AccessRHIResource()
@@ -60,6 +61,8 @@ protected:
 
 	/** Cached slate SlatePostRT assets / buffers in use */
 	ESlatePostRT CachedSlatePostBuffers;
+	/** Cached result of whether we use virtual textures */
+	bool CachedIsVirtualTexture;
 };
 
 
