Fix there being multiple ambigious functions of "GetUserIndex" on FInputEvent and its subclasses.

FInputEvent::GetUserIndex already exists and is public, so the subclasses FMotionEvent and FPointerEvent have no need to expose their own versions.

#rnx

diff --git a/Engine/Source/Runtime/SlateCore/Public/Input/Events.h b/Engine/Source/Runtime/SlateCore/Public/Input/Events.h
--- a/Engine/Source/Runtime/SlateCore/Public/Input/Events.h
+++ b/Engine/Source/Runtime/SlateCore/Public/Input/Events.h
@@ -989,9 +989,6 @@ public:
 	/** How much did the mouse wheel turn since the last mouse event */
 	float GetWheelDelta() const { return UE_REAL_TO_FLOAT(WheelOrGestureDelta.Y); }
 
-	/** Returns the index of the user that caused the event */
-	int32 GetUserIndex() const { return UserIndex; }
-
 	/** Returns the unique identifier of the pointer (e.g., finger index) */
 	uint32 GetPointerIndex() const { return PointerIndex; }
 
@@ -1117,9 +1114,6 @@ public:
 
 public:
 
-	/** Returns the index of the user that caused the event */
-	uint32 GetUserIndex() const { return UserIndex; }
-
 	/** Returns the current tilt of the device/controller */
 	const FVector& GetTilt() const { return Tilt; }
 
