USlateThemeManager : Add an ExportCurrentThemeTo method.

diff --git a/Engine/Source/Editor/UnrealEd/Private/Settings/EditorStyleClasses.cpp b/Engine/Source/Editor/UnrealEd/Private/Settings/EditorStyleClasses.cpp
--- a/Engine/Source/Editor/UnrealEd/Private/Settings/EditorStyleClasses.cpp
+++ b/Engine/Source/Editor/UnrealEd/Private/Settings/EditorStyleClasses.cpp
@@ -133,6 +133,10 @@ bool UEditorStyleSettings::OnImportBegin(const FString& ImportFromPath)
 bool UEditorStyleSettings::OnExportBegin(const FString& ExportToPath)
 {
 	SaveConfig(CPF_Config, *ExportToPath);
+	if (!IPlatformFile::GetPlatformPhysical().FileExists(*ExportToPath))
+	{
+		return false;
+	}
 
 	const FStyleTheme& CurrentTheme = USlateThemeManager::Get().GetCurrentTheme();
 
@@ -143,14 +147,6 @@ bool UEditorStyleSettings::OnExportBegin(const FString& ExportToPath)
 
 	const FString DestinationPath = PathPart / CurrentTheme.DisplayName.ToString() + TEXT(".json");
 
-	// Save the theme file to the same destination, replacing any existing file at that path.
-	if (IPlatformFile::GetPlatformPhysical().FileExists(*DestinationPath))
-	{
-		IPlatformFile::GetPlatformPhysical().DeleteFile(*DestinationPath);
-	}
-
-	USlateThemeManager::Get().SaveCurrentThemeAs(DestinationPath);
-
-	return IPlatformFile::GetPlatformPhysical().FileExists(*ExportToPath)
-		&& IPlatformFile::GetPlatformPhysical().FileExists(*DestinationPath);
+	// Export the theme file to the same destination, replacing any existing file at that path.
+	return USlateThemeManager::Get().ExportCurrentThemeTo(DestinationPath);
 }
diff --git a/Engine/Source/Editor/UnrealEd/Private/Settings/EditorStyleSettingsCustomization.cpp b/Engine/Source/Editor/UnrealEd/Private/Settings/EditorStyleSettingsCustomization.cpp
--- a/Engine/Source/Editor/UnrealEd/Private/Settings/EditorStyleSettingsCustomization.cpp
+++ b/Engine/Source/Editor/UnrealEd/Private/Settings/EditorStyleSettingsCustomization.cpp
@@ -574,14 +574,7 @@ FReply FEditorStyleSettingsCustomization::OnExportThemeClicked()
 	{
 		const FString DestPath = OutFiles[0];
 
-		if (IPlatformFile::GetPlatformPhysical().FileExists(*DestPath))
-		{
-			IPlatformFile::GetPlatformPhysical().DeleteFile(*DestPath);
-		}
-
-		USlateThemeManager::Get().SaveCurrentThemeAs(DestPath);
-
-		if (IPlatformFile::GetPlatformPhysical().FileExists(*DestPath))
+		if (USlateThemeManager::Get().ExportCurrentThemeTo(DestPath))
 		{
 			ShowNotification(LOCTEXT("ExportThemeSuccess", "Export theme succeeded"), SNotificationItem::CS_Success);
 		}
diff --git a/Engine/Source/Runtime/SlateCore/Private/Styling/StyleColors.cpp b/Engine/Source/Runtime/SlateCore/Private/Styling/StyleColors.cpp
--- a/Engine/Source/Runtime/SlateCore/Private/Styling/StyleColors.cpp
+++ b/Engine/Source/Runtime/SlateCore/Private/Styling/StyleColors.cpp
@@ -164,26 +164,7 @@ void USlateThemeManager::SaveCurrentThemeAs(const FString& Filename)
 	FString NewPath = CurrentTheme.Filename;
 	{
 		FString Output;
-		TSharedRef<TJsonWriter<>> WriterRef = TJsonWriterFactory<>::Create(&Output);
-		TJsonWriter<>& Writer = WriterRef.Get();
-		Writer.WriteObjectStart();
-		Writer.WriteValue(TEXT("Version"), 1);
-		Writer.WriteValue(TEXT("Id"), CurrentTheme.Id.ToString());
-		Writer.WriteValue(TEXT("DisplayName"), CurrentTheme.DisplayName.ToString());
-		
-		{
-			Writer.WriteObjectStart(TEXT("Colors"));
-			UEnum* Enum = StaticEnum<EStyleColor>();
-			check(Enum);
-			for (int32 ColorIndex = 0; ColorIndex < (int32)EStyleColor::MAX; ++ColorIndex)
-			{
-				FName EnumName = Enum->GetNameByIndex(ColorIndex);
-				Writer.WriteValue(EnumName.ToString(), *ActiveColors.StyleColors[ColorIndex].ToString());
-			}
-			Writer.WriteObjectEnd();
-		}
-		Writer.WriteObjectEnd();
-		Writer.Close();
+		WriteCurrentTheme(Output);
 
 		if (FPlatformFileManager::Get().GetPlatformFile().FileExists(*Filename))
 		{
@@ -202,6 +183,22 @@ void USlateThemeManager::SaveCurrentThemeAs(const FString& Filename)
 	
 }
 
+bool USlateThemeManager::ExportCurrentThemeTo(const FString& Filename)
+{
+	FString Output;
+	if (!WriteCurrentTheme(Output))
+	{
+		return false;
+	}
+
+	if (IPlatformFile::GetPlatformPhysical().FileExists(*Filename))
+	{
+		IPlatformFile::GetPlatformPhysical().DeleteFile(*Filename);
+	}
+
+	return FFileHelper::SaveStringToFile(Output, *Filename);
+}
+
 void USlateThemeManager::ApplyTheme(FGuid ThemeId)
 {
 	if (ThemeId.IsValid())
@@ -408,6 +405,34 @@ bool USlateThemeManager::ReadTheme(const FString& ThemeData, FStyleTheme& Theme)
 	return true;
 }
 
+bool USlateThemeManager::WriteCurrentTheme(FString& OutString)
+{
+	const FStyleTheme& CurrentTheme = GetCurrentTheme();
+
+	TSharedRef<TJsonWriter<>> WriterRef = TJsonWriterFactory<>::Create(&OutString);
+	TJsonWriter<>& Writer = WriterRef.Get();
+	Writer.WriteObjectStart();
+	Writer.WriteValue(TEXT("Version"), 1);
+	Writer.WriteValue(TEXT("Id"), CurrentTheme.Id.ToString());
+	Writer.WriteValue(TEXT("DisplayName"), CurrentTheme.DisplayName.ToString());
+
+	{
+		Writer.WriteObjectStart(TEXT("Colors"));
+		UEnum* Enum = StaticEnum<EStyleColor>();
+		check(Enum);
+		for (int32 ColorIndex = 0; ColorIndex < (int32)EStyleColor::MAX; ++ColorIndex)
+		{
+			FName EnumName = Enum->GetNameByIndex(ColorIndex);
+			Writer.WriteValue(EnumName.ToString(), *ActiveColors.StyleColors[ColorIndex].ToString());
+		}
+		Writer.WriteObjectEnd();
+	}
+	Writer.WriteObjectEnd();
+	Writer.Close();
+
+	return true;
+}
+
 void USlateThemeManager::EnsureValidCurrentTheme()
 {
 	DefaultDarkTheme.DisplayName = NSLOCTEXT("StyleColors", "DefaultDarkTheme", "Dark");
diff --git a/Engine/Source/Runtime/SlateCore/Public/Styling/StyleColors.h b/Engine/Source/Runtime/SlateCore/Public/Styling/StyleColors.h
--- a/Engine/Source/Runtime/SlateCore/Public/Styling/StyleColors.h
+++ b/Engine/Source/Runtime/SlateCore/Public/Styling/StyleColors.h
@@ -206,6 +206,16 @@ public:
 	 */
 	SLATECORE_API void SaveCurrentThemeAs(const FString& Filename);
 
+	/**
+	 * Exports the current theme to the given filename.
+	 *
+	 * The given filename will be replaced if it already exists.
+	 * 
+	 * @param The Filename to export the theme to
+	 * @return true if the theme was successfully exported
+	 */
+	SLATECORE_API bool ExportCurrentThemeTo(const FString& Filename);
+
 	/**
 	 * Applies a theme as the active theme
 	 */
@@ -297,6 +307,7 @@ private:
 	FStyleTheme& GetMutableCurrentTheme() { return *Themes.FindByKey(CurrentThemeId); }
 	SLATECORE_API void LoadThemesFromDirectory(const FString& Directory);
 	SLATECORE_API bool ReadTheme(const FString& ThemeData, FStyleTheme& OutTheme);
+	bool WriteCurrentTheme(FString& OutString);
 	SLATECORE_API void EnsureValidCurrentTheme();
 	SLATECORE_API void LoadThemeColors(FStyleTheme& Theme);
 
