[Truncation Warnings] Add a define to conditionally add final to virtual functions for slate APIs transitioning to FVector2f from FVector2D

diff --git a/Engine/Source/Runtime/SlateCore/Public/Types/SlateVector2.h b/Engine/Source/Runtime/SlateCore/Public/Types/SlateVector2.h
--- a/Engine/Source/Runtime/SlateCore/Public/Types/SlateVector2.h
+++ b/Engine/Source/Runtime/SlateCore/Public/Types/SlateVector2.h
@@ -48,6 +48,12 @@ static_assert(UE_ENABLE_SLATE_VECTOR_DEPRECATION_MECHANISMS || !UE_REPORT_SLATE_
 	#define UE_SLATE_VECTOR_DEPRECATED_DEFAULT()
 #endif
 
+#if UE_REPORT_SLATE_VECTOR_DEPRECATION && UE_WARNINGS_AS_ERRORS
+	#define UE_SLATE_DEPRECATED_VECTOR_VIRTUAL_FUNCTION final
+#else
+	#define UE_SLATE_DEPRECATED_VECTOR_VIRTUAL_FUNCTION
+#endif
+
 
 #undef UE_INCLUDETOOL_IGNORE_INCONSISTENT_STATE
 
