Missing log format %s or arguments

diff --git a/Engine/Source/Runtime/SlateCore/Private/Types/SlateAttributeDescriptor.cpp b/Engine/Source/Runtime/SlateCore/Private/Types/SlateAttributeDescriptor.cpp
--- a/Engine/Source/Runtime/SlateCore/Private/Types/SlateAttributeDescriptor.cpp
+++ b/Engine/Source/Runtime/SlateCore/Private/Types/SlateAttributeDescriptor.cpp
@@ -540,7 +540,7 @@ void FSlateAttributeDescriptor::OverrideInvalidationReason(FName ContainerName,
 	check(!AttributeName.IsNone());
 
 	FAttribute* FoundAttribute = FindAttribute(AttributeName);
-	if (ensureAlwaysMsgf(FoundAttribute != nullptr, TEXT("The attribute 's' doesn't exist."), *AttributeName.ToString()))
+	if (ensureAlwaysMsgf(FoundAttribute != nullptr, TEXT("The attribute '%s' doesn't exist."), *AttributeName.ToString()))
 	{
 		if (ensureAlwaysMsgf(FoundAttribute->ContainerName == ContainerName, TEXT("You are not allowed to modify an attribute from a different container.")))
 		{
@@ -555,7 +555,7 @@ void FSlateAttributeDescriptor::OverrideOnValueChanged(FName ContainerName, FNam
 	check(!AttributeName.IsNone());
 
 	FAttribute* FoundAttribute = FindAttribute(AttributeName);
-	if (ensureAlwaysMsgf(FoundAttribute != nullptr, TEXT("The attribute 's' doesn't exist."), *AttributeName.ToString()))
+	if (ensureAlwaysMsgf(FoundAttribute != nullptr, TEXT("The attribute '%s' doesn't exist."), *AttributeName.ToString()))
 	{
 		if (ensureAlwaysMsgf(FoundAttribute->ContainerName == ContainerName, TEXT("You are not allowed to modify an attribute from a different container.")))
 		{
