Remove csv stats from SlateSdfGenerator.

Primarily fixing a memory corruption issue causing invalid stats to be output.

#rnx

diff --git a/Engine/Source/Runtime/SlateCore/Private/Fonts/SlateSdfGenerator.cpp b/Engine/Source/Runtime/SlateCore/Private/Fonts/SlateSdfGenerator.cpp
--- a/Engine/Source/Runtime/SlateCore/Private/Fonts/SlateSdfGenerator.cpp
+++ b/Engine/Source/Runtime/SlateCore/Private/Fonts/SlateSdfGenerator.cpp
@@ -328,10 +328,7 @@ public:
 
 	void DoOutlineDecomposition();
 
-	void BeginCsvTrace(TSharedPtr<FFreeTypeFace> InFace, uint32 InGlyphIndex);
-	void EndCsvTrace();
-
-	FSlateSdfGenerator::ERequestResponse Prepare(const FSlateSdfGenerator::FRequestDescriptor& InDescriptor, FSlateSdfGenerator::FRequestOutputInfo& OutOutputInfo, bool bCsvTrace);
+	FSlateSdfGenerator::ERequestResponse Prepare(const FSlateSdfGenerator::FRequestDescriptor& InDescriptor, FSlateSdfGenerator::FRequestOutputInfo& OutOutputInfo);
 	void End(const FSlateSdfGenerator::FForEachRequestDoneCallback& Callback);
 	void Reset();
 
@@ -354,7 +351,7 @@ static int32 GetNumSdfChannels(FSlateSdfGenerator::ESdfType InSdfType)
 	return 0;
 }
 
-FSlateSdfGenerator::ERequestResponse FSdfGeneratorTask::Prepare(const FSlateSdfGenerator::FRequestDescriptor& InDescriptor, FSlateSdfGenerator::FRequestOutputInfo& OutOutputInfo, bool bCsvTrace)
+FSlateSdfGenerator::ERequestResponse FSdfGeneratorTask::Prepare(const FSlateSdfGenerator::FRequestDescriptor& InDescriptor, FSlateSdfGenerator::FRequestOutputInfo& OutOutputInfo)
 {
 	#if WITH_FREETYPE
 	TSharedPtr<FFreeTypeFace> FontFace = InDescriptor.FontFace.Pin();
@@ -379,10 +376,6 @@ FSlateSdfGenerator::ERequestResponse FSdfGeneratorTask::Prepare(const FSlateSdfG
 			OutOutputInfo.BearingX = GlyphSdfMapping.GetBearingX();
 			OutOutputInfo.BearingY = GlyphSdfMapping.GetBearingY();
 
-			if (bCsvTrace)
-			{
-				BeginCsvTrace(FontFace, InDescriptor.GlyphIndex);
-			}
 			return FSlateSdfGenerator::ERequestResponse::SUCCESS;
 		}
 	}
@@ -392,7 +385,6 @@ FSlateSdfGenerator::ERequestResponse FSdfGeneratorTask::Prepare(const FSlateSdfG
 
 void FSdfGeneratorTask::End(const FSlateSdfGenerator::FForEachRequestDoneCallback& Callback)
 {
-	EndCsvTrace();
 	Callback(Descriptor, MoveTemp(OutputPixels));
 	Descriptor = {};
 	MsdfgenShape.contours.clear();
@@ -400,53 +392,11 @@ void FSdfGeneratorTask::End(const FSlateSdfGenerator::FForEachRequestDoneCallbac
 
 void FSdfGeneratorTask::Reset()
 {
-	EndCsvTrace();
 	Descriptor = {};
 	MsdfgenShape.contours.clear();
 	OutputPixels.Empty();
 }
 
-void FSdfGeneratorTask::BeginCsvTrace(TSharedPtr<FFreeTypeFace> InFace, uint32 InGlyphIndex)
-{
-	#if SLATE_CSV_TRACKER
-	if (DebugName.IsEmpty())
-	{
-		#if WITH_FREETYPE
-		DebugName = FString("Sdf/");
-		DebugName += FString(InFace->GetFace()->family_name ? InFace->GetFace()->family_name : "FontNotDef");
-		DebugName += FString("/");
-		char GlyphDebugName[16];
-		GlyphDebugName[0] = GlyphDebugName[15] = '\0';
-		if (FT_HAS_GLYPH_NAMES(InFace->GetFace())
-			&& 0 == ::FT_Get_Glyph_Name(InFace->GetFace(), InGlyphIndex, GlyphDebugName, sizeof(GlyphDebugName))
-			&& FCStringAnsi::Strnlen(GlyphDebugName, sizeof(GlyphDebugName)) > 0)
-		{
-			DebugName += FString(GlyphDebugName);
-		}
-		else
-		{
-			DebugName += FString::FromInt(InGlyphIndex);
-		}
-		#else
-		DebugName = FString("Sdf/FontNotDef/");
-		DebugName += FString::FromInt(InGlyphIndex);
-		#endif
-		FCsvProfiler::BeginStat(TCHAR_TO_ANSI(*DebugName), CSV_CATEGORY_INDEX(Slate));
-	}
-	#endif // SLATE_CSV_TRACKER
-}
-
-void FSdfGeneratorTask::EndCsvTrace()
-{
-	#if SLATE_CSV_TRACKER
-	if (!DebugName.IsEmpty())
-	{
-		FCsvProfiler::EndStat(TCHAR_TO_ANSI(*DebugName), CSV_CATEGORY_INDEX(Slate));
-		DebugName.Empty();
-	}
-	#endif
-}
-
 void FSdfGeneratorTask::DoWork()
 {
 	#ifdef SLATE_LOG_SDF_GENERATOR_TIME
@@ -650,7 +600,7 @@ FSlateSdfGenerator::ERequestResponse FSlateSdfGeneratorImpl::Spawn(const FReques
 		return ERequestResponse::BUSY;
 	}
 	FAsyncTask<SdfUtils::FSdfGeneratorTask>* Task = FreeTasks.Pop(EAllowShrinking::No);
-	const ERequestResponse Result = Task->GetTask().Prepare(InRequest, OutCharInfo, true);
+	const ERequestResponse Result = Task->GetTask().Prepare(InRequest, OutCharInfo);
 	if (Result == ERequestResponse::SUCCESS)
 	{
 		StartedTasks.Push(Task);
@@ -675,7 +625,7 @@ FSlateSdfGenerator::ERequestResponse FSlateSdfGeneratorImpl::SpawnWithPlaceholde
 		return Result;
 	}
 	FAsyncTask<SdfUtils::FSdfGeneratorTask>* Task = FreeTasks.Pop(EAllowShrinking::No);
-	const ERequestResponse Result = Task->GetTask().Prepare(InRequest, OutCharInfo, true);
+	const ERequestResponse Result = Task->GetTask().Prepare(InRequest, OutCharInfo);
 	if (Result == ERequestResponse::SUCCESS)
 	{
 		Task->GetTask().MakePlaceholder(OutRawPixels);
@@ -697,7 +647,7 @@ FSlateSdfGenerator::ERequestResponse FSlateSdfGeneratorImpl::Respawn(const FRequ
 	}
 	FRequestOutputInfo OutCharInfo = {};
 	FAsyncTask<SdfUtils::FSdfGeneratorTask>* Task = FreeTasks.Pop(EAllowShrinking::No);
-	ERequestResponse Result = Task->GetTask().Prepare(InRequest, OutCharInfo, true);
+	ERequestResponse Result = Task->GetTask().Prepare(InRequest, OutCharInfo);
 	if (Result == ERequestResponse::SUCCESS)
 	{
 		if (
@@ -725,7 +675,7 @@ FSlateSdfGenerator::ERequestResponse FSlateSdfGeneratorImpl::Respawn(const FRequ
 FSlateSdfGenerator::ERequestResponse FSlateSdfGeneratorImpl::MakePlaceholder(const FRequestDescriptor& InRequest, FRequestOutputInfo& OutCharInfo, TArray<uint8>& OutRawPixels)
 {
 	SdfUtils::FSdfGeneratorTask PlaceholderTask;
-	const ERequestResponse Result = PlaceholderTask.Prepare(InRequest, OutCharInfo, false);
+	const ERequestResponse Result = PlaceholderTask.Prepare(InRequest, OutCharInfo);
 	if (Result == ERequestResponse::SUCCESS)
 	{
 		PlaceholderTask.MakePlaceholder(OutRawPixels);
