Fixed hit test grid not updated properly in retainers with slow path rendering.
When using the slow path, the hit test grid was keeping all its content, resulting in no longer relevant elements to stay in the grid, that would catch the inputs. Now with the fix, the hit test grid will be cleared for retainers in the slow path, like it is done already in SWindow::PaintSlowPath.

#tests Ran the repro steps of the 2 different UDNs with similar issues. Bug no longer occurs.
[RN] Slate: fixed a bug with retainer widget where the hit test grid could keep no longer hit test testable elements, resulting in incorrect input processing.

diff --git a/Engine/Source/Runtime/SlateCore/Private/FastUpdate/SlateInvalidationRoot.cpp b/Engine/Source/Runtime/SlateCore/Private/FastUpdate/SlateInvalidationRoot.cpp
--- a/Engine/Source/Runtime/SlateCore/Private/FastUpdate/SlateInvalidationRoot.cpp
+++ b/Engine/Source/Runtime/SlateCore/Private/FastUpdate/SlateInvalidationRoot.cpp
@@ -1676,7 +1676,7 @@ void VerifyHittest(SWidget* InvalidationRootWidget, FSlateInvalidationWidgetList
 			{
 				UE_SLATE_LOG_ERROR_IF_FALSE(WidgetList.IsValidIndex(ParentWidetIndex) && EnumHasAllFlags(WidgetList[ParentWidetIndex].CurrentInvalidateReason, EInvalidateWidgetReason::ChildOrder)
 					, CVarSlateInvalidationRootVerifyHittestGrid
-					, TEXT("The parent widget doesn't has ChildOrder invaldiation. The hittest grid has an invalid widget."));
+					, TEXT("The parent widget doesn't has ChildOrder invalidation. The hittest grid has an invalid widget."));
 			}
 		}
 	}
diff --git a/Engine/Source/Runtime/UMG/Private/Slate/SRetainerWidget.cpp b/Engine/Source/Runtime/UMG/Private/Slate/SRetainerWidget.cpp
--- a/Engine/Source/Runtime/UMG/Private/Slate/SRetainerWidget.cpp
+++ b/Engine/Source/Runtime/UMG/Private/Slate/SRetainerWidget.cpp
@@ -745,6 +745,8 @@ TSharedRef<SWidget> SRetainerWidget::GetRootWidget()
 
 int32 SRetainerWidget::PaintSlowPath(const FSlateInvalidationContext& Context)
 {
+	HittestGrid->Clear();
+
 	const FGeometry& OriginalPaintSpaceGeometry = GetPaintSpaceGeometry();
 	FSlateRenderTransform SimplifiedRenderTransform(OriginalPaintSpaceGeometry.GetAccumulatedRenderTransform().GetMatrix().GetScale(), OriginalPaintSpaceGeometry.GetAccumulatedRenderTransform().GetTranslation());
 	FTransform2D ScalePosTransform(OriginalPaintSpaceGeometry.Scale, OriginalPaintSpaceGeometry.AbsolutePosition);
