Deprecating SDragAndDropVerticalBox::CanAcceptDrop. Adds new CanAcceptDrop which adds a slot index parameter.

diff --git a/Engine/Source/Editor/LandscapeEditor/Private/LandscapeEditorDetailCustomization_Layers.cpp b/Engine/Source/Editor/LandscapeEditor/Private/LandscapeEditorDetailCustomization_Layers.cpp
--- a/Engine/Source/Editor/LandscapeEditor/Private/LandscapeEditorDetailCustomization_Layers.cpp
+++ b/Engine/Source/Editor/LandscapeEditor/Private/LandscapeEditorDetailCustomization_Layers.cpp
@@ -173,7 +173,7 @@ void FLandscapeEditorCustomNodeBuilder_Layers::GenerateChildContent(IDetailChild
 	if (LandscapeEdMode != NULL)
 	{
 		TSharedPtr<SDragAndDropVerticalBox> LayerList = SNew(SDragAndDropVerticalBox)
-			.OnCanAcceptDrop(this, &FLandscapeEditorCustomNodeBuilder_Layers::HandleCanAcceptDrop)
+			.OnCanAcceptDropAdvanced(this, &FLandscapeEditorCustomNodeBuilder_Layers::HandleCanAcceptDrop)
 			.OnAcceptDrop(this, &FLandscapeEditorCustomNodeBuilder_Layers::HandleAcceptDrop)
 			.OnDragDetected(this, &FLandscapeEditorCustomNodeBuilder_Layers::HandleDragDetected);
 
@@ -1299,7 +1299,7 @@ FReply FLandscapeEditorCustomNodeBuilder_Layers::HandleDragDetected(const FGeome
 	return FReply::Unhandled();
 }
 
-TOptional<SDragAndDropVerticalBox::EItemDropZone> FLandscapeEditorCustomNodeBuilder_Layers::HandleCanAcceptDrop(const FDragDropEvent& DragDropEvent, SDragAndDropVerticalBox::EItemDropZone DropZone, SVerticalBox::FSlot* Slot)
+TOptional<SDragAndDropVerticalBox::EItemDropZone> FLandscapeEditorCustomNodeBuilder_Layers::HandleCanAcceptDrop(const FDragDropEvent& DragDropEvent, SDragAndDropVerticalBox::EItemDropZone DropZone, int32 SlotIndex, SVerticalBox::FSlot* Slot)
 {
 	TSharedPtr<FLandscapeListElementDragDropOp> DragDropOperation = DragDropEvent.GetOperationAs<FLandscapeListElementDragDropOp>();
 	if (DragDropOperation.IsValid())
diff --git a/Engine/Source/Editor/LandscapeEditor/Private/LandscapeEditorDetailCustomization_Layers.h b/Engine/Source/Editor/LandscapeEditor/Private/LandscapeEditorDetailCustomization_Layers.h
--- a/Engine/Source/Editor/LandscapeEditor/Private/LandscapeEditorDetailCustomization_Layers.h
+++ b/Engine/Source/Editor/LandscapeEditor/Private/LandscapeEditorDetailCustomization_Layers.h
@@ -71,7 +71,7 @@ protected:
 	// Drag/Drop handling
 	int32 SlotIndexToLayerIndex(int32 SlotIndex);
 	FReply HandleDragDetected(const FGeometry& MyGeometry, const FPointerEvent& MouseEvent, int32 SlotIndex, SVerticalBox::FSlot* Slot);
-	TOptional<SDragAndDropVerticalBox::EItemDropZone> HandleCanAcceptDrop(const FDragDropEvent& DragDropEvent, SDragAndDropVerticalBox::EItemDropZone DropZone, SVerticalBox::FSlot* Slot);
+	TOptional<SDragAndDropVerticalBox::EItemDropZone> HandleCanAcceptDrop(const FDragDropEvent& DragDropEvent, SDragAndDropVerticalBox::EItemDropZone DropZone, int32 SlotIndex, SVerticalBox::FSlot* Slot);
 	FReply HandleAcceptDrop(FDragDropEvent const& DragDropEvent, SDragAndDropVerticalBox::EItemDropZone DropZone, int32 SlotIndex, SVerticalBox::FSlot* Slot);
 
 	bool IsLayerSelected(int32 LayerIndex) const;
diff --git a/Engine/Source/Editor/LandscapeEditor/Private/LandscapeEditorDetailCustomization_LayersBrushStack.cpp b/Engine/Source/Editor/LandscapeEditor/Private/LandscapeEditorDetailCustomization_LayersBrushStack.cpp
--- a/Engine/Source/Editor/LandscapeEditor/Private/LandscapeEditorDetailCustomization_LayersBrushStack.cpp
+++ b/Engine/Source/Editor/LandscapeEditor/Private/LandscapeEditorDetailCustomization_LayersBrushStack.cpp
@@ -85,7 +85,7 @@ void FLandscapeEditorCustomNodeBuilder_LayersBrushStack::GenerateChildContent(ID
 	if (FEdModeLandscape* LandscapeEdMode = GetEditorMode())
 	{
 		TSharedPtr<SDragAndDropVerticalBox> BrushesList = SNew(SDragAndDropVerticalBox)
-			.OnCanAcceptDrop(this, &FLandscapeEditorCustomNodeBuilder_LayersBrushStack::HandleCanAcceptDrop)
+			.OnCanAcceptDropAdvanced(this, &FLandscapeEditorCustomNodeBuilder_LayersBrushStack::HandleCanAcceptDrop)
 			.OnAcceptDrop(this, &FLandscapeEditorCustomNodeBuilder_LayersBrushStack::HandleAcceptDrop)
 			.OnDragDetected(this, &FLandscapeEditorCustomNodeBuilder_LayersBrushStack::HandleDragDetected);
 
@@ -520,7 +520,7 @@ FReply FLandscapeEditorCustomNodeBuilder_LayersBrushStack::HandleDragDetected(co
 	return FReply::Unhandled();
 }
 
-TOptional<SDragAndDropVerticalBox::EItemDropZone> FLandscapeEditorCustomNodeBuilder_LayersBrushStack::HandleCanAcceptDrop(const FDragDropEvent& DragDropEvent, SDragAndDropVerticalBox::EItemDropZone DropZone, SVerticalBox::FSlot* Slot)
+TOptional<SDragAndDropVerticalBox::EItemDropZone> FLandscapeEditorCustomNodeBuilder_LayersBrushStack::HandleCanAcceptDrop(const FDragDropEvent& DragDropEvent, SDragAndDropVerticalBox::EItemDropZone DropZone, int32 SlotIndex, SVerticalBox::FSlot* Slot)
 {
 	TSharedPtr<FLandscapeBrushDragDropOp> DragDropOperation = DragDropEvent.GetOperationAs<FLandscapeBrushDragDropOp>();
 
diff --git a/Engine/Source/Editor/LandscapeEditor/Private/LandscapeEditorDetailCustomization_LayersBrushStack.h b/Engine/Source/Editor/LandscapeEditor/Private/LandscapeEditorDetailCustomization_LayersBrushStack.h
--- a/Engine/Source/Editor/LandscapeEditor/Private/LandscapeEditorDetailCustomization_LayersBrushStack.h
+++ b/Engine/Source/Editor/LandscapeEditor/Private/LandscapeEditorDetailCustomization_LayersBrushStack.h
@@ -91,7 +91,7 @@ protected:
 
 	// Drag/Drop handling
 	FReply HandleDragDetected(const FGeometry& MyGeometry, const FPointerEvent& MouseEvent, int32 SlotIndex, SVerticalBox::FSlot* Slot);
-	TOptional<SDragAndDropVerticalBox::EItemDropZone> HandleCanAcceptDrop(const FDragDropEvent& DragDropEvent, SDragAndDropVerticalBox::EItemDropZone DropZone, SVerticalBox::FSlot* Slot);
+	TOptional<SDragAndDropVerticalBox::EItemDropZone> HandleCanAcceptDrop(const FDragDropEvent& DragDropEvent, SDragAndDropVerticalBox::EItemDropZone DropZone, int32 SlotIndex, SVerticalBox::FSlot* Slot);
 	FReply HandleAcceptDrop(FDragDropEvent const& DragDropEvent, SDragAndDropVerticalBox::EItemDropZone DropZone, int32 SlotIndex, SVerticalBox::FSlot* Slot);
 };
 
diff --git a/Engine/Source/Editor/LandscapeEditor/Private/LandscapeEditorDetailCustomization_TargetLayers.cpp b/Engine/Source/Editor/LandscapeEditor/Private/LandscapeEditorDetailCustomization_TargetLayers.cpp
--- a/Engine/Source/Editor/LandscapeEditor/Private/LandscapeEditorDetailCustomization_TargetLayers.cpp
+++ b/Engine/Source/Editor/LandscapeEditor/Private/LandscapeEditorDetailCustomization_TargetLayers.cpp
@@ -976,7 +976,7 @@ void FLandscapeEditorCustomNodeBuilder_TargetLayers::GenerateChildContent(IDetai
 	if (LandscapeEdMode != nullptr)
 	{
 		TSharedPtr<SDragAndDropVerticalBox> TargetLayerList = SNew(SDragAndDropVerticalBox)
-			.OnCanAcceptDrop(this, &FLandscapeEditorCustomNodeBuilder_TargetLayers::HandleCanAcceptDrop)
+			.OnCanAcceptDropAdvanced(this, &FLandscapeEditorCustomNodeBuilder_TargetLayers::HandleCanAcceptDrop)
 			.OnAcceptDrop(this, &FLandscapeEditorCustomNodeBuilder_TargetLayers::HandleAcceptDrop)
 			.OnDragDetected(this, &FLandscapeEditorCustomNodeBuilder_TargetLayers::HandleDragDetected)
 			.IsEnabled(TAttribute<bool>::Create(TAttribute<bool>::FGetter::CreateLambda([LandscapeEdMode]() { return LandscapeEdMode->HasValidLandscapeEditLayerSelection(); })));
@@ -1406,7 +1406,7 @@ FReply FLandscapeEditorCustomNodeBuilder_TargetLayers::HandleDragDetected(const
 	return FReply::Unhandled();
 }
 
-TOptional<SDragAndDropVerticalBox::EItemDropZone> FLandscapeEditorCustomNodeBuilder_TargetLayers::HandleCanAcceptDrop(const FDragDropEvent& DragDropEvent, SDragAndDropVerticalBox::EItemDropZone DropZone, SVerticalBox::FSlot* Slot)
+TOptional<SDragAndDropVerticalBox::EItemDropZone> FLandscapeEditorCustomNodeBuilder_TargetLayers::HandleCanAcceptDrop(const FDragDropEvent& DragDropEvent, SDragAndDropVerticalBox::EItemDropZone DropZone, int32 SlotIndex, SVerticalBox::FSlot* Slot)
 {
 	TSharedPtr<FTargetLayerDragDropOp> DragDropOperation = DragDropEvent.GetOperationAs<FTargetLayerDragDropOp>();
 
diff --git a/Engine/Source/Editor/LandscapeEditor/Private/LandscapeEditorDetailCustomization_TargetLayers.h b/Engine/Source/Editor/LandscapeEditor/Private/LandscapeEditorDetailCustomization_TargetLayers.h
--- a/Engine/Source/Editor/LandscapeEditor/Private/LandscapeEditorDetailCustomization_TargetLayers.h
+++ b/Engine/Source/Editor/LandscapeEditor/Private/LandscapeEditorDetailCustomization_TargetLayers.h
@@ -124,7 +124,7 @@ protected:
 
 	// Drag/Drop handling
 	FReply HandleDragDetected(const FGeometry& MyGeometry, const FPointerEvent& MouseEvent, int32 SlotIndex, SVerticalBox::FSlot* Slot);
-	TOptional<SDragAndDropVerticalBox::EItemDropZone> HandleCanAcceptDrop(const FDragDropEvent& DragDropEvent, SDragAndDropVerticalBox::EItemDropZone DropZone, SVerticalBox::FSlot* Slot);
+	TOptional<SDragAndDropVerticalBox::EItemDropZone> HandleCanAcceptDrop(const FDragDropEvent& DragDropEvent, SDragAndDropVerticalBox::EItemDropZone DropZone, int32 SlotIndex, SVerticalBox::FSlot* Slot);
 	FReply HandleAcceptDrop(FDragDropEvent const& DragDropEvent, SDragAndDropVerticalBox::EItemDropZone DropZone, int32 SlotIndex, SVerticalBox::FSlot* Slot);
 
 	// Create Layers From Materials
diff --git a/Engine/Source/Runtime/SlateCore/Private/Widgets/SBoxPanel.cpp b/Engine/Source/Runtime/SlateCore/Private/Widgets/SBoxPanel.cpp
--- a/Engine/Source/Runtime/SlateCore/Private/Widgets/SBoxPanel.cpp
+++ b/Engine/Source/Runtime/SlateCore/Private/Widgets/SBoxPanel.cpp
@@ -245,7 +245,7 @@ void SDragAndDropVerticalBox::Construct(const FArguments& InArgs)
 {
 	SVerticalBox::Construct(SVerticalBox::FArguments());
 
-	OnCanAcceptDrop = InArgs._OnCanAcceptDrop;
+	OnCanAcceptDropAdvanced = InArgs._OnCanAcceptDropAdvanced;
 	OnAcceptDrop = InArgs._OnAcceptDrop;
 	OnDragDetected_Handler = InArgs._OnDragDetected;
 	OnDragEnter_Handler = InArgs._OnDragEnter;
@@ -333,7 +333,7 @@ SDragAndDropVerticalBox::EItemDropZone SDragAndDropVerticalBox::ZoneFromPointerP
 
 FReply SDragAndDropVerticalBox::OnDragOver(const FGeometry& MyGeometry, const FDragDropEvent& DragDropEvent)
 {
-	if (OnCanAcceptDrop.IsBound())
+	if (OnCanAcceptDropAdvanced.IsBound())
 	{
 		FArrangedChildren ArrangedChildren(EVisibility::Visible);
 		ArrangeChildren(MyGeometry, ArrangedChildren);
@@ -353,7 +353,7 @@ FReply SDragAndDropVerticalBox::OnDragOver(const FGeometry& MyGeometry, const FD
 				{
 					SVerticalBox::FSlot* Slot = static_cast<SVerticalBox::FSlot*>(&static_cast<FSlotBase&>(Children[DragOverSlotIndex]));
 
-					ItemDropZone = OnCanAcceptDrop.Execute(DragDropEvent, ItemHoverZone, Slot);
+					ItemDropZone = OnCanAcceptDropAdvanced.Execute(DragDropEvent, ItemHoverZone, DragOverSlotIndex, Slot);
 					CurrentDragOperationScreenSpaceLocation = DragDropEvent.GetScreenSpacePosition();
 					CurrentDragOverSlotIndex = DragOverSlotIndex;
 
@@ -383,9 +383,9 @@ FReply SDragAndDropVerticalBox::OnDrop(const FGeometry& MyGeometry, const FDragD
 				SVerticalBox::FSlot* Slot = static_cast<SVerticalBox::FSlot*>(&static_cast<FSlotBase&>(Children[NodeUnderPositionIndex]));
 				TOptional<EItemDropZone> ReportedZone = ItemDropZone;
 
-				if (OnCanAcceptDrop.IsBound() && ItemDropZone.IsSet())
+				if (OnCanAcceptDropAdvanced.IsBound() && ItemDropZone.IsSet())
 				{
-					ReportedZone = OnCanAcceptDrop.Execute(DragDropEvent, ItemDropZone.GetValue(), Slot);
+					ReportedZone = OnCanAcceptDropAdvanced.Execute(DragDropEvent, ItemDropZone.GetValue(), NodeUnderPositionIndex, Slot);
 				}
 
 				if (ReportedZone.IsSet())
diff --git a/Engine/Source/Runtime/SlateCore/Public/Widgets/SBoxPanel.h b/Engine/Source/Runtime/SlateCore/Public/Widgets/SBoxPanel.h
--- a/Engine/Source/Runtime/SlateCore/Public/Widgets/SBoxPanel.h
+++ b/Engine/Source/Runtime/SlateCore/Public/Widgets/SBoxPanel.h
@@ -580,15 +580,33 @@ public:
 	DECLARE_DELEGATE_OneParam(FOnDragAndDropVerticalBoxDragLeave, FDragDropEvent const&);
 	DECLARE_DELEGATE_RetVal_OneParam(FReply, FOnDragAndDropVerticalBoxDrop, FDragDropEvent const&);
 
-	/** Delegate signature for querying whether this FDragDropEvent will be handled by the drop target of type ItemType. */
+	UE_DEPRECATED(5.6, "FOnCanAcceptDrop is deprecated. Please use FOnCanAcceptDropAdvanced which adds the drop slot index parameter.")
 	DECLARE_DELEGATE_RetVal_ThreeParams(TOptional<EItemDropZone>, FOnCanAcceptDrop, const FDragDropEvent&, SDragAndDropVerticalBox::EItemDropZone, SVerticalBox::FSlot*);
+	/** Delegate signature for querying whether this FDragDropEvent will be handled by the drop target of type ItemType. */
+	DECLARE_DELEGATE_RetVal_FourParams(TOptional<EItemDropZone>, FOnCanAcceptDropAdvanced, const FDragDropEvent&, SDragAndDropVerticalBox::EItemDropZone, int32, SVerticalBox::FSlot*);
 	/** Delegate signature for handling the drop of FDragDropEvent onto target of type ItemType */
 	DECLARE_DELEGATE_RetVal_FourParams(FReply, FOnAcceptDrop, const FDragDropEvent&, SDragAndDropVerticalBox::EItemDropZone, int32, SVerticalBox::FSlot*);
 
 	SLATE_BEGIN_ARGS(SDragAndDropVerticalBox)
 		{}
 
-		// High Level DragAndDrop
+
+PRAGMA_DISABLE_DEPRECATION_WARNINGS
+		FOnCanAcceptDropAdvanced ConvertCanAcceptDropFn(const FOnCanAcceptDrop& LegacyDelegate)
+		{
+			return FOnCanAcceptDropAdvanced::CreateLambda([LegacyDelegate](const FDragDropEvent& Event, SDragAndDropVerticalBox::EItemDropZone DropZone, int32 SlotIndex, SVerticalBox::FSlot* Slot)
+				{
+					if (LegacyDelegate.IsBound())
+					{
+						return LegacyDelegate.Execute(Event, DropZone, Slot);
+					}
+
+					return TOptional<EItemDropZone>();
+				});
+		}
+
+		SLATE_EVENT_DEPRECATED(5.6, "Use FOnCanAcceptDropAdvanced instead", FOnCanAcceptDrop, OnCanAcceptDrop, OnCanAcceptDropAdvanced, ConvertCanAcceptDropFn)
+PRAGMA_ENABLE_DEPRECATION_WARNINGS
 
 		/**
 		* Handle this event to determine whether a drag and drop operation can be executed on top of the target row widget.
@@ -597,7 +615,7 @@ public:
 		*      This delegate will be called to figure out if we should give visual feedback on whether an item will
 		*      successfully drop into the list.
 		*/
-		SLATE_EVENT(FOnCanAcceptDrop, OnCanAcceptDrop)
+		SLATE_EVENT(FOnCanAcceptDropAdvanced, OnCanAcceptDropAdvanced)
 
 		/**
 		* Perform a drop operation onto the target row widget
@@ -632,8 +650,8 @@ public:
 	SLATECORE_API virtual int32 OnPaint(const FPaintArgs& Args, const FGeometry& AllottedGeometry, const FSlateRect& MyCullingRect, FSlateWindowElementList& OutDrawElements, int32 LayerId, const FWidgetStyle& InWidgetStyle, bool bParentEnabled) const override;
 
 private:
-	/** @see SDragAndDropVerticalBox's OnCanAcceptDrop event */
-	FOnCanAcceptDrop OnCanAcceptDrop;
+	/** @see SDragAndDropVerticalBox's OnCanAcceptDropAdvanced event */
+	FOnCanAcceptDropAdvanced OnCanAcceptDropAdvanced;
 
 	/** @see SDragAndDropVerticalBox's OnAcceptDrop event */
 	FOnAcceptDrop OnAcceptDrop;
