Fix graphical issues when displaying some glyphs when is SDF activated.
The preprocessed geometry was not written in the cook file (and we didn't try to read it neither), so glyphs with contours imcompatible with the SDF algo where displayed incorrectly in game (but displayed without issue in editor, as preprocessd geometry was working as expected in this case).

#tests Cooked game: validated that the bug no longer occurs in the lobby, also tested in game and validated that fonts are loaded correctly, as Slate.Font.AsyncLazyLoad is activated in some experiences. No regression in the editor when loading fonts.
[RN] Fixed glyphs that may be displayed incorrectly with SDF activated.

diff --git a/Engine/Source/Runtime/Engine/Private/FontFace.cpp b/Engine/Source/Runtime/Engine/Private/FontFace.cpp
--- a/Engine/Source/Runtime/Engine/Private/FontFace.cpp
+++ b/Engine/Source/Runtime/Engine/Private/FontFace.cpp
@@ -196,7 +196,11 @@ void UFontFace::CookAdditionalFilesOverride(const TCHAR* PackageFilename, const
 		// We replace the package name with the cooked font face name
 		// Note: This must match the replacement logic in UFontFace::GetCookedFilename
 		const FString CookedFontFilename = FPaths::GetPath(PackageFilename) / GetName() + TEXT(".ufont");
-		const TArray<uint8>& Data = FontFaceData->GetData();
+		
+		TArray<uint8> Data;
+		FMemoryWriter Ar(Data, true);
+		FontFaceData->Serialize(Ar);
+		
 		const int32 NumBytes = Data.Num() * Data.GetTypeSize();
 		WriteAdditionalFile(*CookedFontFilename, (void*)Data.GetData(), NumBytes);
 	}
diff --git a/Engine/Source/Runtime/SlateCore/Private/Fonts/FontCacheCompositeFont.cpp b/Engine/Source/Runtime/SlateCore/Private/Fonts/FontCacheCompositeFont.cpp
--- a/Engine/Source/Runtime/SlateCore/Private/Fonts/FontCacheCompositeFont.cpp
+++ b/Engine/Source/Runtime/SlateCore/Private/Fonts/FontCacheCompositeFont.cpp
@@ -11,13 +11,14 @@
 #include "Application/SlateApplicationBase.h"
 #include "Async/Async.h"
 #include "HAL/PlatformProcess.h"
+#include "Serialization/MemoryReader.h"
 #include "Trace/SlateMemoryTags.h"
 
 static bool GAsyncFontLazyLoad = false;
 FAutoConsoleVariableRef CVarAsyncLazyLoad(
 	TEXT("Slate.Font.AsyncLazyLoad"),
 	GAsyncFontLazyLoad,
-	TEXT("Causes unloaded font faces that are lazily loaded, to be loaded asynchronusly, until then the font won't measure correctly.  Once complete the UI will invalidate."));
+	TEXT("Causes unloaded font faces that are lazily loaded, to be loaded asynchronously, until then the font won't measure correctly.  Once complete the UI will invalidate."));
 
 DECLARE_CYCLE_STAT(TEXT("Load Font"), STAT_SlateLoadFont, STATGROUP_Slate);
 
@@ -92,6 +93,22 @@ void FCachedTypefaceData::GetCachedFontData(TArray<const FFontData*>& OutFontDat
 	}
 }
 
+namespace Private
+{
+	bool IsFontFileName(const FString& Filename)
+	{	
+		return Filename.EndsWith(".ttf") || Filename.EndsWith(".ttc") || Filename.EndsWith(".otf") || Filename.EndsWith(".otc") || Filename.EndsWith(".woff") || Filename.EndsWith(".woff2");
+	}
+
+	bool IsFontFileData(const TArray<uint8>& FontFaceDataArray)
+	{	//Detect ttf, otf, ttc/otc, woff1 and woff2
+		return FontFaceDataArray.Num() >= 4 &&	((FontFaceDataArray[0] == 0x00 && FontFaceDataArray[1] == 0x01 && FontFaceDataArray[2] == 0x00 && FontFaceDataArray[3] == 0x00) ||
+												(FontFaceDataArray[0] == 0x4f && FontFaceDataArray[1] == 0x54 && FontFaceDataArray[2] == 0x54 && FontFaceDataArray[3] == 0x4f) ||
+												(FontFaceDataArray[0] == 0x74 && FontFaceDataArray[1] == 0x74 && FontFaceDataArray[2] == 0x63 && FontFaceDataArray[3] == 0x66) ||	//TTC
+												(FontFaceDataArray[0] == 0x77 && FontFaceDataArray[1] == 0x4f && FontFaceDataArray[2] == 0x46 && FontFaceDataArray[3] == 0x46) ||	//WOFF1
+												(FontFaceDataArray[0] == 0x77 && FontFaceDataArray[1] == 0x4f && FontFaceDataArray[2] == 0x46 && FontFaceDataArray[3] == 0x32));	//WOFF2
+	}
+}
 
 FCachedCompositeFontData::FCachedCompositeFontData()
 	: CompositeFont(nullptr)
@@ -505,13 +522,25 @@ TSharedPtr<FFreeTypeFace> FCompositeFontCache::GetFontFace(const FFontData& InFo
 					{
 						const double FontDataLoadStartTime = FPlatformTime::Seconds();
 
-						TArray<uint8> FontFaceData;
-						if (FFileHelper::LoadFileToArray(FontFaceData, *InFontData.GetFontFilename()))
+						TArray<uint8> FontFaceDataArray;
+						const FString& FontFilename = InFontData.GetFontFilename();
+						if (FFileHelper::LoadFileToArray(FontFaceDataArray, *FontFilename))
 						{
-							const int32 FontDataSizeKB = (FontFaceData.Num() + 1023) / 1024;
+							const int32 FontDataSizeKB = (FontFaceDataArray.Num() + 1023) / 1024;
 							LoadLogMessage = FString::Printf(TEXT("Took %f seconds to synchronously load lazily loaded font '%s' (%dK)"), float(FPlatformTime::Seconds() - FontDataLoadStartTime), *InFontData.GetFontFilename(), FontDataSizeKB);
 
-							FaceAndMemory = MakeShared<FFreeTypeFace>(FTLibrary, FFontFaceData::MakeFontFaceData(MoveTemp(FontFaceData)), InFontData.GetSubFaceIndex(), InFontData.GetLayoutMethod());
+							if (Private::IsFontFileName(FontFilename) || Private::IsFontFileData(FontFaceDataArray))
+							{
+								FaceAndMemory = MakeShared<FFreeTypeFace>(FTLibrary, FFontFaceData::MakeFontFaceData(MoveTemp(FontFaceDataArray)), InFontData.GetSubFaceIndex(), InFontData.GetLayoutMethod());
+							}
+							else
+							{
+								FFontFaceDataRef FontFaceData = FFontFaceData::MakeFontFaceData();
+								FMemoryReader Ar(FontFaceDataArray, true);
+								FontFaceData->Serialize(Ar);
+
+								FaceAndMemory = MakeShared<FFreeTypeFace>(FTLibrary, FontFaceData, InFontData.GetSubFaceIndex(), InFontData.GetLayoutMethod());
+							}
 						}
 					}
 				}
@@ -582,15 +611,27 @@ void FCompositeFontCache::Update()
 			FAsyncTask<FAsyncLoadFontFaceData>* LoadFontFaceTask = LoadFontFaceTasks[LoadTaskIndex];
 			if (LoadFontFaceTask->IsDone())
 			{
-				TArray<uint8>& FontFaceData = LoadFontFaceTask->GetTask().GetFontFaceData();
+				TArray<uint8>& FontFaceDataArray = LoadFontFaceTask->GetTask().GetFontFaceData();
+				const FString& FontFilename = LoadFontFaceTask->GetTask().GetFontData().GetFontFilename();
 
 				const FFontData& FontData = LoadFontFaceTask->GetTask().GetFontData();
 				TSharedPtr<FFreeTypeFace> FaceAndMemory = FontFaceMap.FindRef(FontData);
 				if (FaceAndMemory.IsValid())
 				{
-					if (FontFaceData.Num() > 0)
+					if (FontFaceDataArray.Num() > 0)
 					{
-						FaceAndMemory->CompleteAsyncLoad(FTLibrary, FFontFaceData::MakeFontFaceData(MoveTemp(FontFaceData)), FontData.GetSubFaceIndex());
+						if (Private::IsFontFileName(FontFilename) || Private::IsFontFileData(FontFaceDataArray))
+						{
+							FaceAndMemory->CompleteAsyncLoad(FTLibrary, FFontFaceData::MakeFontFaceData(MoveTemp(FontFaceDataArray)), FontData.GetSubFaceIndex());
+						}
+						else
+						{
+							FFontFaceDataRef FontFaceData = FFontFaceData::MakeFontFaceData();
+							FMemoryReader Ar(FontFaceDataArray, true);
+							FontFaceData->Serialize(Ar);
+
+							FaceAndMemory->CompleteAsyncLoad(FTLibrary, FontFaceData, FontData.GetSubFaceIndex());
+						}
 					}
 					else
 					{
diff --git a/Engine/Source/Runtime/SlateCore/Private/Fonts/FontGeometryPreprocessing.cpp b/Engine/Source/Runtime/SlateCore/Private/Fonts/FontGeometryPreprocessing.cpp
--- a/Engine/Source/Runtime/SlateCore/Private/Fonts/FontGeometryPreprocessing.cpp
+++ b/Engine/Source/Runtime/SlateCore/Private/Fonts/FontGeometryPreprocessing.cpp
@@ -126,7 +126,7 @@ bool PreprocessFontGeometry(FPreprocessedFontGeometry& OutPreprocessedFontGeomet
 			// Restore MsdfgenShape to its initial state by re-applying contour reversals
 			PreprocessedGlyphGeometryView.UpdateWindings(MsdfgenShape);
 			// If contour windings are fully reversed, do not store individual glyph preprocessing data yet in case this is true for the majority of glyphs
-			if (PreprocessedGlyphGeometryView.ReversedWindings())
+			if (PreprocessedGlyphGeometryView.HasAllContoursReversedWindings())
 			{
 				bReversedWindings = true;
 			}
diff --git a/Engine/Source/Runtime/SlateCore/Private/Fonts/PreprocessedFontGeometry.cpp b/Engine/Source/Runtime/SlateCore/Private/Fonts/PreprocessedFontGeometry.cpp
--- a/Engine/Source/Runtime/SlateCore/Private/Fonts/PreprocessedFontGeometry.cpp
+++ b/Engine/Source/Runtime/SlateCore/Private/Fonts/PreprocessedFontGeometry.cpp
@@ -42,12 +42,34 @@ static msdfgen::Point2 DecompressPoint(const int16* Coordinates)
 	return msdfgen::Point2(DecompressCoordinate(Coordinates[0]), DecompressCoordinate(Coordinates[1]));
 }
 
+FPreprocessedGlyphGeometryView::FPreprocessedGlyphGeometryView(
+	uint8 InFlags,
+	int32 InContourCount,
+	const uint8* InContourDataPtr,
+	const int16* InCoordinateDataPtr,
+	int32 InContourDataLength,
+	int32 InCoordinateDataLength
+)
+	: Flags(InFlags)
+	, ContourCount(InContourCount)
+	, ContourDataPtr(InContourDataPtr)
+	, CoordinateDataPtr(InCoordinateDataPtr)
+	, ContourDataLength(InContourDataLength)
+	, CoordinateDataLength(InCoordinateDataLength)
+{
+	uint8 FilteredValue = (Flags & (FPreprocessedGlyphGeometry::FLAG_REVERSE_WINDINGS | FPreprocessedGlyphGeometry::FLAG_CONTOUR_WINDINGS | FPreprocessedGlyphGeometry::FLAG_FULL_GEOMETRY));
+	bool bHasSingleBit = FilteredValue != 0 && (FilteredValue & (FilteredValue - 1)) == 0;
+	check(FilteredValue == 0 || bHasSingleBit); //Ensure we have no invalid combination of flags (FLAG_REVERSE_WINDINGS, FLAG_CONTOUR_WINDINGS and FLAG_FULL_GEOMETRY are mutually exclusive)
+}
+
 bool FPreprocessedGlyphGeometryView::UpdateGeometry(msdfgen::Shape& OutMsdfgenShape) const
 {
-	if (!(Flags&FPreprocessedGlyphGeometry::FLAG_FULL_GEOMETRY))
+	const bool bIsFullGeometry = (Flags & FPreprocessedGlyphGeometry::FLAG_FULL_GEOMETRY);
+	if (!bIsFullGeometry)
 	{
 		return false;
 	}
+
 	check(!(Flags&(FPreprocessedGlyphGeometry::FLAG_REVERSE_WINDINGS|FPreprocessedGlyphGeometry::FLAG_CONTOUR_WINDINGS)));
 	OutMsdfgenShape.inverseYAxis = (Flags&FPreprocessedGlyphGeometry::FLAG_INVERSE_Y_AXIS) != 0;
 	OutMsdfgenShape.contours.clear();
@@ -163,7 +185,7 @@ bool FPreprocessedGlyphGeometryView::UpdateWindings(msdfgen::Shape& InOutMsdfgen
 	return ContourIndex == ContourCount;
 }
 
-bool FPreprocessedGlyphGeometryView::ReversedWindings() const
+bool FPreprocessedGlyphGeometryView::HasAllContoursReversedWindings() const
 {
 	return (Flags&FPreprocessedGlyphGeometry::FLAG_REVERSE_WINDINGS) != 0;
 }
diff --git a/Engine/Source/Runtime/SlateCore/Private/Fonts/PreprocessedFontGeometry.h b/Engine/Source/Runtime/SlateCore/Private/Fonts/PreprocessedFontGeometry.h
--- a/Engine/Source/Runtime/SlateCore/Private/Fonts/PreprocessedFontGeometry.h
+++ b/Engine/Source/Runtime/SlateCore/Private/Fonts/PreprocessedFontGeometry.h
@@ -44,14 +44,7 @@ public:
 		const int16* InCoordinateDataPtr,
 		int32 InContourDataLength,
 		int32 InCoordinateDataLength
-	)
-		: Flags(InFlags)
-		, ContourCount(InContourCount)
-		, ContourDataPtr(InContourDataPtr)
-		, CoordinateDataPtr(InCoordinateDataPtr)
-		, ContourDataLength(InContourDataLength)
-		, CoordinateDataLength(InCoordinateDataLength)
-	{ }
+	);
 
 	/** If full geometry data is available, OutMsdfgenShape will be overwritten with this data. Otherwise, false is returned. */
 	bool UpdateGeometry(SlateMsdfgen::FShape& OutMsdfgenShape) const;
@@ -60,7 +53,7 @@ public:
 	bool UpdateWindings(SlateMsdfgen::FShape& InOutMsdfgenShape) const;
 
 	/** Returns true if all contours have reversed windings. */
-	bool ReversedWindings() const;
+	bool HasAllContoursReversedWindings() const;
 
 private:
 	/** A bitfield of glyph flags (see FLAG_ constants in FPreprocessedGlyphGeometry) */
