Fix constructor/destructor API setups with merged modules

#rnx

diff --git a/Engine/Source/Runtime/Core/Private/Serialization/Archive.cpp b/Engine/Source/Runtime/Core/Private/Serialization/Archive.cpp
--- a/Engine/Source/Runtime/Core/Private/Serialization/Archive.cpp
+++ b/Engine/Source/Runtime/Core/Private/Serialization/Archive.cpp
@@ -643,6 +643,8 @@ FString FArchiveProxy::GetLocalizationNamespace() const
 }
 #endif // USE_STABLE_LOCALIZATION_KEYS
 
+FNameAsStringProxyArchive::~FNameAsStringProxyArchive() = default;
+
 /**
  * Serialize the given FName as an FString
  */
diff --git a/Engine/Source/Runtime/Core/Public/Serialization/NameAsStringProxyArchive.h b/Engine/Source/Runtime/Core/Public/Serialization/NameAsStringProxyArchive.h
--- a/Engine/Source/Runtime/Core/Public/Serialization/NameAsStringProxyArchive.h
+++ b/Engine/Source/Runtime/Core/Public/Serialization/NameAsStringProxyArchive.h
@@ -19,5 +19,7 @@ struct FNameAsStringProxyArchive : public FArchiveProxy
 		 :	FArchiveProxy(InInnerArchive)
 	 { }
 
+	 CORE_API ~FNameAsStringProxyArchive();
+
 	 CORE_API virtual FArchive& operator<<(class FName& N);
 };
diff --git a/Engine/Source/Runtime/CoreUObject/Private/Serialization/ArchiveUObject.cpp b/Engine/Source/Runtime/CoreUObject/Private/Serialization/ArchiveUObject.cpp
--- a/Engine/Source/Runtime/CoreUObject/Private/Serialization/ArchiveUObject.cpp
+++ b/Engine/Source/Runtime/CoreUObject/Private/Serialization/ArchiveUObject.cpp
@@ -154,6 +154,15 @@ FArchive& FArchiveUObject::SerializeWeakObjectPtr(FArchive& Ar, FWeakObjectPtr&
 /*----------------------------------------------------------------------------
 	FObjectAndNameAsStringProxyArchive.
 ----------------------------------------------------------------------------*/
+
+FObjectAndNameAsStringProxyArchive::FObjectAndNameAsStringProxyArchive(FArchive& InInnerArchive, bool bInLoadIfFindFails)
+	: FNameAsStringProxyArchive(InInnerArchive)
+	, bLoadIfFindFails(bInLoadIfFindFails)
+{
+}
+
+FObjectAndNameAsStringProxyArchive::~FObjectAndNameAsStringProxyArchive() = default;
+
 /**
  * Serialize the given UObject* as an FString
  */
diff --git a/Engine/Source/Runtime/CoreUObject/Public/Serialization/ObjectAndNameAsStringProxyArchive.h b/Engine/Source/Runtime/CoreUObject/Public/Serialization/ObjectAndNameAsStringProxyArchive.h
--- a/Engine/Source/Runtime/CoreUObject/Public/Serialization/ObjectAndNameAsStringProxyArchive.h
+++ b/Engine/Source/Runtime/CoreUObject/Public/Serialization/ObjectAndNameAsStringProxyArchive.h
@@ -26,10 +26,9 @@ struct FObjectAndNameAsStringProxyArchive : public FNameAsStringProxyArchive
 	 * @param InInnerArchive - The inner archive to proxy.
 	 * @param bInLoadIfFindFails - Indicates whether to try and load a ref'd object if we don't find it
 	 */
-	FObjectAndNameAsStringProxyArchive(FArchive& InInnerArchive, bool bInLoadIfFindFails)
-		:	FNameAsStringProxyArchive(InInnerArchive)
-		,	bLoadIfFindFails(bInLoadIfFindFails)
-	{ }
+	COREUOBJECT_API FObjectAndNameAsStringProxyArchive(FArchive& InInnerArchive, bool bInLoadIfFindFails);
+
+	COREUOBJECT_API virtual ~FObjectAndNameAsStringProxyArchive();
 
 	/** If we fail to find an object during loading, try and load it. */
 	bool bLoadIfFindFails;
diff --git a/Engine/Source/Runtime/Engine/Classes/GameFramework/PlayerController.h b/Engine/Source/Runtime/Engine/Classes/GameFramework/PlayerController.h
--- a/Engine/Source/Runtime/Engine/Classes/GameFramework/PlayerController.h
+++ b/Engine/Source/Runtime/Engine/Classes/GameFramework/PlayerController.h
@@ -123,17 +123,16 @@ public:
 /** Data structure used to setup an input mode that allows only the UI to respond to user input. */
 struct FInputModeUIOnly : public FInputModeDataBase
 {
+	ENGINE_API FInputModeUIOnly();
+
+	ENGINE_API virtual ~FInputModeUIOnly();
+
 	/** Widget to focus */
 	ENGINE_API FInputModeUIOnly& SetWidgetToFocus(TSharedPtr<SWidget> InWidgetToFocus);
 
 	/** Sets the mouse locking behavior of the viewport */
 	ENGINE_API FInputModeUIOnly& SetLockMouseToViewportBehavior(EMouseLockMode InMouseLockMode);
 
-	FInputModeUIOnly()
-		: WidgetToFocus()
-		, MouseLockMode(EMouseLockMode::LockInFullscreen)
-	{}
-
 #if UE_ENABLE_DEBUG_DRAWING
 	ENGINE_API virtual const FString& GetDebugDisplayName() const override;
 #endif	// UE_ENABLE_DEBUG_DRAWING
@@ -179,6 +178,10 @@ struct FUpdateLevelStreamingLevelStatus
 /** Data structure used to setup an input mode that allows the UI to respond to user input, and if the UI doesn't handle it player input / player controller gets a chance. */
 struct FInputModeGameAndUI : public FInputModeDataBase
 {
+	ENGINE_API FInputModeGameAndUI();
+
+	ENGINE_API virtual ~FInputModeGameAndUI();
+
 	/** Widget to focus */
 	FInputModeGameAndUI& SetWidgetToFocus(TSharedPtr<SWidget> InWidgetToFocus) { WidgetToFocus = InWidgetToFocus; return *this; }
 
@@ -190,12 +193,6 @@ struct FInputModeGameAndUI : public FInputModeDataBase
 
 	virtual bool ShouldFlushInputOnViewportFocus() const override { return false; }
 
-	FInputModeGameAndUI()
-		: WidgetToFocus()
-		, MouseLockMode(EMouseLockMode::DoNotLock)
-		, bHideCursorDuringCapture(true)
-	{}
-
 #if UE_ENABLE_DEBUG_DRAWING
 	ENGINE_API virtual const FString& GetDebugDisplayName() const override;
 #endif	// UE_ENABLE_DEBUG_DRAWING
@@ -212,13 +209,13 @@ protected:
 /** Data structure used to setup an input mode that allows only the player input / player controller to respond to user input. */
 struct FInputModeGameOnly : public FInputModeDataBase
 {
+	ENGINE_API FInputModeGameOnly();
+
+	ENGINE_API virtual ~FInputModeGameOnly();
+
 	/** Whether the mouse down that causes capture should be consumed, and not passed to player input processing */
 	FInputModeGameOnly& SetConsumeCaptureMouseDown(bool InConsumeCaptureMouseDown) { bConsumeCaptureMouseDown = InConsumeCaptureMouseDown; return *this; }
 
-	FInputModeGameOnly()
-		: bConsumeCaptureMouseDown(true)
-	{}
-
 #if UE_ENABLE_DEBUG_DRAWING
 	ENGINE_API virtual const FString& GetDebugDisplayName() const override;
 #endif	// UE_ENABLE_DEBUG_DRAWING
diff --git a/Engine/Source/Runtime/Engine/Private/ComponentInstanceDataCache.cpp b/Engine/Source/Runtime/Engine/Private/ComponentInstanceDataCache.cpp
--- a/Engine/Source/Runtime/Engine/Private/ComponentInstanceDataCache.cpp
+++ b/Engine/Source/Runtime/Engine/Private/ComponentInstanceDataCache.cpp
@@ -574,6 +574,8 @@ UObject* FInstanceCacheDataBase::GetUniqueTransientObject(UClass* ForClass)
 	return UniqueTransientPackage.DuplicatedObject;
 }
 
+FInstanceCacheDataBase::~FInstanceCacheDataBase() = default;
+
 void FInstanceCacheDataBase::AddReferencedObjects(FReferenceCollector& Collector)
 {
 	Collector.AddReferencedObjects(ReferencedObjects);
diff --git a/Engine/Source/Runtime/Engine/Private/DynamicMeshBuilder.cpp b/Engine/Source/Runtime/Engine/Private/DynamicMeshBuilder.cpp
--- a/Engine/Source/Runtime/Engine/Private/DynamicMeshBuilder.cpp
+++ b/Engine/Source/Runtime/Engine/Private/DynamicMeshBuilder.cpp
@@ -204,11 +204,17 @@ public: // From FTickableObjectRenderThread
 };
 TGlobalResource<FGlobalDynamicMeshVertexPool> GDynamicMeshVertexPool;
 
+FDynamicMeshIndexBuffer32::FDynamicMeshIndexBuffer32() = default;
+FDynamicMeshIndexBuffer32::~FDynamicMeshIndexBuffer32() = default;
+
 void FDynamicMeshIndexBuffer32::InitRHI(FRHICommandListBase& RHICmdList)
 {
 	IndexBufferRHI = UE::RHIResourceUtils::CreateIndexBufferFromArray(RHICmdList, TEXT("FDynamicMeshIndexBuffer32"), UsageFlags, MakeConstArrayView(Indices));
 }
 
+FDynamicMeshIndexBuffer16::FDynamicMeshIndexBuffer16() = default;
+FDynamicMeshIndexBuffer16::~FDynamicMeshIndexBuffer16() = default;
+
 void FDynamicMeshIndexBuffer16::InitRHI(FRHICommandListBase& RHICmdList)
 {
 	IndexBufferRHI = UE::RHIResourceUtils::CreateIndexBufferFromArray(RHICmdList, TEXT("FDynamicMeshIndexBuffer16"), UsageFlags, MakeConstArrayView(Indices));
diff --git a/Engine/Source/Runtime/Engine/Private/LocalVertexFactory.cpp b/Engine/Source/Runtime/Engine/Private/LocalVertexFactory.cpp
--- a/Engine/Source/Runtime/Engine/Private/LocalVertexFactory.cpp
+++ b/Engine/Source/Runtime/Engine/Private/LocalVertexFactory.cpp
@@ -271,6 +271,15 @@ static bool ShouldIgnore(const EShaderPlatform Platform)
 	return false;
 }
 
+FLocalVertexFactory::FLocalVertexFactory(ERHIFeatureLevel::Type InFeatureLevel, const char* InDebugName)
+	: FVertexFactory(InFeatureLevel)
+	, ColorStreamIndex(INDEX_NONE)
+	, DebugName(InDebugName)
+{
+}
+
+FLocalVertexFactory::~FLocalVertexFactory() = default;
+
 /**
  * Should we cache the material's shadertype on this platform with this vertex factory? 
  */
diff --git a/Engine/Source/Runtime/Engine/Private/PlayerController.cpp b/Engine/Source/Runtime/Engine/Private/PlayerController.cpp
--- a/Engine/Source/Runtime/Engine/Private/PlayerController.cpp
+++ b/Engine/Source/Runtime/Engine/Private/PlayerController.cpp
@@ -6150,6 +6150,8 @@ const FString& FInputModeDataBase::GetDebugDisplayName() const
 }
 #endif	// UE_ENABLE_DEBUG_DRAWING
 
+FInputModeUIOnly::~FInputModeUIOnly() = default;
+
 FInputModeUIOnly& FInputModeUIOnly::SetWidgetToFocus(TSharedPtr<SWidget> InWidgetToFocus)
 {
 #if !(UE_BUILD_SHIPPING || UE_BUILD_TEST)
@@ -6162,6 +6164,12 @@ FInputModeUIOnly& FInputModeUIOnly::SetWidgetToFocus(TSharedPtr<SWidget> InWidge
 	return *this;
 }
 
+FInputModeUIOnly::FInputModeUIOnly()
+	: WidgetToFocus()
+	, MouseLockMode(EMouseLockMode::LockInFullscreen)
+{
+}
+
 FInputModeUIOnly& FInputModeUIOnly::SetLockMouseToViewportBehavior(EMouseLockMode InMouseLockMode)
 {
 	MouseLockMode = InMouseLockMode;
@@ -6193,6 +6201,15 @@ void FInputModeUIOnly::ApplyInputMode(FReply& SlateOperations, class UGameViewpo
 	}
 }
 
+FInputModeGameAndUI::FInputModeGameAndUI()
+	: WidgetToFocus()
+	, MouseLockMode(EMouseLockMode::DoNotLock)
+	, bHideCursorDuringCapture(true)
+{
+}
+
+FInputModeGameAndUI::~FInputModeGameAndUI() = default;
+
 void FInputModeGameAndUI::ApplyInputMode(FReply& SlateOperations, class UGameViewportClient& GameViewportClient) const
 {
 	TSharedPtr<SViewport> ViewportWidget = GameViewportClient.GetGameViewportWidget();
@@ -6219,6 +6236,13 @@ const FString& FInputModeGameAndUI::GetDebugDisplayName() const
 }
 #endif	// UE_ENABLE_DEBUG_DRAWING
 
+FInputModeGameOnly::FInputModeGameOnly()
+	: bConsumeCaptureMouseDown(true)
+{
+}
+
+FInputModeGameOnly::~FInputModeGameOnly() = default;
+
 #if UE_ENABLE_DEBUG_DRAWING
 const FString& FInputModeGameOnly::GetDebugDisplayName() const
 {
diff --git a/Engine/Source/Runtime/Engine/Private/SceneView.cpp b/Engine/Source/Runtime/Engine/Private/SceneView.cpp
--- a/Engine/Source/Runtime/Engine/Private/SceneView.cpp
+++ b/Engine/Source/Runtime/Engine/Private/SceneView.cpp
@@ -3131,6 +3131,11 @@ bool FSceneViewFamily::SupportsScreenPercentage() const
 	return false;
 }
 
+FSceneViewFamilyContext::FSceneViewFamilyContext(const ConstructionValues& CVS)
+	: FSceneViewFamily(CVS)
+{
+}
+
 FSceneViewFamilyContext::~FSceneViewFamilyContext()
 {
 	// Cleanup the views allocated for this view family.
diff --git a/Engine/Source/Runtime/Engine/Private/SceneViewExtension.cpp b/Engine/Source/Runtime/Engine/Private/SceneViewExtension.cpp
--- a/Engine/Source/Runtime/Engine/Private/SceneViewExtension.cpp
+++ b/Engine/Source/Runtime/Engine/Private/SceneViewExtension.cpp
@@ -7,6 +7,10 @@
 // FSceneViewExtensionBase
 //
 
+FSceneViewExtensionBase::FSceneViewExtensionBase(const FAutoRegister&)
+{
+}
+
 FSceneViewExtensionBase::~FSceneViewExtensionBase()
 {
 	// The engine stores view extensions by TWeakPtr<ISceneViewExtension>,
diff --git a/Engine/Source/Runtime/Engine/Public/ComponentInstanceDataCache.h b/Engine/Source/Runtime/Engine/Public/ComponentInstanceDataCache.h
--- a/Engine/Source/Runtime/Engine/Public/ComponentInstanceDataCache.h
+++ b/Engine/Source/Runtime/Engine/Public/ComponentInstanceDataCache.h
@@ -69,7 +69,7 @@ struct FInstanceCacheDataBase
 {
 	GENERATED_BODY()
 
-	virtual ~FInstanceCacheDataBase() = default;
+	ENGINE_API virtual ~FInstanceCacheDataBase();
 
 	ENGINE_API virtual void AddReferencedObjects(FReferenceCollector& Collector);
 
diff --git a/Engine/Source/Runtime/Engine/Public/DynamicMeshBuilder.h b/Engine/Source/Runtime/Engine/Public/DynamicMeshBuilder.h
--- a/Engine/Source/Runtime/Engine/Public/DynamicMeshBuilder.h
+++ b/Engine/Source/Runtime/Engine/Public/DynamicMeshBuilder.h
@@ -227,6 +227,9 @@ private:
 class FDynamicMeshIndexBuffer32 : public FIndexBuffer
 {
 public:
+	ENGINE_API FDynamicMeshIndexBuffer32();
+	ENGINE_API virtual ~FDynamicMeshIndexBuffer32();
+
 	TArray<uint32> Indices;
 	EBufferUsageFlags UsageFlags = EBufferUsageFlags::Static;
 
@@ -236,6 +239,9 @@ public:
 class FDynamicMeshIndexBuffer16 : public FIndexBuffer
 {
 public:
+	ENGINE_API FDynamicMeshIndexBuffer16();
+	ENGINE_API virtual ~FDynamicMeshIndexBuffer16();
+
 	TArray<uint16> Indices;
 	EBufferUsageFlags UsageFlags = EBufferUsageFlags::Static;
 
diff --git a/Engine/Source/Runtime/Engine/Public/LocalVertexFactory.h b/Engine/Source/Runtime/Engine/Public/LocalVertexFactory.h
--- a/Engine/Source/Runtime/Engine/Public/LocalVertexFactory.h
+++ b/Engine/Source/Runtime/Engine/Public/LocalVertexFactory.h
@@ -52,12 +52,9 @@ class FLocalVertexFactory : public FVertexFactory
 	DECLARE_VERTEX_FACTORY_TYPE_API(FLocalVertexFactory, ENGINE_API);
 public:
 
-	FLocalVertexFactory(ERHIFeatureLevel::Type InFeatureLevel, const char* InDebugName)
-		: FVertexFactory(InFeatureLevel)
-		, ColorStreamIndex(INDEX_NONE)
-		, DebugName(InDebugName)
-	{
-	}
+	ENGINE_API FLocalVertexFactory(ERHIFeatureLevel::Type InFeatureLevel, const char* InDebugName);
+
+	ENGINE_API virtual ~FLocalVertexFactory();
 
 	struct FDataType : public FStaticMeshDataType
 	{
diff --git a/Engine/Source/Runtime/Engine/Public/SceneView.h b/Engine/Source/Runtime/Engine/Public/SceneView.h
--- a/Engine/Source/Runtime/Engine/Public/SceneView.h
+++ b/Engine/Source/Runtime/Engine/Public/SceneView.h
@@ -2574,9 +2574,7 @@ class FSceneViewFamilyContext : public FSceneViewFamily
 {
 public:
 	/** Initialization constructor. */
-	FSceneViewFamilyContext( const ConstructionValues& CVS)
-		:	FSceneViewFamily(CVS)
-	{}
+	ENGINE_API FSceneViewFamilyContext(const ConstructionValues& CVS);
 
 	/** Destructor. */
 	ENGINE_API virtual ~FSceneViewFamilyContext();
diff --git a/Engine/Source/Runtime/Engine/Public/SceneViewExtension.h b/Engine/Source/Runtime/Engine/Public/SceneViewExtension.h
--- a/Engine/Source/Runtime/Engine/Public/SceneViewExtension.h
+++ b/Engine/Source/Runtime/Engine/Public/SceneViewExtension.h
@@ -122,9 +122,6 @@ public:
 	};
 
 public:
-	ISceneViewExtension() {}
-	
-	virtual ~ISceneViewExtension() {}
 
 	/**
      * Called on game thread when creating the view family.
@@ -250,7 +247,7 @@ class FAutoRegister
 class FSceneViewExtensionBase : public ISceneViewExtension, public TSharedFromThis<FSceneViewExtensionBase, ESPMode::ThreadSafe>
 {
 public:
-	FSceneViewExtensionBase(const FAutoRegister&) {}
+	ENGINE_API FSceneViewExtensionBase(const FAutoRegister&);
 	ENGINE_API virtual ~FSceneViewExtensionBase();
 
 	// Array of Functors that can be used to activate an extension for the current frame and given context.
diff --git a/Engine/Source/Runtime/Experimental/Iris/Core/Private/Iris/ReplicationSystem/FastArrayReplicationFragment.cpp b/Engine/Source/Runtime/Experimental/Iris/Core/Private/Iris/ReplicationSystem/FastArrayReplicationFragment.cpp
--- a/Engine/Source/Runtime/Experimental/Iris/Core/Private/Iris/ReplicationSystem/FastArrayReplicationFragment.cpp
+++ b/Engine/Source/Runtime/Experimental/Iris/Core/Private/Iris/ReplicationSystem/FastArrayReplicationFragment.cpp
@@ -94,6 +94,8 @@ FFastArrayReplicationFragmentBase::FFastArrayReplicationFragmentBase(EReplicatio
 #endif
 }
 
+FFastArrayReplicationFragmentBase::~FFastArrayReplicationFragmentBase() = default;
+
 const FReplicationStateDescriptor* FFastArrayReplicationFragmentBase::GetArrayElementDescriptor() const
 {
 	const FReplicationStateDescriptor* FastArrayDescriptor = ReplicationStateDescriptor.GetReference();
diff --git a/Engine/Source/Runtime/Experimental/Iris/Core/Public/Iris/ReplicationSystem/Private/FastArrayReplicationFragmentInternal.h b/Engine/Source/Runtime/Experimental/Iris/Core/Public/Iris/ReplicationSystem/Private/FastArrayReplicationFragmentInternal.h
--- a/Engine/Source/Runtime/Experimental/Iris/Core/Public/Iris/ReplicationSystem/Private/FastArrayReplicationFragmentInternal.h
+++ b/Engine/Source/Runtime/Experimental/Iris/Core/Public/Iris/ReplicationSystem/Private/FastArrayReplicationFragmentInternal.h
@@ -104,6 +104,7 @@ public:
 
 protected:
 	IRISCORE_API FFastArrayReplicationFragmentBase(EReplicationFragmentTraits InTraits, UObject* InOwner, const FReplicationStateDescriptor* InDescriptor, bool bValidateDescriptor = true);
+	IRISCORE_API virtual ~FFastArrayReplicationFragmentBase();
 
 	// FReplicationFragment Implementation
 	IRISCORE_API virtual void CollectOwner(FReplicationStateOwnerCollector* Owners) const override;
diff --git a/Engine/Source/Runtime/MovieScene/Private/Channels/MovieSceneFloatChannel.cpp b/Engine/Source/Runtime/MovieScene/Private/Channels/MovieSceneFloatChannel.cpp
--- a/Engine/Source/Runtime/MovieScene/Private/Channels/MovieSceneFloatChannel.cpp
+++ b/Engine/Source/Runtime/MovieScene/Private/Channels/MovieSceneFloatChannel.cpp
@@ -67,6 +67,19 @@ bool FMovieSceneFloatValue::operator!=(const FMovieSceneFloatValue& Other) const
 	return !(*this == Other);
 }
 
+FMovieSceneFloatChannel::FMovieSceneFloatChannel()
+	: PreInfinityExtrap(RCCE_Constant)
+	, PostInfinityExtrap(RCCE_Constant)
+	, DefaultValue(0.f)
+	, bHasDefaultValue(false)
+#if WITH_EDITORONLY_DATA
+	, bShowCurve(false)
+#endif
+{
+}
+
+FMovieSceneFloatChannel::~FMovieSceneFloatChannel() = default;
+
 UE::MovieScene::FPiecewiseCurve FMovieSceneFloatChannel::AsPiecewiseCurve(bool bWithPreAndPostInfinityExtrap) const
 {
 	using namespace UE::MovieScene;
diff --git a/Engine/Source/Runtime/MovieScene/Public/Channels/MovieSceneChannel.h b/Engine/Source/Runtime/MovieScene/Public/Channels/MovieSceneChannel.h
--- a/Engine/Source/Runtime/MovieScene/Public/Channels/MovieSceneChannel.h
+++ b/Engine/Source/Runtime/MovieScene/Public/Channels/MovieSceneChannel.h
@@ -113,7 +113,8 @@ struct FMovieSceneChannel
 	GENERATED_BODY()
 
 	FMovieSceneChannel() {}
-	virtual ~FMovieSceneChannel() {}
+
+	MOVIESCENE_API virtual ~FMovieSceneChannel() {}
 
 	/**
 	 * Get the time for the specified key handle
diff --git a/Engine/Source/Runtime/MovieScene/Public/Channels/MovieSceneFloatChannel.h b/Engine/Source/Runtime/MovieScene/Public/Channels/MovieSceneFloatChannel.h
--- a/Engine/Source/Runtime/MovieScene/Public/Channels/MovieSceneFloatChannel.h
+++ b/Engine/Source/Runtime/MovieScene/Public/Channels/MovieSceneFloatChannel.h
@@ -125,15 +125,9 @@ struct FMovieSceneFloatChannel : public FMovieSceneChannel
 	typedef float CurveValueType;
 	typedef FMovieSceneFloatValue ChannelValueType;
 
-	FMovieSceneFloatChannel() 
-		: PreInfinityExtrap(RCCE_Constant)
-		, PostInfinityExtrap(RCCE_Constant)
-		, DefaultValue(0.f)
-		, bHasDefaultValue(false)
-#if WITH_EDITORONLY_DATA
-		, bShowCurve(false)
-#endif
-	{}
+	MOVIESCENE_API FMovieSceneFloatChannel();
+
+	MOVIESCENE_API virtual ~FMovieSceneFloatChannel();
 
 	MOVIESCENE_API UE::MovieScene::FPiecewiseCurve AsPiecewiseCurve(bool bWithPreAndPostInfinityExtrap = true) const;
 
diff --git a/Engine/Source/Runtime/RenderCore/Private/VertexFactory.cpp b/Engine/Source/Runtime/RenderCore/Private/VertexFactory.cpp
--- a/Engine/Source/Runtime/RenderCore/Private/VertexFactory.cpp
+++ b/Engine/Source/Runtime/RenderCore/Private/VertexFactory.cpp
@@ -251,6 +251,8 @@ bool LoadFromCompactBinary(FCbFieldView Field, FVertexFactoryTypeDependency& Out
 }
 #endif // WITH_EDITOR
 
+FVertexFactory::~FVertexFactory() = default;
+
 void FVertexFactory::GetStreams(ERHIFeatureLevel::Type InFeatureLevel, EVertexInputStreamType VertexStreamType, FVertexInputStreamArray& OutVertexStreams) const
 {
 	check(IsInitialized());
diff --git a/Engine/Source/Runtime/RenderCore/Public/VertexFactory.h b/Engine/Source/Runtime/RenderCore/Public/VertexFactory.h
--- a/Engine/Source/Runtime/RenderCore/Public/VertexFactory.h
+++ b/Engine/Source/Runtime/RenderCore/Public/VertexFactory.h
@@ -667,6 +667,8 @@ public:
 		}
 	}
 
+	RENDERCORE_API virtual ~FVertexFactory();
+
 	virtual FVertexFactoryType* GetType() const { return nullptr; }
 
 	RENDERCORE_API void GetStreams(ERHIFeatureLevel::Type InFeatureLevel, EVertexInputStreamType VertexStreamType, FVertexInputStreamArray& OutVertexStreams) const;
diff --git a/Engine/Source/Runtime/Renderer/Private/SceneRendering.cpp b/Engine/Source/Runtime/Renderer/Private/SceneRendering.cpp
--- a/Engine/Source/Runtime/Renderer/Private/SceneRendering.cpp
+++ b/Engine/Source/Runtime/Renderer/Private/SceneRendering.cpp
@@ -5666,6 +5666,7 @@ TSharedRef<ISceneViewExtension, ESPMode::ThreadSafe> GetRendererViewExtension()
 	class FRendererViewExtension : public ISceneViewExtension
 	{
 	public:
+		virtual ~FRendererViewExtension() = default;
 		virtual void PostRenderView_RenderThread(FRDGBuilder& GraphBuilder, FSceneView& InView)
 		{
 			FViewInfo& View = static_cast<FViewInfo&>(InView);
diff --git a/Engine/Source/Runtime/Slate/Private/Widgets/Layout/SScrollBar.cpp b/Engine/Source/Runtime/Slate/Private/Widgets/Layout/SScrollBar.cpp
--- a/Engine/Source/Runtime/Slate/Private/Widgets/Layout/SScrollBar.cpp
+++ b/Engine/Source/Runtime/Slate/Private/Widgets/Layout/SScrollBar.cpp
@@ -251,6 +251,8 @@ SScrollBar::SScrollBar()
 {
 }
 
+SScrollBar::~SScrollBar() = default;
+
 FSlateColor SScrollBar::GetTrackOpacity() const
 {
 	if (bDraggingThumb || IsHovered())
diff --git a/Engine/Source/Runtime/Slate/Public/Widgets/Layout/SScrollBar.h b/Engine/Source/Runtime/Slate/Public/Widgets/Layout/SScrollBar.h
--- a/Engine/Source/Runtime/Slate/Public/Widgets/Layout/SScrollBar.h
+++ b/Engine/Source/Runtime/Slate/Public/Widgets/Layout/SScrollBar.h
@@ -160,6 +160,7 @@ public:
 	SLATE_API virtual void EndScrolling();
 
 	SLATE_API SScrollBar();
+	SLATE_API ~SScrollBar();
 
 protected:
 	
diff --git a/Engine/Source/Runtime/SlateCore/Private/Input/Events.cpp b/Engine/Source/Runtime/SlateCore/Private/Input/Events.cpp
--- a/Engine/Source/Runtime/SlateCore/Private/Input/Events.cpp
+++ b/Engine/Source/Runtime/SlateCore/Private/Input/Events.cpp
@@ -15,6 +15,8 @@
 const FTouchKeySet FTouchKeySet::StandardSet(EKeys::LeftMouseButton);
 const FTouchKeySet FTouchKeySet::EmptySet(EKeys::Invalid);
 
+FInputEvent::~FInputEvent() = default;
+
 FGeometry FInputEvent::FindGeometry(const TSharedRef<SWidget>& WidgetToFind) const
 {
 	return EventPath->FindArrangedWidget(WidgetToFind).Get(FArrangedWidget::GetNullWidget()).Geometry;
@@ -40,11 +42,56 @@ bool FInputEvent::IsKeyEvent() const
 	return false;
 }
 
+FCharacterEvent::~FCharacterEvent() = default;
+
 FText FCharacterEvent::ToText() const
 {
 	return FText::Format( NSLOCTEXT("Events", "Char", "Char({0})"), FText::FromString(FString::ConstructFromPtrSize(&Character, 1)) );
 }
 
+FKeyEvent::FKeyEvent()
+	: FInputEvent(FModifierKeysState(), 0, false)
+	, Key()
+	, CharacterCode(0)
+	, KeyCode(0)
+{
+}
+
+FKeyEvent::FKeyEvent(const FKey InKey,
+	const FModifierKeysState& InModifierKeys,
+	const uint32 InUserIndex,
+	const bool bInIsRepeat,
+	const uint32 InCharacterCode,
+	const uint32 InKeyCode
+)
+	: FInputEvent(InModifierKeys, InUserIndex, bInIsRepeat)
+	, Key(InKey)
+	, CharacterCode(InCharacterCode)
+	, KeyCode(InKeyCode)
+{
+}
+
+FKeyEvent::FKeyEvent(const FKey InKey,
+	const FModifierKeysState& InModifierKeys,
+	const FInputDeviceId InDeviceId,
+	const bool bInIsRepeat,
+	const uint32 InCharacterCode,
+	const uint32 InKeyCode,
+	const TOptional<int32> InOptionalSlateUserIndex
+)
+	: FInputEvent(InModifierKeys, InDeviceId, bInIsRepeat)
+	, Key(InKey)
+	, CharacterCode(InCharacterCode)
+	, KeyCode(InKeyCode)
+{
+	if (InOptionalSlateUserIndex.IsSet())
+	{
+		UserIndex = InOptionalSlateUserIndex.GetValue();
+	}
+}
+
+FKeyEvent::~FKeyEvent() = default;
+
 FText FKeyEvent::ToText() const
 {
 	return FText::Format( NSLOCTEXT("Events", "Key", "Key({0})"), Key.GetDisplayName() );
@@ -55,11 +102,15 @@ bool FKeyEvent::IsKeyEvent() const
 	return true;
 }
 
+FAnalogInputEvent::~FAnalogInputEvent() = default;
+
 FText FAnalogInputEvent::ToText() const
 {
 	return FText::Format(NSLOCTEXT("Events", "AnalogInput", "AnalogInput(key:{0}, value:{1}"), GetKey().GetDisplayName(), AnalogValue);
 }
 
+FPointerEvent::~FPointerEvent() = default;
+
 FText FPointerEvent::ToText() const
 {
 	return FText::Format( NSLOCTEXT("Events", "Pointer", "Pointer(key:{0}, pos:{1}x{2}, delta:{3}x{4})"), EffectingButton.GetDisplayName(), ScreenSpacePosition.X, ScreenSpacePosition.Y, CursorDelta.X, CursorDelta.Y);
diff --git a/Engine/Source/Runtime/SlateCore/Private/Widgets/SBoxPanel.cpp b/Engine/Source/Runtime/SlateCore/Private/Widgets/SBoxPanel.cpp
--- a/Engine/Source/Runtime/SlateCore/Private/Widgets/SBoxPanel.cpp
+++ b/Engine/Source/Runtime/SlateCore/Private/Widgets/SBoxPanel.cpp
@@ -240,6 +240,7 @@ SBoxPanel::SBoxPanel( EOrientation InOrientation )
 
 }
 
+SBoxPanel::~SBoxPanel() = default;
 
 void SDragAndDropVerticalBox::Construct(const FArguments& InArgs)
 {
diff --git a/Engine/Source/Runtime/SlateCore/Public/Input/Events.h b/Engine/Source/Runtime/SlateCore/Public/Input/Events.h
--- a/Engine/Source/Runtime/SlateCore/Public/Input/Events.h
+++ b/Engine/Source/Runtime/SlateCore/Public/Input/Events.h
@@ -198,7 +198,7 @@ public:
 	/**
 	 * Virtual destructor.
 	 */
-	virtual ~FInputEvent( ) { }
+	SLATECORE_API virtual ~FInputEvent();
 
 public:
 
@@ -435,13 +435,7 @@ public:
 	/**
 	 * UStruct Constructor.  Not meant for normal usage.
 	 */
-	FKeyEvent()
-		: FInputEvent(FModifierKeysState(), 0, false)
-		, Key()
-		, CharacterCode(0)
-		, KeyCode(0)
-	{
-	}
+	SLATECORE_API FKeyEvent();
 
 	/**
 	 * Constructor.  Events are immutable once constructed.
@@ -450,37 +444,24 @@ public:
 	 * @param  InModifierKeys  Modifier key state for this event
 	 * @param  bInIsRepeat  True if this key is an auto-repeated keystroke
 	 */
-	FKeyEvent(	const FKey InKey,
-				const FModifierKeysState& InModifierKeys, 
-				const uint32 InUserIndex,
-				const bool bInIsRepeat,
-				const uint32 InCharacterCode,
-				const uint32 InKeyCode
-	)
-		: FInputEvent(InModifierKeys, InUserIndex, bInIsRepeat)
-		, Key(InKey)
-		, CharacterCode(InCharacterCode)
-		, KeyCode(InKeyCode)
-	{ }
+	SLATECORE_API FKeyEvent(const FKey InKey,
+		const FModifierKeysState& InModifierKeys,
+		const uint32 InUserIndex,
+		const bool bInIsRepeat,
+		const uint32 InCharacterCode,
+		const uint32 InKeyCode
+	);
 	
-	FKeyEvent(const FKey InKey,
-			const FModifierKeysState& InModifierKeys, 
-			const FInputDeviceId InDeviceId,
-			const bool bInIsRepeat,
-			const uint32 InCharacterCode,
-			const uint32 InKeyCode,
-			const TOptional<int32> InOptionalSlateUserIndex = TOptional<int32>()
-	)
-		: FInputEvent(InModifierKeys, InDeviceId, bInIsRepeat)
-		, Key(InKey)
-		, CharacterCode(InCharacterCode)
-		, KeyCode(InKeyCode)
-	{
-		if (InOptionalSlateUserIndex.IsSet())
-		{
-			UserIndex = InOptionalSlateUserIndex.GetValue();
-		}
-	}
+	SLATECORE_API FKeyEvent(const FKey InKey,
+		const FModifierKeysState& InModifierKeys,
+		const FInputDeviceId InDeviceId,
+		const bool bInIsRepeat,
+		const uint32 InCharacterCode,
+		const uint32 InKeyCode,
+		const TOptional<int32> InOptionalSlateUserIndex = TOptional<int32>()
+	);
+
+	SLATECORE_API virtual ~FKeyEvent();
 
 	/**
 	 * Returns the name of the key for this event
@@ -589,6 +570,8 @@ public:
 		, AnalogValue(InAnalogValue)
 	{ }
 
+	SLATECORE_API virtual ~FAnalogInputEvent();
+
 	/**
 	 * Returns the analog value between 0 and 1.
 	 * 0 being not pressed at all, 1 being fully pressed.
@@ -643,6 +626,8 @@ public:
 		, Character(InCharacter)
 	{ }
 
+	SLATECORE_API virtual ~FCharacterEvent();
+
 	/**
 	 * Returns the character for this event
 	 *
@@ -968,7 +953,9 @@ public:
 		ScreenSpacePosition = InScreenSpacePosition;
 		LastScreenSpacePosition = InLastScreenSpacePosition;
 	}
-	
+
+	SLATECORE_API virtual ~FPointerEvent();
+
 public:
 
 	/** Returns The position of the cursor in screen space */
diff --git a/Engine/Source/Runtime/SlateCore/Public/Widgets/SBoxPanel.h b/Engine/Source/Runtime/SlateCore/Public/Widgets/SBoxPanel.h
--- a/Engine/Source/Runtime/SlateCore/Public/Widgets/SBoxPanel.h
+++ b/Engine/Source/Runtime/SlateCore/Public/Widgets/SBoxPanel.h
@@ -135,6 +135,7 @@ protected:
 	SLATECORE_API SBoxPanel();
 	/** A Box Panel's orientation. */
 	SLATECORE_API SBoxPanel(EOrientation InOrientation);
+	SLATECORE_API ~SBoxPanel();
 
 	/** Set the orientation of the Box. It will do a full invalidation of the widget. */
 	SLATECORE_API void SetOrientation(EOrientation InOrientation);
diff --git a/Engine/Source/Runtime/UMG/Private/Slate/SObjectWidget.cpp b/Engine/Source/Runtime/UMG/Private/Slate/SObjectWidget.cpp
--- a/Engine/Source/Runtime/UMG/Private/Slate/SObjectWidget.cpp
+++ b/Engine/Source/Runtime/UMG/Private/Slate/SObjectWidget.cpp
@@ -32,7 +32,9 @@ void SObjectWidget::Construct(const FArguments& InArgs, UUserWidget* InWidgetObj
 	}
 }
 
-SObjectWidget::~SObjectWidget(void)
+SObjectWidget::SObjectWidget() = default;
+
+SObjectWidget::~SObjectWidget()
 {
 #if SLATE_VERBOSE_NAMED_EVENTS
 	// This can happen during blueprint compiling, so just ignore it if it happens then, this is really only a concern
diff --git a/Engine/Source/Runtime/UMG/Public/Slate/SObjectWidget.h b/Engine/Source/Runtime/UMG/Public/Slate/SObjectWidget.h
--- a/Engine/Source/Runtime/UMG/Public/Slate/SObjectWidget.h
+++ b/Engine/Source/Runtime/UMG/Public/Slate/SObjectWidget.h
@@ -34,10 +34,11 @@ class SObjectWidget : public SCompoundWidget, public FGCObject
 		_Visibility = EVisibility::SelfHitTestInvisible;
 	}
 
-		SLATE_DEFAULT_SLOT(FArguments, Content)
+	SLATE_DEFAULT_SLOT(FArguments, Content)
 	SLATE_END_ARGS()
 
-	UMG_API virtual ~SObjectWidget(void);
+	UMG_API SObjectWidget();
+	UMG_API virtual ~SObjectWidget();
 
 	UMG_API void Construct(const FArguments& InArgs, UUserWidget* InWidgetObject);
 
