[SlateCore] Added Reserve function to FArrangedChildren for pre-allocation

diff --git a/Engine/Source/Runtime/SlateCore/Public/Layout/ArrangedChildren.h b/Engine/Source/Runtime/SlateCore/Public/Layout/ArrangedChildren.h
--- a/Engine/Source/Runtime/SlateCore/Public/Layout/ArrangedChildren.h
+++ b/Engine/Source/Runtime/SlateCore/Public/Layout/ArrangedChildren.h
@@ -78,6 +78,12 @@ class FArrangedChildren
 			Array.Insert(InWidgetGeometry, Index);
 		}
 	}
+
+	/** Reserves sufficient memory for the given number of arranged widgets. */
+	FORCEINLINE_DEBUGGABLE void Reserve(FArrangedWidgetArray::SizeType Number)
+	{
+		Array.Reserve(Number);
+	}
 	
 	/**
 	 * Add an arranged widget (i.e. widget and its resulting geometry) to the list of Arranged children
