[Widget Path] When users are moving focus, they will no longer traverse through the children of disabled widgets

diff --git a/Engine/Source/Runtime/SlateCore/Private/Layout/WidgetPath.cpp b/Engine/Source/Runtime/SlateCore/Private/Layout/WidgetPath.cpp
--- a/Engine/Source/Runtime/SlateCore/Private/Layout/WidgetPath.cpp
+++ b/Engine/Source/Runtime/SlateCore/Private/Layout/WidgetPath.cpp
@@ -172,7 +172,7 @@ bool FWidgetPath::MoveFocus(int32 PathLevel, EUINavigation NavigationType, bool
 		if (NavigationType == EUINavigation::Next)
 		{
 			// EFocusMoveDirection::Next implies descend, so try to find a focusable descendant.
-			return ExtendPathTo( FFocusableWidgetMatcher() );
+			return ExtendPathTo( FFocusableWidgetMatcher(), EVisibility::Visible, EWidgetPathSearchPurpose::FocusHandling );
 		}
 		else
 		{
diff --git a/Engine/Source/Runtime/SlateCore/Public/Layout/WidgetPath.h b/Engine/Source/Runtime/SlateCore/Public/Layout/WidgetPath.h
--- a/Engine/Source/Runtime/SlateCore/Public/Layout/WidgetPath.h
+++ b/Engine/Source/Runtime/SlateCore/Public/Layout/WidgetPath.h
@@ -9,6 +9,22 @@
 #include "Widgets/SWidget.h"
 #include "Widgets/SWindow.h"
 
+/**
+ * Enumerates different purposes for searching through a widget path. Normally, Standard is appropriate.
+ */
+enum class EWidgetPathSearchPurpose
+{
+	/**
+	 * No specified search purpose. This will be the default and should be used when no other purpose matches
+	 */
+	Standard,
+
+	/**
+	 * We are searching through the WidgetPath to change focus during navigation.
+	 */
+	FocusHandling
+};
+
 /** Matches widgets against InWidget */
 struct FWidgetMatcher
 {
@@ -86,17 +102,18 @@ public:
 	 * Extend the current path such that it reaches some widget that qualifies as a Match
 	 * The widget to match must be a descendant of the last widget currently in the path.
 	 *
-	 * @param Matcher         Some struct that has a "bool IsMatch( const TSharedRef<const SWidget>& InWidget ) const" method
+	 * @param  Matcher          Some struct that has a "bool IsMatch( const TSharedRef<const SWidget>& InWidget ) const" method
 	 * @param  VisibilityFilter	Widgets must have this type of visibility to be included the path
+	 * @param  SearchPurpose	What is the purpose for extending the path
 	 *
 	 * @return true if successful; false otherwise.
 	 */
 	template<typename MatcherType>
-	bool ExtendPathTo( const MatcherType& Matcher, EVisibility VisibilityFilter = EVisibility::Visible )
+	bool ExtendPathTo( const MatcherType& Matcher, EVisibility VisibilityFilter = EVisibility::Visible, EWidgetPathSearchPurpose SearchPurpose = EWidgetPathSearchPurpose::Standard)
 	{
 		const FArrangedWidget& LastWidget = Widgets.Last();
 		
-		FArrangedChildren Extension = GeneratePathToWidget(Matcher, LastWidget, EUINavigation::Next, VisibilityFilter);
+		FArrangedChildren Extension = GeneratePathToWidget(Matcher, LastWidget, EUINavigation::Next, VisibilityFilter, SearchPurpose);
 
 		for( int32 WidgetIndex=0; WidgetIndex < Extension.Num(); ++WidgetIndex )
 		{
@@ -109,24 +126,25 @@ public:
 	/**
 	 * Generate a path from FromWidget to WidgetToFind. The path will not include FromWidget.
 	 *
-	 * @param Matcher         Some struct that has a "bool IsMatch( const TSharedRef<const SWidget>& InWidget ) const" method
-	 * @param FromWidget      Widget from which we a building a path.AddItem*
+	 * @param  Matcher          Some struct that has a "bool IsMatch( const TSharedRef<const SWidget>& InWidget ) const" method
+	 * @param  FromWidget       Widget from which we a building a path.AddItem*
 	 * @param  VisibilityFilter	Widgets must have this type of visibility to be included the path
+	 * @param  SearchPurpose	What is the purpose for generating the path
 	 * 
 	 * @return A path from FromWidget to WidgetToFind; will not include FromWidget.
 	 */
 	template<typename MatcherType>
-	FArrangedChildren GeneratePathToWidget(const MatcherType& Matcher, const FArrangedWidget& FromWidget, EUINavigation NavigationType = EUINavigation::Next, EVisibility VisibilityFilter = EVisibility::Visible)
+	FArrangedChildren GeneratePathToWidget(const MatcherType& Matcher, const FArrangedWidget& FromWidget, EUINavigation NavigationType = EUINavigation::Next, EVisibility VisibilityFilter = EVisibility::Visible, EWidgetPathSearchPurpose SearchPurpose = EWidgetPathSearchPurpose::Standard)
 	{
 		FArrangedChildren PathResult(VisibilityFilter);
 
 		if (NavigationType == EUINavigation::Next)
 		{
-			SearchForWidgetRecursively( Matcher, FromWidget, PathResult, VisibilityFilter );
+			SearchForWidgetRecursively( Matcher, FromWidget, PathResult, VisibilityFilter, SearchPurpose);
 		}
 		else
 		{
-			SearchForWidgetRecursively_Reverse( Matcher, FromWidget, PathResult, VisibilityFilter );
+			SearchForWidgetRecursively_Reverse( Matcher, FromWidget, PathResult, VisibilityFilter, SearchPurpose);
 		}
 		
 
@@ -172,19 +190,20 @@ private:
 	/**
 	 * Utility function to search recursively through a widget hierarchy for a specific widget
 	 *
-	 * @param Matcher         Some struct that has a "bool IsMatch( const TSharedRef<const SWidget>& InWidget ) const" method
+	 * @param  Matcher          Some struct that has a "bool IsMatch( const TSharedRef<const SWidget>& InWidget ) const" method
 	 * @param  InCandidate      The current widget-geometry pair we're testing
 	 * @param  OutReversedPath  The resulting path in reversed order (canonical order is Windows @ index 0, Leafmost widget is last.)
 	 * @param  VisibilityFilter	Widgets must have this type of visibility to be included the path
+	 * @param  SearchPurpose	The purpose for searching for the widget in the path
 	 *
 	 * @return  true if the child widget was found; false otherwise
 	 */
 	template<typename MatchRuleType>
-	static bool SearchForWidgetRecursively( const MatchRuleType& MatchRule, const FArrangedWidget& InCandidate, FArrangedChildren& OutReversedPath, EVisibility VisibilityFilter = EVisibility::Visible );
+	static bool SearchForWidgetRecursively( const MatchRuleType& MatchRule, const FArrangedWidget& InCandidate, FArrangedChildren& OutReversedPath, EVisibility VisibilityFilter = EVisibility::Visible, EWidgetPathSearchPurpose SearchPurpose = EWidgetPathSearchPurpose::Standard);
 
 	/** Identical to SearchForWidgetRecursively, but iterates in reverse order */
 	template<typename MatchRuleType>
-	static bool SearchForWidgetRecursively_Reverse( const MatchRuleType& MatchRule, const FArrangedWidget& InCandidate, FArrangedChildren& OutReversedPath, EVisibility VisibilityFilter = EVisibility::Visible );
+	static bool SearchForWidgetRecursively_Reverse( const MatchRuleType& MatchRule, const FArrangedWidget& InCandidate, FArrangedChildren& OutReversedPath, EVisibility VisibilityFilter = EVisibility::Visible, EWidgetPathSearchPurpose SearchPurpose = EWidgetPathSearchPurpose::Standard);
 };
 
 /**
diff --git a/Engine/Source/Runtime/SlateCore/Public/Layout/WidgetPath.inl b/Engine/Source/Runtime/SlateCore/Public/Layout/WidgetPath.inl
--- a/Engine/Source/Runtime/SlateCore/Public/Layout/WidgetPath.inl
+++ b/Engine/Source/Runtime/SlateCore/Public/Layout/WidgetPath.inl
@@ -11,18 +11,25 @@ struct EVisibility;
 /**
  * Utility function to search recursively through a widget hierarchy for a specific widget
  *
- * @param Matcher         Some struct that has a "bool IsMatch( const TSharedRef<const SWidget>& InWidget ) const" method
+ * @param  Matcher          Some struct that has a "bool IsMatch( const TSharedRef<const SWidget>& InWidget ) const" method
  * @param  InCandidate      The current widget-geometry pair we're testing
  * @param  OutReversedPath  The resulting path in reversed order (canonical order is Windows @ index 0, Leafmost widget is last.)
  * @param  VisibilityFilter	Widgets must have this type of visibility to be included the path
+ * @param  SearchPurpose	The purpose for searching for the widget in the path
  *
  * @return  true if the child widget was found; false otherwise
  */
 template<typename MatchRuleType>
-bool FWidgetPath::SearchForWidgetRecursively( const MatchRuleType& MatchRule, const FArrangedWidget& InCandidate, FArrangedChildren& OutReversedPath, EVisibility VisibilityFilter )
+bool FWidgetPath::SearchForWidgetRecursively( const MatchRuleType& MatchRule, const FArrangedWidget& InCandidate, FArrangedChildren& OutReversedPath, EVisibility VisibilityFilter, EWidgetPathSearchPurpose SearchPurpose )
 {
+	// If we are searching in order to move focus and we want to find a visible widget but the current widget is disabled, we do not want to search through it's children.
+	if (SearchPurpose == EWidgetPathSearchPurpose::FocusHandling && VisibilityFilter == EVisibility::Visible && !InCandidate.Widget->IsEnabled())
+	{
+		return false;
+	}
+
 	const bool bAllow3DWidgets = true;
-	const bool bUpdateVisibilityAttributes = true;
+	const bool bUpdateVisibilityAttributes = SearchPurpose == EWidgetPathSearchPurpose::FocusHandling ? false : true;
 	FArrangedChildren ArrangedChildren(VisibilityFilter, bAllow3DWidgets);
 	InCandidate.Widget->ArrangeChildren( InCandidate.Geometry, ArrangedChildren, bUpdateVisibilityAttributes );
 
@@ -34,7 +41,7 @@ bool FWidgetPath::SearchForWidgetRecursively( const MatchRuleType& MatchRule, co
 			OutReversedPath.AddWidget( SomeChild );
 			return true;
 		}
-		else if ( SearchForWidgetRecursively( MatchRule, SomeChild, OutReversedPath, VisibilityFilter ) )
+		else if ( SearchForWidgetRecursively( MatchRule, SomeChild, OutReversedPath, VisibilityFilter, SearchPurpose) )
 		{
 			OutReversedPath.AddWidget( SomeChild );
 			return true;
@@ -46,10 +53,16 @@ bool FWidgetPath::SearchForWidgetRecursively( const MatchRuleType& MatchRule, co
 
 /** Identical to SearchForWidgetRecursively, but iterates in reverse order */
 template<typename MatchRuleType>
-bool FWidgetPath::SearchForWidgetRecursively_Reverse( const MatchRuleType& MatchRule, const FArrangedWidget& InCandidate, FArrangedChildren& OutReversedPath, EVisibility VisibilityFilter )
+bool FWidgetPath::SearchForWidgetRecursively_Reverse( const MatchRuleType& MatchRule, const FArrangedWidget& InCandidate, FArrangedChildren& OutReversedPath, EVisibility VisibilityFilter, EWidgetPathSearchPurpose SearchPurpose )
 {
+	// If we are searching in order to move focus and we want to find a visible widget but the current widget is disabled, we do not want to search through it's children.
+	if (SearchPurpose == EWidgetPathSearchPurpose::FocusHandling && VisibilityFilter == EVisibility::Visible && !InCandidate.Widget->IsEnabled())
+	{
+		return false;
+	}
+
 	const bool bAllow3DWidgets = true;
-	const bool bUpdateVisibilityAttributes = true;
+	const bool bUpdateVisibilityAttributes = SearchPurpose == EWidgetPathSearchPurpose::FocusHandling ? false : true;
 	FArrangedChildren ArrangedChildren(VisibilityFilter, bAllow3DWidgets);
 	InCandidate.Widget->ArrangeChildren( InCandidate.Geometry, ArrangedChildren, bUpdateVisibilityAttributes );
 
@@ -61,7 +74,7 @@ bool FWidgetPath::SearchForWidgetRecursively_Reverse( const MatchRuleType& Match
 			OutReversedPath.AddWidget( SomeChild );
 			return true;
 		}
-		else if ( SearchForWidgetRecursively_Reverse( MatchRule, SomeChild, OutReversedPath, VisibilityFilter) )
+		else if ( SearchForWidgetRecursively_Reverse( MatchRule, SomeChild, OutReversedPath, VisibilityFilter, SearchPurpose) )
 		{
 			OutReversedPath.AddWidget( SomeChild );
 			return true;
