[Slate Deprecation] Fix report slate deprecation to work in more than a single module

#rnx

diff --git a/Engine/Source/Runtime/SlateCore/Public/Types/SlateVector2.h b/Engine/Source/Runtime/SlateCore/Public/Types/SlateVector2.h
--- a/Engine/Source/Runtime/SlateCore/Public/Types/SlateVector2.h
+++ b/Engine/Source/Runtime/SlateCore/Public/Types/SlateVector2.h
@@ -56,6 +56,7 @@ static_assert(UE_ENABLE_SLATE_VECTOR_DEPRECATION_MECHANISMS || !UE_REPORT_SLATE_
 
 
 #undef UE_INCLUDETOOL_IGNORE_INCONSISTENT_STATE
+#undef UE_REPORT_SLATE_VECTOR_DEPRECATION
 
 
 struct FDeprecateSlateVector2D;
