Improve warning message handling when looking up widget styles.

diff --git a/Engine/Source/Runtime/SlateCore/Public/Styling/ISlateStyle.h b/Engine/Source/Runtime/SlateCore/Public/Styling/ISlateStyle.h
--- a/Engine/Source/Runtime/SlateCore/Public/Styling/ISlateStyle.h
+++ b/Engine/Source/Runtime/SlateCore/Public/Styling/ISlateStyle.h
@@ -44,10 +44,14 @@ public:
 	 */
 	virtual TArray<FName> GetEntriesUsingBrush(const FName BrushName) const = 0;
 
+	/**
+	 * Locates a widget style of a specified type and name.
+	 * Allows customization of whether a warning should be logged.
+	 * @return The found widget style or the default value.
+	 */
 	template< typename WidgetStyleType >            
-	const WidgetStyleType& GetWidgetStyle( FName PropertyName, const ANSICHAR* Specifier = nullptr, const WidgetStyleType* DefaultValue = nullptr ) const 
+	const WidgetStyleType& GetWidgetStyle( FName PropertyName, const ANSICHAR* Specifier, const WidgetStyleType* DefaultValue, bool bWarnIfNotFound ) const 
 	{
-		const bool bWarnIfNotFound = true;
 		const FSlateWidgetStyle* ExistingStyle = GetWidgetStyleInternal(WidgetStyleType::TypeName, Join( PropertyName, Specifier ), DefaultValue, bWarnIfNotFound);
 
 		if ( ExistingStyle == nullptr )
@@ -58,6 +62,17 @@ public:
 		return *static_cast< const WidgetStyleType* >( ExistingStyle );
 	}
 
+	/**
+	 * Locates a widget style of a specified type and name.
+	 * @param DefaultValue The value to use if a widget style could not be found. If null and no style could be found, a warning will be logged.
+	 * @return The found widget style or the default value.
+	 */
+	template< typename WidgetStyleType >
+    const WidgetStyleType& GetWidgetStyle( FName PropertyName, const ANSICHAR* Specifier = nullptr, const WidgetStyleType* DefaultValue = nullptr ) const
+    {
+    	return GetWidgetStyle< WidgetStyleType >( PropertyName, Specifier, DefaultValue, DefaultValue == nullptr );
+    }
+
 	template< typename WidgetStyleType >            
 	bool HasWidgetStyle( FName PropertyName, const ANSICHAR* Specifier = nullptr ) const
 	{
