Call InvalidateRootChildOrder when moving new content in a SWindow, this will fix the debug draw being hide when using the fastpath

diff --git a/Engine/Source/Runtime/SlateCore/Private/Widgets/SWindow.cpp b/Engine/Source/Runtime/SlateCore/Private/Widgets/SWindow.cpp
--- a/Engine/Source/Runtime/SlateCore/Private/Widgets/SWindow.cpp
+++ b/Engine/Source/Runtime/SlateCore/Private/Widgets/SWindow.cpp
@@ -1282,6 +1282,11 @@ void SWindow::SetContent( TSharedRef<SWidget> InContent )
 		ChildSlot.operator[](InContent);
 	}
 	
+	if (Advanced_IsInvalidationRoot())
+	{
+		InvalidateRootChildOrder();
+	}
+
 	Invalidate(EInvalidateWidgetReason::ChildOrder);
 }
 
