[Icon] Update Edit Tool Menu Icon to be UE5 style

diff --git a/Engine/Source/Runtime/SlateCore/Private/Styling/StarshipCoreStyle.cpp b/Engine/Source/Runtime/SlateCore/Private/Styling/StarshipCoreStyle.cpp
--- a/Engine/Source/Runtime/SlateCore/Private/Styling/StarshipCoreStyle.cpp
+++ b/Engine/Source/Runtime/SlateCore/Private/Styling/StarshipCoreStyle.cpp
@@ -2226,8 +2226,8 @@ void FStarshipCoreStyle::SetupMultiboxStyles(TSharedRef<FStyle>& Style)
 
 	// MultiBox
 	{
-		Style->Set("MultiBox.GenericToolBarIcon", new IMAGE_BRUSH("Icons/icon_generic_toolbar", Icon40x40));
-		Style->Set("MultiBox.GenericToolBarIcon.Small", new IMAGE_BRUSH("Icons/icon_generic_toolbar", Icon20x20));
+		Style->Set("MultiBox.GenericToolBarIcon", new IMAGE_BRUSH_SVG("Starship/Common/EditToolbar", Icon40x40));
+		Style->Set("MultiBox.GenericToolBarIcon.Small", new IMAGE_BRUSH_SVG("Starship/Common/EditToolbar", Icon16x16));
 
 		Style->Set("MultiboxHookColor", FLinearColor(0.f, 1.f, 0.f, 1.f));
 	}
diff --git a/Engine/Source/Runtime/SlateCore/Private/Styling/UMGCoreStyle.cpp b/Engine/Source/Runtime/SlateCore/Private/Styling/UMGCoreStyle.cpp
--- a/Engine/Source/Runtime/SlateCore/Private/Styling/UMGCoreStyle.cpp
+++ b/Engine/Source/Runtime/SlateCore/Private/Styling/UMGCoreStyle.cpp
@@ -914,8 +914,8 @@ TSharedRef<ISlateStyle> FUMGCoreStyle::Create()
 
 	// MultiBox
 	{
-		Style->Set( "MultiBox.GenericToolBarIcon", new IMAGE_BRUSH( "Icons/Icon_generic_toolbar", Icon40x40 ) );
-		Style->Set( "MultiBox.GenericToolBarIcon.Small", new IMAGE_BRUSH( "Icons/Icon_generic_toolbar", Icon20x20 ) );
+		Style->Set("MultiBox.GenericToolBarIcon", new IMAGE_BRUSH_SVG("Starship/Common/EditToolbar", Icon40x40));
+		Style->Set("MultiBox.GenericToolBarIcon.Small", new IMAGE_BRUSH_SVG("Starship/Common/EditToolbar", Icon16x16));
 
 		Style->Set( "MultiBox.DeleteButton", FButtonStyle() 
 			.SetNormal ( IMAGE_BRUSH( "/Docking/CloseApp_Hovered", Icon16x16 ) )
