Fixing font issues:
- logs are now more precise, and the recently added log was moved to verbose.
- configure explicitely font of the editor for chinese and display of box drawing and blocking element range.

#tests Ran the editor, validated that log about font was reduced a lot (there are still a few valid edge cases, when displaying the ranges for the composite font editor. No more spamming in chinese. Used the profilegpu command, there are no more font related logs.

diff --git a/Engine/Source/Runtime/SlateCore/Private/Fonts/FontCacheCompositeFont.cpp b/Engine/Source/Runtime/SlateCore/Private/Fonts/FontCacheCompositeFont.cpp
--- a/Engine/Source/Runtime/SlateCore/Private/Fonts/FontCacheCompositeFont.cpp
+++ b/Engine/Source/Runtime/SlateCore/Private/Fonts/FontCacheCompositeFont.cpp
@@ -466,7 +466,7 @@ const FFontData& FCompositeFontCache::GetFontDataForCodepoint(const FSlateFontIn
 		if (const FFontData* FoundFontData = GetFontDataForCharacterInTypeface(CachedFallbackTypefaceData, CachedDefaultTypefaceData, /*InSkipCharacterCheck*/true))
 		{
 			OutScalingFactor = CachedFallbackTypefaceData->GetScalingFactor();
-			UE_LOG(LogSlate, Log, TEXT("Last resort fallback font was requested. %s"), *FoundFontData->GetFontFilename());
+			UE_LOG(LogSlate, Verbose, TEXT("Using fallback font %s for codepoint U+%x, as it was not found in default nor sub typefaces of SlateFontInfo %s."), *FoundFontData->GetFontFilename(), InCodepoint, *InFontInfo.TypefaceFontName.ToString());
 			return *FoundFontData;
 		}
 
diff --git a/Engine/Source/Runtime/SlateCore/Private/Fonts/LegacySlateFontInfoCache.cpp b/Engine/Source/Runtime/SlateCore/Private/Fonts/LegacySlateFontInfoCache.cpp
--- a/Engine/Source/Runtime/SlateCore/Private/Fonts/LegacySlateFontInfoCache.cpp
+++ b/Engine/Source/Runtime/SlateCore/Private/Fonts/LegacySlateFontInfoCache.cpp
@@ -127,6 +127,15 @@ TSharedRef<const FCompositeFont> FLegacySlateFontInfoCache::GetDefaultFont()
 			APPEND_FONT(MutableDefaultFont->FallbackTypeface.Typeface, "Regular", "DroidSansFallback.ttf", EFontHinting::Default);
 		}
 
+		// Explicit fallback for glyphs missing from Roboto.
+		{
+			FCompositeSubFont& SubFont = MutableDefaultFont->SubTypefaces[MutableDefaultFont->SubTypefaces.AddDefaulted()];
+			APPEND_RANGE(SubFont, BoxDrawing);
+			APPEND_RANGE(SubFont, BlockElements);
+			APPEND_RANGE(SubFont, HalfwidthAndFullwidthForms);
+			APPEND_FONT(SubFont.Typeface, "Regular", "DroidSansFallback.ttf", EFontHinting::Default);
+		}
+
 		// Arabic
 		{
 			FCompositeSubFont& SubFont = MutableDefaultFont->SubTypefaces[MutableDefaultFont->SubTypefaces.AddDefaulted()];
@@ -175,6 +184,37 @@ TSharedRef<const FCompositeFont> FLegacySlateFontInfoCache::GetDefaultFont()
 			APPEND_EDITOR_FONT(SubFont.Typeface, "VeryLight", "GenEiGothicPro-Light.otf", EFontHinting::AutoLight);
 		}
 
+		// Simplified Chinese (editor-only)
+		if (GIsEditor)
+		{
+			FCompositeSubFont& SubFont = MutableDefaultFont->SubTypefaces[MutableDefaultFont->SubTypefaces.AddDefaulted()];
+			SubFont.Cultures = TEXT("zh-Hans");
+			APPEND_RANGE(SubFont, CJKCompatibility);
+			APPEND_RANGE(SubFont, CJKCompatibilityForms);
+			APPEND_RANGE(SubFont, CJKCompatibilityIdeographs);
+			APPEND_RANGE(SubFont, CJKCompatibilityIdeographsSupplement);
+			APPEND_RANGE(SubFont, CJKRadicalsSupplement);
+			APPEND_RANGE(SubFont, CJKStrokes);
+			APPEND_RANGE(SubFont, CJKSymbolsAndPunctuation);
+			APPEND_RANGE(SubFont, CJKUnifiedIdeographs);
+			APPEND_RANGE(SubFont, CJKUnifiedIdeographsExtensionA);
+			APPEND_RANGE(SubFont, CJKUnifiedIdeographsExtensionB);
+			APPEND_RANGE(SubFont, CJKUnifiedIdeographsExtensionC);
+			APPEND_RANGE(SubFont, CJKUnifiedIdeographsExtensionD);
+			APPEND_RANGE(SubFont, CJKUnifiedIdeographsExtensionE);
+			APPEND_RANGE(SubFont, EnclosedCJKLettersAndMonths);
+			APPEND_FONT(SubFont.Typeface, "Regular", "DroidSansFallback.ttf", EFontHinting::Default);
+			APPEND_FONT(SubFont.Typeface, "Italic", "DroidSansFallback.ttf", EFontHinting::Default);
+			APPEND_FONT(SubFont.Typeface, "Bold", "DroidSansFallback.ttf", EFontHinting::Default);
+			APPEND_FONT(SubFont.Typeface, "BoldItalic", "DroidSansFallback.ttf", EFontHinting::Default);
+			APPEND_FONT(SubFont.Typeface, "BoldCondensed", "DroidSansFallback.ttf", EFontHinting::Default);
+			APPEND_FONT(SubFont.Typeface, "BoldCondensedItalic", "DroidSansFallback.ttf", EFontHinting::Default);
+			APPEND_FONT(SubFont.Typeface, "Black", "DroidSansFallback.ttf", EFontHinting::Default);
+			APPEND_FONT(SubFont.Typeface, "BlackItalic", "DroidSansFallback.ttf", EFontHinting::Default);
+			APPEND_FONT(SubFont.Typeface, "Light", "DroidSansFallback.ttf", EFontHinting::Default);
+			APPEND_FONT(SubFont.Typeface, "VeryLight", "DroidSansFallback.ttf", EFontHinting::Default);
+		}
+
 		// Korean (editor-only)
 		if (GIsEditor)
 		{
diff --git a/Engine/Source/Runtime/SlateCore/Private/Fonts/SlateFontRenderer.cpp b/Engine/Source/Runtime/SlateCore/Private/Fonts/SlateFontRenderer.cpp
--- a/Engine/Source/Runtime/SlateCore/Private/Fonts/SlateFontRenderer.cpp
+++ b/Engine/Source/Runtime/SlateCore/Private/Fonts/SlateFontRenderer.cpp
@@ -376,7 +376,7 @@ FFreeTypeFaceGlyphData FSlateFontRenderer::GetFontFaceForCodepoint(const FFontDa
 
 		if (bCanFallback && FLegacySlateFontInfoCache::Get().IsLastResortFontAvailable())
 		{
-			UE_LOG(LogSlate, Warning, TEXT("Could not find Glyph Index %u, getting last resort font data %s"), ReturnVal.GlyphIndex, *InFontData.GetFontFilename());
+			UE_LOG(LogSlate, Warning, TEXT("Could not find Glyph Index %u with codepoint U+%x, getting last resort font data %s"), ReturnVal.GlyphIndex, InCodepoint, *InFontData.GetFontFilename());
 			ReturnVal.FaceAndMemory = CompositeFontCache->GetFontFace(FLegacySlateFontInfoCache::Get().GetLastResortFontData(FLegacySlateFontInfoCache::FFallbackContext(&InFontData, InCodepoint)));
 
 			if (ReturnVal.FaceAndMemory.IsValid())
