Slate: Experimental Navigation Framework.
Feature Guarded by `Slate.Navigation.Experimental`
1. Proximity-based Navigation method - Exposes distance heuristics to support non-orthogonal navigation
2. Routing Policy for Panel Widgets - When a child widget gets focus, route focus the focus to child widgets according to policy

diff --git a/Engine/Source/Editor/UMGEditor/Private/Customizations/WidgetNavigationCustomization.cpp b/Engine/Source/Editor/UMGEditor/Private/Customizations/WidgetNavigationCustomization.cpp
--- a/Engine/Source/Editor/UMGEditor/Private/Customizations/WidgetNavigationCustomization.cpp
+++ b/Engine/Source/Editor/UMGEditor/Private/Customizations/WidgetNavigationCustomization.cpp
@@ -1,16 +1,30 @@
 // Copyright Epic Games, Inc. All Rights Reserved.
 
 #include "Customizations/WidgetNavigationCustomization.h"
+#include "HAL/IConsoleManager.h"
 #include "Widgets/Text/STextBlock.h"
 #include "Framework/MultiBox/MultiBoxBuilder.h"
 #include "Widgets/Input/SEditableTextBox.h"
 #include "Widgets/Input/SComboButton.h"
 
 #include "Blueprint/WidgetNavigation.h"
+#include "Input/NavigationMethod.h"
 
 #include "IDetailChildrenBuilder.h"
+#include "IDetailPropertyRow.h"
+#include "IPropertyUtilities.h"
+#include "SInstancedStructPicker.h"
+#include "StructViewerModule.h"
+#include "ISinglePropertyView.h"
+#include "IStructureDataProvider.h"
+#include "IStructureDetailsView.h"
+#include "Styling/SlateIconFinder.h"
+
+#include "ClassViewerFilter.h"
 #include "DetailWidgetRow.h"
 #include "DetailLayoutBuilder.h"
+#include "PropertyCustomizationHelpers.h"
+#include "PropertyEditorModule.h"
 #include "ScopedTransaction.h"
 #include "Details/SFunctionSelector.h"
 #include "Framework/Application/SlateApplication.h"
@@ -30,9 +44,9 @@ void FWidgetNavigationCustomization::CustomizeHeader(TSharedRef<IPropertyHandle>
 
 void FWidgetNavigationCustomization::CustomizeChildren(TSharedRef<IPropertyHandle> PropertyHandle, IDetailChildrenBuilder& ChildBuilder, IPropertyTypeCustomizationUtils& CustomizationUtils)
 {
-	TWeakPtr<IPropertyHandle> PropertyHandlePtr(PropertyHandle);
+	PropertyUtilities = CustomizationUtils.GetPropertyUtilities();
 
-	//IDetailCategoryBuilder& PropertyCategory = DetailLayout.EditCategory("Events", LOCTEXT("Events", "Events"), ECategoryPriority::Uncommon);
+	TWeakPtr<IPropertyHandle> PropertyHandlePtr(PropertyHandle);
 
 	MakeNavRow(PropertyHandlePtr, ChildBuilder, EUINavigation::Left, LOCTEXT("LeftNavigation", "Left"));
 	MakeNavRow(PropertyHandlePtr, ChildBuilder, EUINavigation::Right, LOCTEXT("RightNavigation", "Right"));
@@ -40,6 +54,40 @@ void FWidgetNavigationCustomization::CustomizeChildren(TSharedRef<IPropertyHandl
 	MakeNavRow(PropertyHandlePtr, ChildBuilder, EUINavigation::Down, LOCTEXT("DownNavigation", "Down"));
 	MakeNavRow(PropertyHandlePtr, ChildBuilder, EUINavigation::Next, LOCTEXT("NextNavigation", "Next"));
 	MakeNavRow(PropertyHandlePtr, ChildBuilder, EUINavigation::Previous, LOCTEXT("PreviousNavigation", "Previous"));
+
+	IConsoleVariable* CVar = IConsoleManager::Get().FindConsoleVariable(TEXT("Slate.Navigation.Experimental"));
+	if (CVar && CVar->GetBool())
+	{
+		MakeNavigationMethodStructPicker(PropertyHandlePtr, ChildBuilder);
+
+		if (DoesWidgetSupportRoutingPolicy(PropertyHandlePtr))
+		{
+			ChildBuilder.AddCustomRow(LOCTEXT("RoutingPolicy", "Routing Policy"))
+				.NameContent()
+				[
+					SNew(STextBlock)
+						.Font(IDetailLayoutBuilder::GetDetailFont())
+						.Text(LOCTEXT("RoutingPolicy", "Routing Policy"))
+				]
+				.ValueContent()
+				.MaxDesiredWidth(300)
+				[
+					SNew(SComboButton)
+						.HAlign(HAlign_Center)
+						.ButtonContent()
+						[
+							SNew(STextBlock)
+								.Text(this, &FWidgetNavigationCustomization::GetRoutingPolicyText, PropertyHandlePtr)
+								.Font(IDetailLayoutBuilder::GetDetailFont())
+						]
+						.ContentPadding(FMargin(2.0f, 1.0f))
+						.MenuContent()
+						[
+							MakeRoutingPolicyMenu(PropertyHandle)
+						]
+				];
+		}
+	}
 }
 
 EUINavigationRule FWidgetNavigationCustomization::GetNavigationRule(TWeakPtr<IPropertyHandle> PropertyHandle, EUINavigation Nav) const
@@ -96,15 +144,29 @@ FText FWidgetNavigationCustomization::GetNavigationText(TWeakPtr<IPropertyHandle
 	return FText::GetEmpty();
 }
 
-FText FWidgetNavigationCustomization::GetExplictWidget(TWeakPtr<IPropertyHandle> PropertyHandle, EUINavigation Nav) const
+
+bool FWidgetNavigationCustomization::DoesWidgetSupportRoutingPolicy(TWeakPtr<IPropertyHandle> PropertyHandle) const
 {
-	TOptional<FName> OptionalName = GetUniformNavigationTargetOrFunction(PropertyHandle, Nav);
-	if (!OptionalName.IsSet())
+	TArray<UObject*> OuterObjects;
+	TSharedPtr<IPropertyHandle> PropertyHandlePtr = PropertyHandle.Pin();
+	PropertyHandlePtr->GetOuterObjects(OuterObjects);
+
+	if (OuterObjects.Num() == 0)
 	{
-		return LOCTEXT("NavigationMultipleValues", "Multiple Values");
+		return false;
 	}
 
-	return FText::FromName(OptionalName.GetValue());
+	// Only support routing for widgets that layout children
+	for (UObject* OuterObject : OuterObjects)
+	{
+		UWidget* Widget = Cast<UWidget>(OuterObject);
+		if (Widget == nullptr || !Widget->IsA<UPanelWidget>())
+		{
+			return false;
+		}
+	}
+
+	return true;
 }
 
 TOptional<FName> FWidgetNavigationCustomization::GetUniformNavigationTargetOrFunction(TWeakPtr<IPropertyHandle> PropertyHandle, EUINavigation Nav) const
@@ -404,6 +466,7 @@ void FWidgetNavigationCustomization::SetNav(UWidget* Widget, EUINavigation Nav,
 
 	Widget->Modify();
 
+
 	UWidgetNavigation* WidgetNavigation = Widget->Navigation;
 	if (!Widget->Navigation)
 	{
@@ -416,6 +479,7 @@ void FWidgetNavigationCustomization::SetNav(UWidget* Widget, EUINavigation Nav,
 		}
 	}
 
+
 	FWidgetNavigationData* DirectionNavigation = nullptr;
 
 	switch ( Nav )
@@ -466,4 +530,315 @@ void FWidgetNavigationCustomization::SetNav(UWidget* Widget, EUINavigation Nav,
 	}
 }
 
+EWidgetNavigationRoutingPolicy FWidgetNavigationCustomization::GetRoutingPolicy(TWeakPtr<IPropertyHandle> PropertyHandle) const
+{
+	TArray<UObject*> OuterObjects;
+	TSharedPtr<IPropertyHandle> PropertyHandlePtr = PropertyHandle.Pin();
+	PropertyHandlePtr->GetOuterObjects(OuterObjects);
+
+	EWidgetNavigationRoutingPolicy Policy = EWidgetNavigationRoutingPolicy::Default;
+	for ( UObject* OuterObject : OuterObjects )
+	{
+		if ( UWidget* Widget = Cast<UWidget>(OuterObject) )
+		{
+			EWidgetNavigationRoutingPolicy CurrentPolicy = EWidgetNavigationRoutingPolicy::Default;
+			UWidgetNavigation* WidgetNavigation = Widget->Navigation;
+			if ( Widget->Navigation != nullptr )
+			{
+				CurrentPolicy = WidgetNavigation->RoutingPolicy;
+			}
+
+			if ( Policy != EWidgetNavigationRoutingPolicy::Default && CurrentPolicy != Policy )
+			{
+				return EWidgetNavigationRoutingPolicy::Default;
+			}
+			Policy = CurrentPolicy;
+		}
+	}
+
+	return Policy;
+}
+
+FText FWidgetNavigationCustomization::GetRoutingPolicyText(TWeakPtr<IPropertyHandle> PropertyHandle) const
+{
+	EWidgetNavigationRoutingPolicy Rule = GetRoutingPolicy(PropertyHandle);
+	if (UEnum* EnumPtr = StaticEnum<EWidgetNavigationRoutingPolicy>())
+	{
+		return EnumPtr->GetDisplayNameTextByValue(StaticCast<int32>(Rule));
+	}
+	return FText::GetEmpty();
+}
+
+FText FWidgetNavigationCustomization::GetExplictWidget(TWeakPtr<IPropertyHandle> PropertyHandle, EUINavigation Nav) const
+{
+	TOptional<FName> OptionalName = GetUniformNavigationTargetOrFunction(PropertyHandle, Nav);
+	if (!OptionalName.IsSet())
+	{
+		return LOCTEXT("NavigationMultipleValues", "Multiple Values");
+	}
+
+	return FText::FromName(OptionalName.GetValue());
+}
+
+TSharedRef<class SWidget> FWidgetNavigationCustomization::MakeRoutingPolicyMenu(TWeakPtr<IPropertyHandle> PropertyHandle)
+{
+	FMenuBuilder MenuBuilder(true, NULL);
+
+	if (UEnum* EnumPtr = StaticEnum<EWidgetNavigationRoutingPolicy>())
+	{
+		for (int32 Index = 0; Index < EnumPtr->NumEnums(); ++Index)
+		{			
+			if (!EnumPtr->HasMetaData(TEXT("Hidden"), Index))
+			{
+				const EWidgetNavigationRoutingPolicy Value = StaticCast<EWidgetNavigationRoutingPolicy>(EnumPtr->GetValueByIndex(Index));
+				const FText DisplayName = EnumPtr->GetDisplayNameTextByIndex(Index);
+				const FText Tooltip = EnumPtr->GetToolTipTextByIndex(Index);
+
+				FUIAction Action(FExecuteAction::CreateSP(this, &FWidgetNavigationCustomization::HandleRoutingPolicyMenuEntryClicked, PropertyHandle, Value));
+				MenuBuilder.AddMenuEntry(DisplayName, Tooltip, FSlateIcon(), Action);
+			}
+		}
+	}
+
+	return MenuBuilder.MakeWidget();
+}
+
+void FWidgetNavigationCustomization::HandleRoutingPolicyMenuEntryClicked(TWeakPtr<IPropertyHandle> PropertyHandle, EWidgetNavigationRoutingPolicy Policy)
+{
+	TArray<UObject*> OuterObjects;
+	TSharedPtr<IPropertyHandle> PropertyHandlePtr = PropertyHandle.Pin();
+	PropertyHandlePtr->GetOuterObjects(OuterObjects);
+
+	const FScopedTransaction Transaction(LOCTEXT("InitializeNavigation", "Edit Widget Navigation"));
+
+	for (UObject* OuterObject : OuterObjects)
+	{
+		if (UWidget* Widget = Cast<UWidget>(OuterObject))
+		{
+			const TSharedPtr<FWidgetBlueprintEditor> WidgetBlueprintEditor = Editor.Pin();
+			if (Widget == WidgetBlueprintEditor->GetPreview())
+			{
+				SetRoutingPolicy(Widget, Policy);
+
+				UWidgetBlueprint* WidgetBlueprint = WidgetBlueprintEditor->GetWidgetBlueprintObj();
+				if (WidgetBlueprint && WidgetBlueprint->GeneratedClass)
+				{
+					SetRoutingPolicy(WidgetBlueprint->GeneratedClass->GetDefaultObject<UUserWidget>(), Policy);
+					FBlueprintEditorUtils::MarkBlueprintAsStructurallyModified(WidgetBlueprint);
+				}
+			}
+			else
+			{
+				FWidgetReference WidgetReference = WidgetBlueprintEditor->GetReferenceFromPreview(Widget);
+
+				SetRoutingPolicy(WidgetReference.GetPreview(), Policy);
+				SetRoutingPolicy(WidgetReference.GetTemplate(), Policy);
+			}
+		}
+	}
+}
+
+void FWidgetNavigationCustomization::SetRoutingPolicy(UWidget* Widget, EWidgetNavigationRoutingPolicy Policy)
+{
+	if (!Widget)
+	{
+		return;
+	}
+
+	Widget->Modify();
+
+	UWidgetNavigation* WidgetNavigation = Widget->Navigation;
+	if (!Widget->Navigation)
+	{
+		WidgetNavigation = NewObject<UWidgetNavigation>(Widget);
+		WidgetNavigation->SetFlags(RF_Transactional);
+		if (Widget->IsTemplate())
+		{
+			EObjectFlags TemplateFlags = Widget->GetMaskedFlags(RF_PropagateToSubObjects) | RF_DefaultSubObject;
+			WidgetNavigation->SetFlags(TemplateFlags);
+		}
+	}
+
+	WidgetNavigation->RoutingPolicy = Policy;
+
+	if ( WidgetNavigation->IsDefaultNavigation() )
+	{
+		// If the navigation rules are all set to the defaults, remove the navigation
+		// information from the widget.
+		Widget->Navigation = nullptr;
+	}
+	else
+	{
+		Widget->Navigation = WidgetNavigation;
+	}
+}
+
+void FWidgetNavigationCustomization::OnNavigationMethodPrePropertyChange(TWeakPtr<IPropertyHandle> PropertyHandle)
+{
+	TArray<UObject*> OuterObjects;
+	TSharedPtr<IPropertyHandle> PropertyHandlePtr = PropertyHandle.Pin();
+	PropertyHandlePtr->GetOuterObjects(OuterObjects);
+
+	const UScriptStruct* Struct = nullptr;
+	for (UObject* OuterObject : OuterObjects)
+	{
+		if (UWidget* Widget = Cast<UWidget>(OuterObject))
+		{
+			Widget->Modify();
+		}
+	}
+}
+
+void FWidgetNavigationCustomization::MakeNavigationMethodStructPicker(TWeakPtr<IPropertyHandle> PropertyHandle, IDetailChildrenBuilder& ChildBuilder)
+{
+	TSharedPtr<IPropertyHandle> PropertyHandlePinned = PropertyHandle.Pin();
+	TSharedPtr<IPropertyHandle> NavigationHandle = PropertyHandlePinned->GetChildHandle(GET_MEMBER_NAME_CHECKED(UWidgetNavigation, NavigationMethod));
+
+	if (NavigationHandle)
+	{
+		NavigationHandle->SetOnChildPropertyValuePreChange(FSimpleDelegate::CreateSP(this, &FWidgetNavigationCustomization::OnNavigationMethodPrePropertyChange, PropertyHandle));
+		ChildBuilder.AddProperty(NavigationHandle.ToSharedRef());
+		return;
+	}
+
+	// Widget Navigation is instanced. So child NavigationMethod property is not available. Create a mock drop-down that will reload with valid property handle on struct type selection
+
+	TSharedRef<FInstancedStructFilter> StructFilter = MakeShared<FInstancedStructFilter>();
+	StructFilter->BaseStruct = FNavigationMethod::StaticStruct();
+	StructFilter->bAllowUserDefinedStructs = false; // Only allow user defined structs when BaseStruct is not set.
+	StructFilter->bAllowBaseStruct = false;
+
+	FStructViewerInitializationOptions Options;
+	Options.bShowNoneOption = true;
+	Options.StructFilter = StructFilter;
+	Options.NameTypeToDisplay = EStructViewerNameTypeToDisplay::DisplayName;
+	Options.DisplayMode = EStructViewerDisplayMode::ListView;
+	Options.bAllowViewOptions = true;
+
+	const FOnStructPicked OnPicked(FOnStructPicked::CreateSPLambda(this, [PropertyHandle, this](const UScriptStruct* SelectedStruct) {
+		HandleNavigationMethodStructPicked(PropertyHandle, SelectedStruct);
+	}));
+	
+	ChildBuilder.AddCustomRow(LOCTEXT("NavigationMethod", "Navigation Method"))
+		.NameContent()
+		[
+			SNew(STextBlock)
+			.Font(IDetailLayoutBuilder::GetDetailFont())
+			.Text(LOCTEXT("NavigationMethod", "Navigation Method"))
+		]
+		.ValueContent()
+		[
+			SNew(SComboButton)
+			.ContentPadding(0)
+			.ButtonContent()
+			[
+				SNew(SHorizontalBox)
+				+ SHorizontalBox::Slot()
+				.AutoWidth()
+				.VAlign(VAlign_Center)
+				.Padding(0.0f, 0.0f, 4.0f, 0.0f)
+				[
+					SNew(SImage)
+						.Image(FSlateIconFinder::FindIconBrushForClass(FNavigationMethod::StaticStruct(), "ClassIcon.Object"))
+				]
+				+ SHorizontalBox::Slot()
+				.VAlign(VAlign_Center)
+				[
+					SNew(STextBlock)
+						.Text(LOCTEXT("NavigationMethodNoneText", "None"))
+						.Font(IDetailLayoutBuilder::GetDetailFont())
+				]
+			]
+			.MenuContent()
+			[
+				SNew(SBox)
+				.WidthOverride(280)
+				[
+					SNew(SVerticalBox)
+					+ SVerticalBox::Slot()
+					.AutoHeight()
+					.MaxHeight(500)
+					[
+						FModuleManager::LoadModuleChecked<FStructViewerModule>("StructViewer").CreateStructViewer(Options, OnPicked)
+					]
+				]
+			]
+		];
+}
+
+void FWidgetNavigationCustomization::HandleNavigationMethodStructPicked(TWeakPtr<IPropertyHandle> PropertyHandle, const UScriptStruct* NavigationMethodStruct)
+{
+	TArray<UObject*> OuterObjects;
+	TSharedPtr<IPropertyHandle> PropertyHandlePinned = PropertyHandle.Pin();
+	PropertyHandlePinned->GetOuterObjects(OuterObjects);
+
+	const FScopedTransaction Transaction(LOCTEXT("InitializeNavigation", "Edit Widget Navigation"));
+
+	for (UObject* OuterObject : OuterObjects)
+	{
+		if (UWidget* Widget = Cast<UWidget>(OuterObject))
+		{
+			const TSharedPtr<FWidgetBlueprintEditor> WidgetBlueprintEditor = Editor.Pin();
+			if (Widget == WidgetBlueprintEditor->GetPreview())
+			{
+				SetNavigationMethodStruct(Widget, NavigationMethodStruct);
+
+				UWidgetBlueprint* WidgetBlueprint = WidgetBlueprintEditor->GetWidgetBlueprintObj();
+				if (WidgetBlueprint && WidgetBlueprint->GeneratedClass)
+				{
+					SetNavigationMethodStruct(WidgetBlueprint->GeneratedClass->GetDefaultObject<UUserWidget>(), NavigationMethodStruct);
+					FBlueprintEditorUtils::MarkBlueprintAsStructurallyModified(WidgetBlueprint);
+				}
+			}
+			else
+			{
+				FWidgetReference WidgetReference = WidgetBlueprintEditor->GetReferenceFromPreview(Widget);
+
+				SetNavigationMethodStruct(WidgetReference.GetPreview(), NavigationMethodStruct);
+				SetNavigationMethodStruct(WidgetReference.GetTemplate(), NavigationMethodStruct);
+			}
+		}
+	}
+}
+
+void FWidgetNavigationCustomization::SetNavigationMethodStruct(UWidget* Widget, const UScriptStruct* NavigationMethodStruct)
+{
+	if (!Widget)
+	{
+		return;
+	}
+
+	Widget->Modify();
+
+	UWidgetNavigation* WidgetNavigation = Widget->Navigation;
+	if (!Widget->Navigation)
+	{
+		WidgetNavigation = NewObject<UWidgetNavigation>(Widget);
+		WidgetNavigation->SetFlags(RF_Transactional);
+		if (Widget->IsTemplate())
+		{
+			EObjectFlags TemplateFlags = Widget->GetMaskedFlags(RF_PropagateToSubObjects) | RF_DefaultSubObject;
+			WidgetNavigation->SetFlags(TemplateFlags);
+		}
+	}
+
+	WidgetNavigation->NavigationMethod.InitializeAsScriptStruct(NavigationMethodStruct);
+
+	if ( WidgetNavigation->IsDefaultNavigation() )
+	{
+		// If the navigation rules are all set to the defaults, remove the navigation
+		// information from the widget.
+		Widget->Navigation = nullptr;
+	}
+	else
+	{
+		Widget->Navigation = WidgetNavigation;
+	}
+
+	// Force refresh the properties panel to get a handle on the NavigationMethod property
+	PropertyUtilities->RequestForceRefresh();
+}
+
+
+
 #undef LOCTEXT_NAMESPACE
diff --git a/Engine/Source/Editor/UMGEditor/Private/Customizations/WidgetNavigationCustomization.h b/Engine/Source/Editor/UMGEditor/Private/Customizations/WidgetNavigationCustomization.h
--- a/Engine/Source/Editor/UMGEditor/Private/Customizations/WidgetNavigationCustomization.h
+++ b/Engine/Source/Editor/UMGEditor/Private/Customizations/WidgetNavigationCustomization.h
@@ -6,9 +6,11 @@
 #include "Layout/Visibility.h"
 #include "WidgetBlueprintEditor.h"
 #include "IPropertyTypeCustomization.h"
+#include "StructUtils/InstancedStruct.h"
 #include "PropertyHandle.h"
 
 class IDetailChildrenBuilder;
+class IPropertyUtilities;
 
 class FWidgetNavigationCustomization : public IPropertyTypeCustomization
 {
@@ -52,6 +54,19 @@ private:
 
 	void SetNav(UWidget* Widget, EUINavigation Nav, TOptional<EUINavigationRule> Rule, TOptional<FName> WidgetToFocus);
 
+	bool DoesWidgetSupportRoutingPolicy(TWeakPtr<IPropertyHandle> PropertyHandle) const;
+	EWidgetNavigationRoutingPolicy GetRoutingPolicy(TWeakPtr<IPropertyHandle> PropertyHandle) const;
+	FText GetRoutingPolicyText(TWeakPtr<IPropertyHandle> PropertyHandle) const;
+	TSharedRef<class SWidget> MakeRoutingPolicyMenu(TWeakPtr<IPropertyHandle> PropertyHandle);
+	void HandleRoutingPolicyMenuEntryClicked(TWeakPtr<IPropertyHandle> PropertyHandle, EWidgetNavigationRoutingPolicy Policy);
+	void SetRoutingPolicy(UWidget* Widget, EWidgetNavigationRoutingPolicy Policy);
+
+	void OnNavigationMethodPrePropertyChange(TWeakPtr<IPropertyHandle> PropertyHandle);
+	void MakeNavigationMethodStructPicker(TWeakPtr<IPropertyHandle> PropertyHandle, IDetailChildrenBuilder& ChildBuilder);
+	void HandleNavigationMethodStructPicked(TWeakPtr<IPropertyHandle> PropertyHandle, const UScriptStruct* InStruct);
+	void SetNavigationMethodStruct(UWidget* Widget, const UScriptStruct* NavigationMethodStruct);
+
 private:
 	TWeakPtr<FWidgetBlueprintEditor> Editor;
+	TSharedPtr<IPropertyUtilities> PropertyUtilities;
 };
diff --git a/Engine/Source/Editor/UMGEditor/UMGEditor.Build.cs b/Engine/Source/Editor/UMGEditor/UMGEditor.Build.cs
--- a/Engine/Source/Editor/UMGEditor/UMGEditor.Build.cs
+++ b/Engine/Source/Editor/UMGEditor/UMGEditor.Build.cs
@@ -48,6 +48,8 @@ public class UMGEditor : ModuleRules
 				"Slate",
 				"SlateCore",
 				"SlateRHIRenderer",
+				"StructViewer",
+				"StructUtilsEditor",
 				"StatusBar",
 				"MessageLog",
 				"MovieScene",
diff --git a/Engine/Source/Runtime/SlateCore/Private/Input/HittestGrid.cpp b/Engine/Source/Runtime/SlateCore/Private/Input/HittestGrid.cpp
--- a/Engine/Source/Runtime/SlateCore/Private/Input/HittestGrid.cpp
+++ b/Engine/Source/Runtime/SlateCore/Private/Input/HittestGrid.cpp
@@ -1,11 +1,17 @@
 // Copyright Epic Games, Inc. All Rights Reserved.
 
 #include "Input/HittestGrid.h"
+#include "Input/NavigationMethod.h"
+#include "Input/NavigationMethodOrthogonal.h"
+#include "Input/NavigationMethodProximity.h"
+#include "Input/NavigationRouting.h"
 #include "Rendering/RenderingCommon.h"
 #include "SlateGlobals.h"
+#include "StructUtils/InstancedStruct.h"
 #include "HAL/IConsoleManager.h"
 #include "Rendering/DrawElements.h"
 #include "Styling/CoreStyle.h"
+#include "Types/NavigationMetaData.h"
 
 DEFINE_LOG_CATEGORY_STATIC(LogHittestDebug, Display, All);
 
@@ -26,16 +32,16 @@ int32 HittestGrid_CollapsedHittestGridArraySizeMax = 0;
 int32 HittestGrid_CollapsedWidgetsArraySizeMax = 0;
 #endif
 
+static bool GEnableNavigationExperimental = false;
+static FAutoConsoleVariableRef CVarEnableNavigationExperimental(
+	TEXT("Slate.Navigation.Experimental"),
+	GEnableNavigationExperimental,
+	TEXT("If true, will use experimental navigation system"));
+
 #if WITH_SLATE_DEBUGGING
 FHittestGrid::FDebuggingFindNextFocusableWidget FHittestGrid::OnFindNextFocusableWidgetExecuted;
 #endif //WITH_SLATE_DEBUGGING
 
-constexpr bool IsCompatibleUserIndex(int32 RequestedUserIndex, int32 TestUserIndex)
-{
-	// INDEX_NONE means the user index is compatible with all
-	return RequestedUserIndex == INDEX_NONE || TestUserIndex == INDEX_NONE|| RequestedUserIndex == TestUserIndex;
-}
-
 //
 // Helper Functions
 //
@@ -100,7 +106,7 @@ bool ContainsInteractableWidget(const TArray<FWidgetAndPointer>& PathToTest)
 }
 
 
-const FVector2f CellSize(128.0f, 128.0f);
+const FVector2f FHittestGrid::CellSize(128.0f, 128.0f);
 
 //
 // FHittestGrid::FWidgetIndex
@@ -165,6 +171,9 @@ FHittestGrid::FHittestGrid()
 	, GridOrigin(0, 0)
 	, GridSize(0, 0)
 	, CurrentUserIndex(INDEX_NONE)
+#if WITH_SLATE_DEBUGGING
+	, LastNavigationMethod()
+#endif
 {
 }
 
@@ -294,26 +303,18 @@ void FHittestGrid::ClearInternal(int32 TotalCells)
 	AppendedGridArray.Reset();
 }
 
-bool FHittestGrid::IsDescendantOf(const SWidget* ParentWidget, const FWidgetData& ChildData) const
+TSharedPtr<SWidget> FHittestGrid::FindNextFocusableWidget(const FArrangedWidget& StartingWidget, const EUINavigation Direction, const FNavigationReply& NavigationReply, const FArrangedWidget& RuleWidget, int32 UserIndex)
 {
-	const TSharedPtr<SWidget> ChildWidgetPtr = ChildData.GetWidget();
-	const SWidget* CurWidget = ChildWidgetPtr.Get();
-	if (CurWidget == ParentWidget)
+	if (GEnableNavigationExperimental)
 	{
-		return false;
+		return FindNextFocusableWidgetExperimental(StartingWidget, Direction, NavigationReply, RuleWidget, UserIndex);
 	}
+	return FindNextFocusableWidgetDefault(StartingWidget, Direction, NavigationReply, RuleWidget, UserIndex);
+}
 
-	while (CurWidget)
-	{
-		if (ParentWidget == CurWidget)
-		{
-			return true;
-		}
-		CurWidget = CurWidget->Advanced_GetPaintParentWidget().Get();
-	}
 
-	return false;
-}
+// TODO: NavigationMethodOrthogonal class deprecates the following and can be eventually removed. 
+//		UE::Slate::Private, FindFocusableWidget, FindNextFocusableWidgetDefault
 
 namespace UE::Slate::Private
 {
@@ -539,7 +540,7 @@ TSharedPtr<SWidget> FHittestGrid::FindFocusableWidget(FSlateRect WidgetRect, con
 
 #undef AddToNextFocusableWidgetCondidateDebugResults
 
-TSharedPtr<SWidget> FHittestGrid::FindNextFocusableWidget(const FArrangedWidget& StartingWidget, const EUINavigation Direction, const FNavigationReply& NavigationReply, const FArrangedWidget& RuleWidget, int32 UserIndex)
+TSharedPtr<SWidget> FHittestGrid::FindNextFocusableWidgetDefault(const FArrangedWidget& StartingWidget, const EUINavigation Direction, const FNavigationReply& NavigationReply, const FArrangedWidget& RuleWidget, int32 UserIndex)
 {
 	FGeometry StartingWidgetGeo = StartingWidget.Widget->GetPaintSpaceGeometry();
 	StartingWidgetGeo.AppendTransform(FSlateLayoutTransform(-GridWindowOrigin));
@@ -624,6 +625,90 @@ TSharedPtr<SWidget> FHittestGrid::FindNextFocusableWidget(const FArrangedWidget&
 	return Widget;
 }
 
+bool FHittestGrid::IsDescendantOf(const SWidget* ParentWidget, const FWidgetData& ChildData) const
+{
+	const TSharedPtr<SWidget> ChildWidgetPtr = ChildData.GetWidget();
+	const SWidget* CurWidget = ChildWidgetPtr.Get();
+	if (CurWidget == ParentWidget)
+	{
+		return false;
+	}
+
+	while (CurWidget)
+	{
+		if (ParentWidget == CurWidget)
+		{
+			return true;
+		}
+		CurWidget = CurWidget->Advanced_GetPaintParentWidget().Get();
+	}
+
+	return false;
+}
+
+TSharedPtr<SWidget> FHittestGrid::FindNextFocusableWidgetExperimental(const FArrangedWidget& StartingWidget, const EUINavigation Direction, const FNavigationReply& NavigationReply, const FArrangedWidget& RuleWidget, int32 UserIndex)
+{
+	// Ensure that the hit test grid is valid before proceeding
+	if (NumCells.X < 1 || NumCells.Y < 1)
+	{
+		return TSharedPtr<SWidget>();
+	}
+
+#if WITH_SLATE_DEBUGGING
+	TArray<FDebuggingFindNextFocusableWidgetArgs::FWidgetResult> IntermediateResults;
+	TArray<FDebuggingFindNextFocusableWidgetArgs::FWidgetResult>* IntermediateResultsPtr = OnFindNextFocusableWidgetExecuted.IsBound() ? &IntermediateResults : nullptr;
+#else
+	TArray<FDebuggingFindNextFocusableWidgetArgs::FWidgetResult>* IntermediateResultsPtr = nullptr;
+#endif
+	
+	TSharedPtr<SWidget> NextFocusedWidget;
+	TInstancedStruct<FNavigationMethod> NavigationMethod = GetNavigationMethod(StartingWidget);
+	if (NavigationMethod.IsValid())
+	{
+		FNavigationMethod* NavigationMethodPtr = NavigationMethod.GetMutablePtr();
+
+		NavigationMethodPtr->Initialize(this, IntermediateResultsPtr);
+		NextFocusedWidget = NavigationMethodPtr->FindNextFocusableWidget(StartingWidget, Direction, NavigationReply, RuleWidget, UserIndex);
+	}
+	else
+	{
+		// Fallback to existing behavior
+		NextFocusedWidget = FindNextFocusableWidgetDefault(StartingWidget, Direction, NavigationReply, RuleWidget, UserIndex);
+	}
+
+	FNavigationRoutingParams RoutingParams(StartingWidget.Widget, NextFocusedWidget, this, NavigationReply, UserIndex);
+	NextFocusedWidget = FNavigationRouting::RouteNavigationFocus(RoutingParams);
+
+#if WITH_SLATE_DEBUGGING
+	if (IntermediateResultsPtr)
+	{
+		FDebuggingFindNextFocusableWidgetArgs Args {StartingWidget, Direction, NavigationReply, RuleWidget, UserIndex, NextFocusedWidget, MoveTemp(*IntermediateResultsPtr)};
+		OnFindNextFocusableWidgetExecuted.Broadcast(this, Args);
+	}
+	
+	NavigationMethod.Reset();
+
+	if (GSlateNavigationDebugging && NavigationMethod.IsValid())
+	{
+		// Duplicate, as this pointer may outlive the navigation method associated with raw instanced struct pointer
+		LastNavigationMethod = NavigationMethod;
+	}
+#endif
+
+	return NextFocusedWidget;
+}
+
+FSlateRect FHittestGrid::GetWidgetRenderBoundingRect(const TSharedPtr<SWidget>& Widget) const
+{
+	if (!Widget.IsValid())
+	{
+		return FSlateRect();
+	}
+	FGeometry WidgetGeo = Widget->GetPaintSpaceGeometry();
+	WidgetGeo.AppendTransform(FSlateLayoutTransform(-GridWindowOrigin));
+	return WidgetGeo.GetRenderBoundingRect();
+}
+
 FIntPoint FHittestGrid::GetCellCoordinate(UE::Slate::FDeprecateVector2DParameter Position) const
 {
 	return FIntPoint(
@@ -1018,6 +1103,23 @@ void FHittestGrid::RemoveStaleAppendedHittestGrid()
 	}
 }
 
+TInstancedStruct<FNavigationMethod> FHittestGrid::GetNavigationMethod(const FArrangedWidget& StartingWidget)
+{
+	TInstancedStruct<FNavigationMethod> NavigationMethodStruct;
+	TSharedPtr<SWidget> Parent = StartingWidget.Widget;
+	while (!NavigationMethodStruct.IsValid() && Parent != nullptr)
+	{
+		TSharedPtr<FNavigationMetaData> NavigationMetaData = Parent->GetMetaData<FNavigationMetaData>();
+		if (NavigationMetaData.IsValid())
+		{
+			NavigationMethodStruct = NavigationMetaData->GetNavigationMethodStruct();
+		}
+		Parent = Parent->GetParentWidget();
+	}
+
+	return NavigationMethodStruct;
+}
+
 #if WITH_SLATE_DEBUGGING
 void FHittestGrid::LogGrid() const
 {
@@ -1131,6 +1233,15 @@ void FHittestGrid::DisplayGrid(int32 InLayer, const FGeometry& AllottedGeometry,
 	}
 }
 
+void FHittestGrid::DisplayNavigation(int32 InLayer, const FGeometry& AllottedGeometry, FSlateWindowElementList& WindowElementList)
+{
+	if (GSlateNavigationDebugging && LastNavigationMethod.IsValid())
+	{
+		FNavigationMethod* NavigationMethodPtr = LastNavigationMethod.GetMutablePtr();
+		NavigationMethodPtr->DrawDebug(InLayer, AllottedGeometry, WindowElementList);
+	}
+}
+
 TArray<FHittestGrid::FWidgetSortData> FHittestGrid::GetAllWidgetSortDatas() const
 {
 	TArray<FWidgetSortData> Result;
diff --git a/Engine/Source/Runtime/SlateCore/Private/Input/NavigationMethod.cpp b/Engine/Source/Runtime/SlateCore/Private/Input/NavigationMethod.cpp
new file mode 100644
--- /dev/null
+++ b/Engine/Source/Runtime/SlateCore/Private/Input/NavigationMethod.cpp
@@ -0,0 +1,144 @@
+// Copyright Epic Games, Inc. All Rights Reserved.
+
+#include "Input/NavigationMethod.h"
+
+#define LOCTEXT_NAMESPACE "NavigationMethod"
+
+#if WITH_SLATE_DEBUGGING
+
+FText FNavigationMethod::FDebuggingText::Valid = LOCTEXT("StateValid", "Valid"); //~ The widget is valid will be consider as the result
+FText FNavigationMethod::FDebuggingText::NotCompatibleWithUserIndex = LOCTEXT("StateNotCompatibleWithUserIndex", "User Index not compatible"); //~ The widget is not compatible with the requested user index
+FText FNavigationMethod::FDebuggingText::DoesNotIntersect = LOCTEXT("StateDoesNotIntersect", "Does not intersect"); //~ The widget rect is not in the correct direction or is not intersecting with the "swept" rectangle
+FText FNavigationMethod::FDebuggingText::PreviousWidgetIsBetter = LOCTEXT("StatePreviousWidgetIsBetter", "Previous Widget was better"); //~ The widget would be valid but the previous valid is closer
+FText FNavigationMethod::FDebuggingText::NotADescendant = LOCTEXT("StateNotADescendant", "Not a descendant"); //~ We have a non escape boundary condition and the widget isn't a descendant of our boundary
+FText FNavigationMethod::FDebuggingText::Disabled = LOCTEXT("StateNotEnabled", "Disabled"); //~ The widget is not enabled
+FText FNavigationMethod::FDebuggingText::ParentDisabled = LOCTEXT("StateParentNotEnabled", "ParentDisabled"); //~ A parent of the widget is disabled
+FText FNavigationMethod::FDebuggingText::DoesNotSuportKeyboardFocus = LOCTEXT("StateDoesNotSuportKeyboardFocus", "Keyboard focus unsupported"); //~ THe widget does not support keyboard focus
+#endif // WITH_SLATE_DEBUGGING
+
+void FNavigationMethod::Initialize(const FHittestGrid* InHittestGrid, TArray<FDebugWidgetResult>* InIntermediateResultsPtr)
+{
+	HittestGrid = InHittestGrid;
+	IntermediateResultsPtr = InIntermediateResultsPtr;
+	DisabledDestinations.Empty();
+}
+
+TSharedPtr<SWidget> FNavigationMethod::FindNextFocusableWidget(const FArrangedWidget& StartingWidget, const EUINavigation Direction, const FNavigationReply& NavigationReply, const FArrangedWidget& RuleWidget, int32 InUserIndex) 
+{ 
+	return StartingWidget.Widget; 
+}
+
+FIntPoint FNavigationMethod::GetCellCoordinate(FVector2f Position) const
+{
+	check(HittestGrid);
+	return HittestGrid->GetCellCoordinate(Position);
+}
+
+bool FNavigationMethod::IsValidCellCoordinate(int32 X, int32 Y) const
+{	
+	check(HittestGrid);
+	return HittestGrid->IsValidCellCoord(X, Y);
+}
+
+bool FNavigationMethod::IsParentsEnabled(const SWidget* Widget)
+{
+	while (Widget)
+	{
+		if (!Widget->IsEnabled())
+		{
+			return false;
+		}
+		Widget = Widget->Advanced_GetPaintParentWidget().Get();
+	}
+	return true;
+}
+
+void FNavigationMethod::ForEachFocusableWidgetsInCell(int32 X, int32 Y, const FNavigationReply& NavigationReply, int32 UserIndex, FWidgetFunc WidgetFunc)
+{
+	check(HittestGrid);
+
+	FHittestGrid::FCollapsedWidgetsArray WidgetIndexes;
+	HittestGrid->GetCollapsedWidgets(WidgetIndexes, X, Y);
+
+	for (int32 WidgetIndex = WidgetIndexes.Num() - 1; WidgetIndex >= 0; --WidgetIndex)
+	{
+		const FHittestGrid::FWidgetData& WidgetData = WidgetIndexes[WidgetIndex].GetWidgetData();
+		const TSharedPtr<SWidget> Widget = WidgetData.GetWidget();
+		if (!Widget.IsValid())
+		{
+			continue;
+		}
+
+		if (!HittestGrid->IsCompatibleUserIndex(UserIndex, WidgetData.UserIndex))
+		{
+#if WITH_SLATE_DEBUGGING
+			AddDebugIntermediateResult(Widget, FDebuggingText::NotCompatibleWithUserIndex);
+#endif
+			continue;
+		}
+
+		// If we have a non escape boundary condition and this widget isn't a descendant of our boundary condition widget then it's invalid so we keep looking.
+		if (NavigationReply.GetBoundaryRule() != EUINavigationRule::Escape
+			&& NavigationReply.GetHandler().IsValid()
+			&& !HittestGrid->IsDescendantOf(NavigationReply.GetHandler().Get(), WidgetData))
+		{
+#if WITH_SLATE_DEBUGGING
+			AddDebugIntermediateResult(Widget, FDebuggingText::NotADescendant);
+#endif
+			continue;
+		}
+
+		if (!Widget->IsEnabled())
+		{
+#if WITH_SLATE_DEBUGGING
+			AddDebugIntermediateResult(Widget, FDebuggingText::Disabled);
+#endif
+			continue;
+		}
+
+		if (!Widget->SupportsKeyboardFocus())
+		{
+#if WITH_SLATE_DEBUGGING
+			AddDebugIntermediateResult(Widget, FDebuggingText::DoesNotSuportKeyboardFocus);
+#endif
+			continue;
+		}
+
+		if (DisabledDestinations.Contains(Widget))
+		{
+#if WITH_SLATE_DEBUGGING
+			AddDebugIntermediateResult(Widget, FDebuggingText::ParentDisabled);
+#endif
+			continue;
+		}
+
+		bool bValid = WidgetFunc(Widget);
+		if (bValid)
+		{
+#if WITH_SLATE_DEBUGGING
+			AddDebugIntermediateResult(Widget, FDebuggingText::Valid);
+#endif
+		}
+	}
+}
+
+#if WITH_SLATE_DEBUGGING
+
+void FNavigationMethod::DrawDebug(int32 InLayer, const FGeometry& AllottedGeometry, FSlateWindowElementList& WindowElementList)
+{
+	(void)InLayer;
+	(void)AllottedGeometry;
+	(void)WindowElementList;
+}
+
+void FNavigationMethod::AddDebugIntermediateResult(const TSharedPtr<const SWidget>& InWidget, const FText Result) 
+{ 
+	if (IntermediateResultsPtr) 
+	{ 
+		IntermediateResultsPtr->Emplace(InWidget, Result); 
+	} 
+}
+
+#endif //WITH_SLATE_DEBUGGING
+
+#undef LOCTEXT_NAMESPACE
\ No newline at end of file
diff --git a/Engine/Source/Runtime/SlateCore/Private/Input/NavigationMethodOrthogonal.cpp b/Engine/Source/Runtime/SlateCore/Private/Input/NavigationMethodOrthogonal.cpp
new file mode 100644
--- /dev/null
+++ b/Engine/Source/Runtime/SlateCore/Private/Input/NavigationMethodOrthogonal.cpp
@@ -0,0 +1,274 @@
+// Copyright Epic Games, Inc. All Rights Reserved.
+
+#include "Input/NavigationMethodOrthogonal.h"
+
+#define LOCTEXT_NAMESPACE "NavigationMethodOrthogonal"
+
+template<typename TCompareFunc, typename TSourceSideFunc, typename TDestSideFunc>
+TSharedPtr<SWidget> FNavigationMethodOrthogonal::FindFocusableWidget(const FSlateRect WidgetRect, const FSlateRect SweptRect, int32 AxisIndex, int32 Increment, const EUINavigation Direction, const FNavigationReply& NavigationReply, TCompareFunc CompareFunc, TSourceSideFunc SourceSideFunc, TDestSideFunc DestSideFunc, int32 UserIndex)
+{
+	const FIntPoint NumCells = HittestGrid->NumCells;
+	const FVector2f CellSize = HittestGrid->CellSize;
+
+	FIntPoint CurrentCellPoint = HittestGrid->GetCellCoordinate(WidgetRect.GetCenter());
+
+	int32 StartingIndex = CurrentCellPoint[AxisIndex];
+
+	float CurrentSourceSide = SourceSideFunc(WidgetRect);
+
+	int32 StrideAxis, StrideAxisMin, StrideAxisMax;
+
+	// Ensure that the hit test grid is valid before proceeding
+	if (NumCells.X < 1 || NumCells.Y < 1)
+	{
+		return TSharedPtr<SWidget>();
+	}
+
+	if (AxisIndex == 0)
+	{
+		StrideAxis = 1;
+		StrideAxisMin = FMath::Min(FMath::Max(FMath::FloorToInt(SweptRect.Top / CellSize.Y), 0), NumCells.Y - 1);
+		StrideAxisMax = FMath::Min(FMath::Max(FMath::FloorToInt(SweptRect.Bottom / CellSize.Y), 0), NumCells.Y - 1);
+	}
+	else
+	{
+		StrideAxis = 0;
+		StrideAxisMin = FMath::Min(FMath::Max(FMath::FloorToInt(SweptRect.Left / CellSize.X), 0), NumCells.X - 1);
+		StrideAxisMax = FMath::Min(FMath::Max(FMath::FloorToInt(SweptRect.Right / CellSize.X), 0), NumCells.X - 1);
+	}
+
+	bool bWrapped = false;
+	while (CurrentCellPoint[AxisIndex] >= 0 && CurrentCellPoint[AxisIndex] < NumCells[AxisIndex])
+	{
+		FIntPoint StrideCellPoint = CurrentCellPoint;
+		int32 CurrentCellProcessed = CurrentCellPoint[AxisIndex];
+
+		// Increment before the search as a wrap case will change our current cell.
+		CurrentCellPoint[AxisIndex] += Increment;
+
+		FSlateRect BestWidgetRect;
+		TSharedPtr<SWidget> BestWidget = TSharedPtr<SWidget>();
+
+		for (StrideCellPoint[StrideAxis] = StrideAxisMin; StrideCellPoint[StrideAxis] <= StrideAxisMax; ++StrideCellPoint[StrideAxis])
+		{
+			FHittestGrid::FCollapsedWidgetsArray WidgetIndexes;
+			HittestGrid->GetCollapsedWidgets(WidgetIndexes, StrideCellPoint.X, StrideCellPoint.Y);
+
+			for (int32 i = WidgetIndexes.Num() - 1; i >= 0; --i)
+			{
+				const FHittestGrid::FWidgetData& TestCandidate = WidgetIndexes[i].GetWidgetData();
+				const TSharedPtr<SWidget> TestWidget = TestCandidate.GetWidget();
+				if (!TestWidget.IsValid())
+				{
+					continue;
+				}
+
+				if (!HittestGrid->IsCompatibleUserIndex(UserIndex, TestCandidate.UserIndex))
+				{
+#if WITH_SLATE_DEBUGGING
+					AddDebugIntermediateResult(TestWidget, FDebuggingText::NotCompatibleWithUserIndex);
+#endif
+					continue;
+				}
+
+				const FSlateRect TestCandidateRect = HittestGrid->GetWidgetRenderBoundingRect(TestWidget);
+				if (!(CompareFunc(DestSideFunc(TestCandidateRect), CurrentSourceSide) && FSlateRect::DoRectanglesIntersect(SweptRect, TestCandidateRect)))
+				{
+#if WITH_SLATE_DEBUGGING
+					AddDebugIntermediateResult(TestWidget, FDebuggingText::DoesNotIntersect);
+#endif
+					continue;
+				}
+
+				// If this found widget isn't closer then the previously found widget then keep looking.
+				if (BestWidget.IsValid() && !CompareFunc(DestSideFunc(BestWidgetRect), DestSideFunc(TestCandidateRect)))
+				{
+#if WITH_SLATE_DEBUGGING
+					AddDebugIntermediateResult(TestWidget, FDebuggingText::PreviousWidgetIsBetter);
+#endif
+					continue;
+				}
+
+				// If we have a non escape boundary condition and this widget isn't a descendant of our boundary condition widget then it's invalid so we keep looking.
+				if (NavigationReply.GetBoundaryRule() != EUINavigationRule::Escape
+					&& NavigationReply.GetHandler().IsValid()
+					&& !HittestGrid->IsDescendantOf(NavigationReply.GetHandler().Get(), TestCandidate))
+				{
+#if WITH_SLATE_DEBUGGING
+					AddDebugIntermediateResult(TestWidget, FDebuggingText::NotADescendant);
+#endif
+					continue;
+				}
+
+				if (!TestWidget->IsEnabled())
+				{
+#if WITH_SLATE_DEBUGGING
+					AddDebugIntermediateResult(TestWidget, FDebuggingText::Disabled);
+#endif
+					continue;
+				}
+				
+				if (!TestWidget->SupportsKeyboardFocus())
+				{
+#if WITH_SLATE_DEBUGGING
+					AddDebugIntermediateResult(TestWidget, FDebuggingText::DoesNotSuportKeyboardFocus);
+#endif
+					continue;
+				}
+
+				if (DisabledDestinations.Contains(TestWidget))
+				{
+#if WITH_SLATE_DEBUGGING
+					AddDebugIntermediateResult(TestWidget, FDebuggingText::ParentDisabled);
+#endif
+					continue;
+				}
+
+				BestWidgetRect = TestCandidateRect;
+				BestWidget = TestWidget;
+
+#if WITH_SLATE_DEBUGGING
+				AddDebugIntermediateResult(TestWidget, FDebuggingText::Valid);
+#endif
+			}
+		}
+
+		if (BestWidget.IsValid())
+		{
+			// Check for the need to apply our rule
+			if (CompareFunc(DestSideFunc(BestWidgetRect), SourceSideFunc(SweptRect)))
+			{
+				switch (NavigationReply.GetBoundaryRule())
+				{
+				case EUINavigationRule::Explicit:
+					return NavigationReply.GetFocusRecipient();
+				case EUINavigationRule::Custom:
+				case EUINavigationRule::CustomBoundary:
+				{
+					const FNavigationDelegate& FocusDelegate = NavigationReply.GetFocusDelegate();
+					if (FocusDelegate.IsBound())
+					{
+						return FocusDelegate.Execute(Direction);
+					}
+					return TSharedPtr<SWidget>();
+				}
+				case EUINavigationRule::Stop:
+					return TSharedPtr<SWidget>();
+				case EUINavigationRule::Wrap:
+					CurrentSourceSide = DestSideFunc(SweptRect);
+					FVector2f SampleSpot = WidgetRect.GetCenter();
+					SampleSpot[AxisIndex] = CurrentSourceSide;
+					CurrentCellPoint = HittestGrid->GetCellCoordinate(SampleSpot);
+					bWrapped = true;
+					break;
+				}
+			}
+
+			// Make sure all parents of the chosen widget are enabled before returning.
+			// Note that IsParentsEnabled is a costly function. We call it here as the last step to minimize the number of calls to it.
+			if (!IsParentsEnabled(BestWidget.Get()))
+			{	
+				// Find the next best widget because this one has disabled parents.
+				DisabledDestinations.Add(BestWidget);
+				return FindFocusableWidget(WidgetRect, SweptRect, AxisIndex, Increment, Direction, NavigationReply, CompareFunc, SourceSideFunc, DestSideFunc, UserIndex);
+			}
+			else
+			{
+				return BestWidget;
+			}
+		}
+
+		// break if we have looped back to where we started.
+		if (bWrapped && StartingIndex == CurrentCellProcessed) { break; }
+
+		// If were going to fail our bounds check and our rule is to a boundary condition (Wrap or CustomBoundary) handle appropriately
+		if (!(CurrentCellPoint[AxisIndex] >= 0 && CurrentCellPoint[AxisIndex] < NumCells[AxisIndex]))
+		{
+			if (NavigationReply.GetBoundaryRule() == EUINavigationRule::Wrap)
+			{
+				if (bWrapped)
+				{
+					// If we've already wrapped, unfortunately it must be that the starting widget wasn't within the boundary
+					break;
+				}
+				CurrentSourceSide = DestSideFunc(SweptRect);
+				FVector2f SampleSpot = WidgetRect.GetCenter();
+				SampleSpot[AxisIndex] = CurrentSourceSide;
+				CurrentCellPoint = HittestGrid->GetCellCoordinate(SampleSpot);
+				bWrapped = true;
+			}
+			else if (NavigationReply.GetBoundaryRule() == EUINavigationRule::CustomBoundary)
+			{
+				const FNavigationDelegate& FocusDelegate = NavigationReply.GetFocusDelegate();
+				if (FocusDelegate.IsBound())
+				{
+					return FocusDelegate.Execute(Direction);
+				}
+			}
+		}
+	}
+
+	return TSharedPtr<SWidget>();
+}
+
+TSharedPtr<SWidget> FNavigationMethodOrthogonal::FindNextFocusableWidget(const FArrangedWidget& StartingWidget, const EUINavigation Direction, const FNavigationReply& NavigationReply, const FArrangedWidget& RuleWidget, int32 UserIndex)
+{
+	FSlateRect WidgetRect = HittestGrid->GetWidgetRenderBoundingRect(StartingWidget.Widget);
+	FSlateRect BoundingRuleRect = HittestGrid->GetWidgetRenderBoundingRect(RuleWidget.Widget);
+	FSlateRect SweptWidgetRect = WidgetRect;
+
+	switch (Direction)
+	{
+	case EUINavigation::Left:
+		SweptWidgetRect.Left = BoundingRuleRect.Left;
+		SweptWidgetRect.Right = BoundingRuleRect.Right;
+		SweptWidgetRect.Top += 0.5f;
+		SweptWidgetRect.Bottom -= 0.5f;
+		return FindFocusableWidget(WidgetRect, SweptWidgetRect, 0, -1, Direction, NavigationReply,
+			[](float A, float B) { return A - 0.1f < B; }, // Compare function
+			[](FSlateRect SourceRect) { return SourceRect.Left; }, // Source side function
+			[](FSlateRect DestRect) { return DestRect.Right; }, // Dest side function
+			UserIndex);
+		break;
+	case EUINavigation::Right:
+		SweptWidgetRect.Left = BoundingRuleRect.Left;
+		SweptWidgetRect.Right = BoundingRuleRect.Right;
+		SweptWidgetRect.Top += 0.5f;
+		SweptWidgetRect.Bottom -= 0.5f;
+		return FindFocusableWidget(WidgetRect, SweptWidgetRect, 0, 1, Direction, NavigationReply,
+			[](float A, float B) { return A + 0.1f > B; }, // Compare function
+			[](FSlateRect SourceRect) { return SourceRect.Right; }, // Source side function
+			[](FSlateRect DestRect) { return DestRect.Left; }, // Dest side function
+			UserIndex);
+		break;
+	case EUINavigation::Up:
+		SweptWidgetRect.Top = BoundingRuleRect.Top;
+		SweptWidgetRect.Bottom = BoundingRuleRect.Bottom;
+		SweptWidgetRect.Left += 0.5f;
+		SweptWidgetRect.Right -= 0.5f;
+		return FindFocusableWidget(WidgetRect, SweptWidgetRect, 1, -1, Direction, NavigationReply,
+			[](float A, float B) { return A - 0.1f < B; }, // Compare function
+			[](FSlateRect SourceRect) { return SourceRect.Top; }, // Source side function
+			[](FSlateRect DestRect) { return DestRect.Bottom; }, // Dest side function
+			UserIndex);
+		break;
+	case EUINavigation::Down:
+		SweptWidgetRect.Top = BoundingRuleRect.Top;
+		SweptWidgetRect.Bottom = BoundingRuleRect.Bottom;
+		SweptWidgetRect.Left += 0.5f;
+		SweptWidgetRect.Right -= 0.5f;
+		return FindFocusableWidget(WidgetRect, SweptWidgetRect, 1, 1, Direction, NavigationReply,
+			[](float A, float B) { return A + 0.1f > B; }, // Compare function
+			[](FSlateRect SourceRect) { return SourceRect.Bottom; }, // Source side function
+			[](FSlateRect DestRect) { return DestRect.Top; }, // Dest side function
+			UserIndex);
+		break;
+
+	default:
+		break;
+	}
+
+	return TSharedPtr<SWidget>(nullptr);
+}
+
+#undef LOCTEXT_NAMESPACE
\ No newline at end of file
diff --git a/Engine/Source/Runtime/SlateCore/Private/Input/NavigationMethodOrthogonal.h b/Engine/Source/Runtime/SlateCore/Private/Input/NavigationMethodOrthogonal.h
new file mode 100644
--- /dev/null
+++ b/Engine/Source/Runtime/SlateCore/Private/Input/NavigationMethodOrthogonal.h
@@ -0,0 +1,21 @@
+// Copyright Epic Games, Inc. All Rights Reserved.
+
+#pragma once 
+
+#include "Input/NavigationMethod.h"
+
+#include "NavigationMethodOrthogonal.generated.h"
+
+USTRUCT(DisplayName="Orthogonal")
+struct FNavigationMethodOrthogonal : public FNavigationMethod
+{
+	GENERATED_BODY()
+
+public:
+	virtual TSharedPtr<SWidget> FindNextFocusableWidget(const FArrangedWidget& StartingWidget, const EUINavigation Direction, const FNavigationReply& NavigationReply, const FArrangedWidget& RuleWidget, int32 InUserIndex);
+	
+private:
+	/** Default Navigation Behavior. Sweeps hittestgrid cells orthogonal to direction of movement. Finds widget whoe edge is nearest starting widget opposing edge **/
+	template<typename TCompareFunc, typename TSourceSideFunc, typename TDestSideFunc>
+	TSharedPtr<SWidget> FindFocusableWidget(const FSlateRect WidgetRect, const FSlateRect SweptRect, int32 AxisIndex, int32 Increment, const EUINavigation Direction, const FNavigationReply& NavigationReply, TCompareFunc CompareFunc, TSourceSideFunc SourceSideFunc, TDestSideFunc DestSideFunc, int32 UserIndex);
+};
diff --git a/Engine/Source/Runtime/SlateCore/Private/Input/NavigationMethodProximity.cpp b/Engine/Source/Runtime/SlateCore/Private/Input/NavigationMethodProximity.cpp
new file mode 100644
--- /dev/null
+++ b/Engine/Source/Runtime/SlateCore/Private/Input/NavigationMethodProximity.cpp
@@ -0,0 +1,595 @@
+// Copyright Epic Games, Inc. All Rights Reserved.
+
+#include "Input/NavigationMethodProximity.h"
+#include "Styling/CoreStyle.h"
+
+static bool GNavigationSearchFromCorners = false;
+static FAutoConsoleVariableRef CVarNavigationSearchFromCorners(
+	TEXT("Slate.Navigation.SearchFromCorners"),
+	GNavigationSearchFromCorners,
+	TEXT("If true, will scan for candidate widget from the corners the source widget. This will increase with width of widgets that can be navigated to"));
+
+static bool GNavigationIgnoreCustomRules = false;
+static FAutoConsoleVariableRef CVarNavigationIgnoreCustomRules(
+	TEXT("Slate.Navigation.IgnoreCustomRules"),
+	GNavigationIgnoreCustomRules,
+	TEXT("Will disable custom navigation rules. Used for testing purposes."));
+
+#define LOCTEXT_NAMESPACE "NavigationMethodProximity"
+
+namespace UE::Slate::Private
+{
+	float MinkowskiDistance(const FVector2f& A, const FVector2f& B, float P)
+	{
+		if (P <= 0.0f) P = FLT_MIN;
+
+		float sum = 0.0f;
+		sum += FMath::Pow(FMath::Abs(A.X - B.X), P);
+		sum += FMath::Pow(FMath::Abs(A.Y - B.Y), P);
+		return FMath::Pow(sum, 1.0f / P);
+	}
+}
+
+void FNavigationMethodProximity::Initialize(const FHittestGrid* InHittestGrid, TArray<FDebugWidgetResult>* InIntermediateResultsPtr)
+{
+	FNavigationMethod::Initialize(InHittestGrid, InIntermediateResultsPtr);
+	VisitedWidgets.Empty();
+
+#if WITH_SLATE_DEBUGGING
+	DrawDebugData = {};
+#endif //WITH_SLATE_DEBUGGING
+}
+
+TSharedPtr<SWidget> FNavigationMethodProximity::FindNextFocusableWidget(const FArrangedWidget& StartingWidget, const EUINavigation Direction, const FNavigationReply& NavigationReply, const FArrangedWidget& RuleWidget, int32 UserIndex)
+{
+	// Prevent self navigation
+	DisabledDestinations.Add(StartingWidget.Widget);
+
+	const FSlateRect WidgetRect = HittestGrid->GetWidgetRenderBoundingRect(StartingWidget.Widget);
+	const FSlateRect BoundaryRect = HittestGrid->GetWidgetRenderBoundingRect(RuleWidget.Widget);
+
+#if WITH_SLATE_DEBUGGING
+	PrepareDrawDebugData(WidgetRect, Direction);
+#endif //WITH_SLATE_DEBUGGING
+
+	return FindFocusableWidgetFromRect(WidgetRect, BoundaryRect, Direction, NavigationReply, UserIndex);
+}
+
+TSharedPtr<SWidget> FNavigationMethodProximity::FindFocusableWidgetFromRect(const FSlateRect WidgetRect, const FSlateRect BoundaryRect, const EUINavigation Direction, const FNavigationReply& NavigationReply, int32 UserIndex)
+{
+	// Get the cell scan direction. Will scan in the direction of navigation
+	const float SearchDirectionRadians = GetSearchDirectionRadians(Direction);
+	FIntPoint CellDelta(FMath::RoundToInt32(FMath::Cos(SearchDirectionRadians)), FMath::RoundToInt32(FMath::Sin(SearchDirectionRadians)));
+
+	TArray<FCellSearchInfo> CellsToSearch;
+	PopulateCellsToSearch(WidgetRect, Direction, false, CellsToSearch);
+
+	TSharedPtr<SWidget> BestWidget;
+	float BestWidgetDistance = FLT_MAX;
+
+	TSharedPtr<SWidget> BestBoundaryWidget;
+	float BestBoundaryWidgetDistance = FLT_MAX;
+
+	while (CellsToSearch.Num())
+	{
+		// Scan in reverse to support removal during iteration
+		for (int32 CellsToSearchIndex = CellsToSearch.Num() - 1; CellsToSearchIndex >= 0; --CellsToSearchIndex )
+		{
+			FCellSearchInfo& SearchInfo = CellsToSearch[CellsToSearchIndex];
+			FCellSearchResult CellResult = FindBestWidgetInCell(SearchInfo, BoundaryRect, Direction, NavigationReply, UserIndex);	
+			SearchInfo.CurrentCellPoint += CellDelta;
+
+			if (CellResult.Widget.IsValid())
+			{
+				// Make sure all parents of the chosen widget are enabled before returning.
+				// Note that IsParentsEnabled is a costly function. We call it here as the last step to minimize the number of calls to it.
+				if (!IsParentsEnabled(CellResult.Widget.Get()))
+				{
+					// Find the next best widget because this one has disabled parents.
+					DisabledDestinations.Add(CellResult.Widget);
+					return FindFocusableWidgetFromRect(WidgetRect, BoundaryRect, Direction, NavigationReply, UserIndex);
+				}
+			}
+	
+			switch(CellResult.Action)
+			{
+			case ECellSearchResultAction::EndSearch:
+			{
+				CellsToSearch.RemoveAt(CellsToSearchIndex);
+				break;
+			}
+			case ECellSearchResultAction::ApplyBoundaryRule:
+			{
+				CellsToSearch.RemoveAt(CellsToSearchIndex);
+				HandleBoundaryRule(WidgetRect, BoundaryRect, Direction, NavigationReply, CellResult, SearchInfo, CellsToSearch);
+				break;
+			}
+			case ECellSearchResultAction::ContinueSearch:
+				break;		
+			}
+
+			if (CellResult.Widget.IsValid())
+			{
+				// if the widget was set via a boundary rule, or discovered after a wrapped search. 
+				if (!SearchInfo.bWrapped && CellResult.Action != ECellSearchResultAction::ApplyBoundaryRule)
+				{
+					if (CellResult.Distance < BestWidgetDistance)
+					{
+						BestWidget = CellResult.Widget;
+						BestWidgetDistance = CellResult.Distance;
+					}
+				}
+				else
+				{
+					if (CellResult.Distance < BestBoundaryWidgetDistance)
+					{
+						BestBoundaryWidget = CellResult.Widget;
+						BestBoundaryWidgetDistance = CellResult.Distance;
+					}
+				}
+			}
+		}
+	}
+
+	// Standard navigation will take precedence over boundary rules. 
+	// This is to avoid wrapped widgets very close to the boundary edge from take priority over further focusablke widgets within the bounds
+	if (BestWidget == nullptr)
+	{
+		return BestBoundaryWidget;
+	}
+
+	return BestWidget;
+}
+
+FNavigationMethodProximity::FCellSearchResult FNavigationMethodProximity::FindBestWidgetInCell(const FCellSearchInfo& SearchInfo, const FSlateRect BoundaryRect, const EUINavigation Direction, const FNavigationReply& NavigationReply, int32 UserIndex)
+{
+	FCellSearchResult CellResult;
+	CellResult.Action = ECellSearchResultAction::ContinueSearch;
+
+	// If cell to check is beyond the grid, apply boundary rule
+	if (!HittestGrid->IsValidCellCoord(SearchInfo.CurrentCellPoint))
+	{
+		CellResult.Action = ECellSearchResultAction::ApplyBoundaryRule;
+		return CellResult;
+	}
+
+	auto WidgetFunc = [&](TSharedPtr<SWidget> Candidate) 
+	{
+		if (VisitedWidgets.Contains(Candidate.Get()))
+		{
+			return false;
+		}
+		VisitedWidgets.Add(Candidate.Get());
+
+		const FSlateRect CandidateRect = HittestGrid->GetWidgetRenderBoundingRect(Candidate);
+				
+		if (!IsOutsideNavigationEdge(Direction, SearchInfo.CurrentSearchRect, CandidateRect))
+		{
+#if WITH_SLATE_DEBUGGING
+			AddDebugIntermediateResult(Candidate, FDebuggingText::DoesNotIntersect);
+#endif
+			return false;
+		}
+
+		float Distance = CalculateDistance(Direction, SearchInfo.CurrentSearchRect, CandidateRect);
+
+#if WITH_SLATE_DEBUGGING
+		DebugCandidateDistance(Direction, SearchInfo.CurrentSearchRect, CandidateRect, Distance);
+#endif
+
+		// If this found widget isn't closer then the previously found widget then keep looking.
+		if (CellResult.Widget.IsValid() && Distance >= CellResult.Distance)
+		{
+#if WITH_SLATE_DEBUGGING
+			AddDebugIntermediateResult(Candidate, FDebuggingText::PreviousWidgetIsBetter);
+#endif
+			return false;
+		}
+
+		CellResult.Widget = Candidate;
+		CellResult.WidgetRect = CandidateRect;
+		CellResult.Distance = Distance;
+		CellResult.SearchRect = SearchInfo.CurrentSearchRect;
+		CellResult.Action = ECellSearchResultAction::EndSearch;
+		return true;
+	};
+
+	const FIntPoint CurrentCellPoint = SearchInfo.CurrentCellPoint;
+	ForEachFocusableWidgetsInCell(CurrentCellPoint.X, CurrentCellPoint.Y, NavigationReply, UserIndex, WidgetFunc);
+
+	// If best widget is outside the bounds, apply boundary rule
+	if (CellResult.Widget.IsValid() && !FSlateRect::DoRectanglesIntersect(BoundaryRect, CellResult.WidgetRect))
+	{	
+		CellResult.Action = ECellSearchResultAction::ApplyBoundaryRule;
+	}
+
+	return CellResult;
+}
+
+void FNavigationMethodProximity::HandleBoundaryRule(const FSlateRect SourceRect, const FSlateRect BoundaryRect, const EUINavigation Direction, const FNavigationReply& NavigationReply, FCellSearchResult& OutCellResult, FCellSearchInfo& OutSearchInfo, TArray<FCellSearchInfo>& OutCellsToSearch)
+{
+	switch (NavigationReply.GetBoundaryRule())
+	{
+	case EUINavigationRule::Explicit:
+	{
+		if (GNavigationIgnoreCustomRules)
+		{
+			break;
+		}
+
+		OutCellResult.Widget = NavigationReply.GetFocusRecipient();
+		break;
+	}
+	case EUINavigationRule::Custom:
+	case EUINavigationRule::CustomBoundary:
+	{
+		if (GNavigationIgnoreCustomRules)
+		{
+			break;
+		}
+
+		const FNavigationDelegate& FocusDelegate = NavigationReply.GetFocusDelegate();
+		if (FocusDelegate.IsBound())
+		{
+			OutCellResult.Widget = FocusDelegate.Execute(Direction);
+		}
+		break;
+	}
+	case EUINavigationRule::Wrap:
+	{
+		if (!OutSearchInfo.bWrapped)
+		{
+			// Rescan search space by adding CellsToSearch from this new point. Add as wrapped so they do not take precedence of non-wrapped
+			const FSlateRect WrappedRect = WrapRectWithinBoundary(Direction, BoundaryRect, SourceRect);
+			PopulateCellsToSearch(WrappedRect, Direction, true, OutCellsToSearch);
+
+			OutSearchInfo.bWrapped = true;
+			
+			// Invalidate widget, will wrap to next
+			OutCellResult.Widget = nullptr;
+		}
+		break;
+	}
+	case EUINavigationRule::Stop:
+	{
+		OutCellResult.Widget = nullptr;
+		break;
+	}
+	case EUINavigationRule::Escape:
+	{
+		break;
+	}
+	default:
+	{
+		break;
+	}
+	}
+
+	if (OutCellResult.Widget.IsValid())
+	{
+		OutCellResult.WidgetRect = HittestGrid->GetWidgetRenderBoundingRect(OutCellResult.Widget);
+		OutCellResult.Distance = CalculateDistance(Direction, OutCellResult.SearchRect, OutCellResult.WidgetRect);
+	}
+	else
+	{
+		OutCellResult.Distance = FLT_MAX;
+	}
+}
+
+void FNavigationMethodProximity::CalculateSamplePoints(EUINavigation Direction, const FSlateRect& SourceRect, const FSlateRect& DestinationRect, FVector2f& OutSourcePoint, FVector2f& OutDestinationPoint) const
+{
+	// NOTE: Current implementation is naive, and could have improved accuracy. 
+	//		Using nearest points on rectangles could get a better metric for movement distance. 
+	//      However, this may cause neighboring perpindicular widgets to have a greater influence than desired. 
+	//		As a workaround, perhaps take into account the distance of the point on the target to the distance of the cone (or direction of travel)
+	const FVector2f SearchDirection = GetSearchDirection(Direction);
+
+	OutSourcePoint = SourceRect.GetCenter() + SourceRect.GetSize() * 0.5f * SearchDirection;
+	OutDestinationPoint = DestinationRect.GetCenter() + DestinationRect.GetSize() * -0.5f * SearchDirection; // The target rect offset is the inverse of the source
+}
+
+float FNavigationMethodProximity::CalculateDistance(EUINavigation Direction, const FSlateRect& SourceRect, const FSlateRect& DestinationRect) const
+{
+	FVector2f SourcePoint;
+	FVector2f DestinationPoint; 
+	CalculateSamplePoints(Direction, SourceRect, DestinationRect, SourcePoint, DestinationPoint);
+	return UE::Slate::Private::MinkowskiDistance(SourcePoint, DestinationPoint, AlignmentFactor);
+}
+
+void FNavigationMethodProximity::PopulateCellsToSearch(const FSlateRect& SourceRect, EUINavigation Direction, bool bWrapped, TArray<FCellSearchInfo>& OutCellsToSearch)
+{
+	// Find the starting cells by raymarching two mirror rays at an angle from the search point
+	const float SearchDirectionRadians = GetSearchDirectionRadians(Direction);
+
+	FVector2f SearchPointMin, SearchPointMax;
+	GetSearchPoints(Direction, SourceRect, SearchPointMin, SearchPointMax);
+
+	FVector2f CellPointMin = GetCellCoordinate(SearchPointMin);
+	FVector2f CellPointMax = GetCellCoordinate(SearchPointMax);
+
+	// Fill cells between edges
+	FVector2f EdgeDelta = (CellPointMax - CellPointMin);
+	EdgeDelta.Normalize();
+	for (float Y = CellPointMin.Y; Y < CellPointMin.Y + EdgeDelta.Y; ++Y)
+	{
+		for (float X = CellPointMin.X; X < CellPointMin.X + EdgeDelta.X; ++X)
+		{
+			const FVector2f EdgePoint(X, Y);
+			FCellSearchInfo SearchInfo(this, SourceRect, EdgePoint, bWrapped);
+
+			OutCellsToSearch.Add(SearchInfo);
+		}
+	}
+
+	// Apply the search angle to the angle of the navigation
+	const float SearchAngleRadians = FMath::DegreesToRadians(SearchAngleDegrees);
+	const FVector2f DeltaMin = FVector2f(FMath::Cos(SearchDirectionRadians + SearchAngleRadians), FMath::Sin(SearchDirectionRadians + SearchAngleRadians));
+	const FVector2f DeltaMax = FVector2f(FMath::Cos(SearchDirectionRadians - SearchAngleRadians), FMath::Sin(SearchDirectionRadians - SearchAngleRadians));
+
+	bool CheckCellPointMin = true;
+	bool CheckCellPointMax = true;
+	while (CheckCellPointMin && CheckCellPointMax)
+	{
+		CheckCellPointMin = IsValidCellCoordinate(CellPointMin.X, CellPointMin.Y);
+		CheckCellPointMax = IsValidCellCoordinate(CellPointMax.X, CellPointMax.Y);
+
+		if (CheckCellPointMin)
+		{
+			FCellSearchInfo SearchInfo(this, SourceRect, CellPointMin, bWrapped);
+			OutCellsToSearch.Add(SearchInfo);
+
+			CellPointMin += DeltaMin;
+		}
+		if (CheckCellPointMax)
+		{
+			FCellSearchInfo SearchInfo(this, SourceRect, CellPointMax, bWrapped);
+
+			OutCellsToSearch.Add(SearchInfo);
+			CellPointMax += DeltaMax;
+		}
+	}
+
+#if WITH_SLATE_DEBUGGING
+	for (const FCellSearchInfo& Cell : OutCellsToSearch)
+	{
+		DebugSearchCell(Cell);
+	}
+#endif //WITH_SLATE_DEBUGGING
+}
+
+FVector2f FNavigationMethodProximity::GetSearchDirection(const EUINavigation Direction) const
+{	
+	switch (Direction)
+	{
+	case EUINavigation::Left:
+		return FVector2f(-1.0f, 0.0f);
+	case EUINavigation::Right:
+		return FVector2f(1.0f, 0.0f);
+	case EUINavigation::Up:
+		return FVector2f(0.0f, -1.0f);
+	case EUINavigation::Down:
+		return FVector2f(0.0f, 1.0f);
+	default:
+		check(false);
+	}
+	return FVector2f();
+}
+
+void FNavigationMethodProximity::GetSearchPoints(const EUINavigation Direction, const FSlateRect& SourceRect, FVector2f& OutCellSearchPointMin, FVector2f& OutCellSearchPointMax) const
+{
+	if (!GNavigationSearchFromCorners)
+	{
+		// Return from the center of the search edge
+		OutCellSearchPointMin = SourceRect.GetCenter() + SourceRect.GetSize() * 0.5f * GetSearchDirection(Direction);
+		OutCellSearchPointMax = OutCellSearchPointMin;
+		return;
+	}
+
+	switch (Direction)
+	{
+	case EUINavigation::Left:
+		OutCellSearchPointMin = FVector2f(SourceRect.Left, SourceRect.Top);
+		OutCellSearchPointMax = FVector2f(SourceRect.Left, SourceRect.Bottom);
+		break;
+	case EUINavigation::Right:
+		OutCellSearchPointMin = FVector2f(SourceRect.Right, SourceRect.Top);
+		OutCellSearchPointMax = FVector2f(SourceRect.Right, SourceRect.Bottom);
+		break;
+	case EUINavigation::Up:
+		OutCellSearchPointMax = FVector2f(SourceRect.Left, SourceRect.Top);
+		OutCellSearchPointMax = FVector2f(SourceRect.Right, SourceRect.Top);
+		break;
+	case EUINavigation::Down:
+		OutCellSearchPointMin = FVector2f(SourceRect.Left, SourceRect.Bottom);
+		OutCellSearchPointMax = FVector2f(SourceRect.Right, SourceRect.Bottom);
+		break;
+	}
+}
+
+float FNavigationMethodProximity::GetSearchDirectionRadians(const EUINavigation Direction) const
+{
+	switch (Direction)
+	{
+	case EUINavigation::Left:
+		return PI;
+	case EUINavigation::Right:
+		return 0.0f;
+	case EUINavigation::Up:
+		return -PI / 2.0f;
+	case EUINavigation::Down:
+		return PI / 2.0f;
+	default:
+		check(false);
+	}
+	return 0;
+}
+
+bool FNavigationMethodProximity::IsOutsideNavigationEdge(const EUINavigation Direction, const FSlateRect& SourceRect, const FSlateRect& DestinationRect) const
+{
+	switch (Direction)
+	{
+	case EUINavigation::Left: 
+		return DestinationRect.Right - OverlapThreshold < SourceRect.Left;
+	case EUINavigation::Right:
+		return DestinationRect.Left + OverlapThreshold > SourceRect.Right;
+	case EUINavigation::Up:
+		return DestinationRect.Bottom - OverlapThreshold < SourceRect.Top;
+	case EUINavigation::Down:
+		return DestinationRect.Top + OverlapThreshold > SourceRect.Bottom;
+	default:
+		check(false);
+	}
+	return 0;
+}
+
+FSlateRect FNavigationMethodProximity::WrapRectWithinBoundary(const EUINavigation Direction, const FSlateRect& BoundaryRect, const FSlateRect& SourceRect) const
+{
+	switch (Direction)
+	{
+	case EUINavigation::Left:
+	{
+		FSlateRect OutRect = SourceRect;
+		float Width = OutRect.GetSize().X;
+		OutRect.Left = BoundaryRect.Right;
+		OutRect.Right = OutRect.Left + Width;
+		return OutRect;
+	}
+	case EUINavigation::Right:
+	{
+		FSlateRect OutRect = SourceRect;
+		float Width = OutRect.GetSize().X;
+		OutRect.Right = BoundaryRect.Left;
+		OutRect.Left = OutRect.Right - Width;
+		return OutRect;
+	}
+	case EUINavigation::Up:
+	{
+		FSlateRect OutRect = SourceRect;
+		float Height = OutRect.GetSize().Y;
+		OutRect.Top = BoundaryRect.Bottom;
+		OutRect.Bottom = OutRect.Top + Height;
+		return OutRect;
+	}
+	case EUINavigation::Down:
+	{
+		FSlateRect OutRect = SourceRect;
+		float Height = OutRect.GetSize().Y;
+		OutRect.Bottom = BoundaryRect.Top;
+		OutRect.Top = OutRect.Bottom - Height;
+		return OutRect;
+	}
+	default:
+		check(false);
+	}
+
+	return FSlateRect();
+}
+
+#if WITH_SLATE_DEBUGGING
+
+void FNavigationMethodProximity::PrepareDrawDebugData(const FSlateRect& SourceRect, EUINavigation Direction)
+{
+	// Draw wireframe rect
+	DrawDebugData.Points.Add(SourceRect.GetTopLeft());
+	DrawDebugData.Points.Add(SourceRect.GetTopRight());
+	DrawDebugData.Colors.Add(FLinearColor::Green);
+	DrawDebugData.Colors.Add(FLinearColor::Green);
+
+	DrawDebugData.Points.Add(SourceRect.GetTopRight());
+	DrawDebugData.Points.Add(SourceRect.GetBottomRight());
+	DrawDebugData.Colors.Add(FLinearColor::Green);
+	DrawDebugData.Colors.Add(FLinearColor::Green);
+
+	DrawDebugData.Points.Add(SourceRect.GetBottomRight());
+	DrawDebugData.Points.Add(SourceRect.GetBottomLeft());
+	DrawDebugData.Colors.Add(FLinearColor::Green);
+	DrawDebugData.Colors.Add(FLinearColor::Green);
+
+	DrawDebugData.Points.Add(SourceRect.GetTopLeft());
+	DrawDebugData.Points.Add(SourceRect.GetBottomLeft());
+	DrawDebugData.Colors.Add(FLinearColor::Green);
+	DrawDebugData.Colors.Add(FLinearColor::Green);
+
+	// Draw Search Cone
+	FVector2f SearchPointMin, SearchPointMax;
+	GetSearchPoints(Direction, SourceRect, SearchPointMin, SearchPointMax);
+	
+	const float SearchDirectionRadians = GetSearchDirectionRadians(Direction);
+	const float SearchAngleRadians = FMath::DegreesToRadians(SearchAngleDegrees);
+	const FVector2f DirectionMin = FVector2f(FMath::Cos(SearchDirectionRadians + SearchAngleRadians), FMath::Sin(SearchDirectionRadians + SearchAngleRadians));
+	const FVector2f DirectionMax = FVector2f(FMath::Cos(SearchDirectionRadians - SearchAngleRadians), FMath::Sin(SearchDirectionRadians - SearchAngleRadians));
+
+	const float RayLength = 500;
+	DrawDebugData.Points.Add(SearchPointMin);
+	DrawDebugData.Points.Add(SearchPointMin + DirectionMin * RayLength);
+	DrawDebugData.Colors.Add(FLinearColor::Blue);
+	DrawDebugData.Colors.Add(FLinearColor::Blue);
+	
+	DrawDebugData.Points.Add(SearchPointMax);
+	DrawDebugData.Points.Add(SearchPointMax + DirectionMax * RayLength);
+	DrawDebugData.Colors.Add(FLinearColor::Blue);
+	DrawDebugData.Colors.Add(FLinearColor::Blue);
+	
+}
+
+void FNavigationMethodProximity::DebugCandidateDistance(EUINavigation Direction, const FSlateRect& SourceRect, const FSlateRect& DestinationRect, float Distance)
+{
+	FVector2f SourcePoint, TargetPoint;
+	CalculateSamplePoints(Direction, SourceRect, DestinationRect, SourcePoint, TargetPoint);
+
+	DrawDebugData.Points.Add(SourcePoint);
+	DrawDebugData.Points.Add(TargetPoint);
+	DrawDebugData.LabelPositions.Add((SourcePoint + TargetPoint) * 0.5f);
+	DrawDebugData.Colors.Add(FLinearColor::Green);
+	DrawDebugData.Colors.Add(FLinearColor::Green);
+	DrawDebugData.Distances.Add(Distance);
+}
+
+void FNavigationMethodProximity::DebugSearchCell(const FCellSearchInfo& Cell)
+{
+	FVector2f CellWorldSpace = FVector2f(static_cast<float>(Cell.CurrentCellPoint.X), static_cast<float>(Cell.CurrentCellPoint.Y)) * HittestGrid->GetCellSize();
+	FSlateRect Rect(CellWorldSpace, CellWorldSpace + HittestGrid->GetCellSize());
+
+	DrawDebugData.Points.Add(Rect.GetTopLeft());
+	DrawDebugData.Points.Add(Rect.GetTopRight());
+	DrawDebugData.Colors.Add(FLinearColor::Red);
+	DrawDebugData.Colors.Add(FLinearColor::Red);
+
+	DrawDebugData.Points.Add(Rect.GetTopRight());
+	DrawDebugData.Points.Add(Rect.GetBottomRight());
+	DrawDebugData.Colors.Add(FLinearColor::Red);
+	DrawDebugData.Colors.Add(FLinearColor::Red);
+
+	DrawDebugData.Points.Add(Rect.GetBottomRight());
+	DrawDebugData.Points.Add(Rect.GetBottomLeft());
+	DrawDebugData.Colors.Add(FLinearColor::Red);
+	DrawDebugData.Colors.Add(FLinearColor::Red);
+
+	DrawDebugData.Points.Add(Rect.GetTopLeft());
+	DrawDebugData.Points.Add(Rect.GetBottomLeft());
+	DrawDebugData.Colors.Add(FLinearColor::Red);
+	DrawDebugData.Colors.Add(FLinearColor::Red);
+}
+
+void FNavigationMethodProximity::DrawDebug(int32 InLayer, const FGeometry& AllottedGeometry, FSlateWindowElementList& WindowElementList) 
+{
+	ensure(DrawDebugData.Colors.Num() == DrawDebugData.Points.Num());
+	
+	FSlateDrawElement::MakeLines(WindowElementList, InLayer, AllottedGeometry.ToPaintGeometry(), DrawDebugData.Points, DrawDebugData.Colors);
+	
+	const FSlateFontInfo FontInfo = FCoreStyle::GetDefaultFontStyle("Regular", 8);	
+	for (int32 Index = 0; Index < DrawDebugData.Distances.Num(); ++Index)
+	{
+		const float Distance = DrawDebugData.Distances[Index];
+		const FVector2f LabelPosition = DrawDebugData.LabelPositions[Index];
+		const FLinearColor Color = FLinearColor(0, 0.89, 0.05, 1); // Greenish
+		const FText Label = FText::Format(LOCTEXT("NavigationLabel", "{0}"), Distance);
+		const FPaintGeometry LabelGeometry = AllottedGeometry.ToPaintGeometry(FSlateLayoutTransform(LabelPosition));
+		FSlateDrawElement::MakeText(WindowElementList, InLayer, LabelGeometry, Label, FontInfo, ESlateDrawEffect::None, Color);
+	}
+}
+
+#endif //WITH_SLATE_DEBUGGING
+
+
+#undef LOCTEXT_NAMESPACE
\ No newline at end of file
diff --git a/Engine/Source/Runtime/SlateCore/Private/Input/NavigationMethodProximity.h b/Engine/Source/Runtime/SlateCore/Private/Input/NavigationMethodProximity.h
new file mode 100644
--- /dev/null
+++ b/Engine/Source/Runtime/SlateCore/Private/Input/NavigationMethodProximity.h
@@ -0,0 +1,104 @@
+// Copyright Epic Games, Inc. All Rights Reserved.
+
+#pragma once 
+
+#include "Input/HittestGrid.h"
+#include "Input/NavigationMethod.h"
+
+#include "NavigationMethodProximity.generated.h"
+
+USTRUCT(DisplayName="Proximity")
+struct FNavigationMethodProximity : public FNavigationMethod
+{
+	GENERATED_BODY()
+
+public:		
+	TSharedPtr<SWidget> FindNextFocusableWidget(const FArrangedWidget& StartingWidget, const EUINavigation Direction, const FNavigationReply& NavigationReply, const FArrangedWidget& RuleWidget, int32 InUserIndex) override;
+
+	void Initialize(const FHittestGrid* InHittestGrid, TArray<FDebugWidgetResult>* InIntermediateResultsPtr) override;
+
+private:
+	/** Controls the preference for off-axis widgets. Values 0.0 to 2.0. Lower values increases preference for axis aligned widgets **/
+	UPROPERTY(EditAnywhere, Category = Navigation, meta = (ClampMin = "0.0", ClampMax = "2.0"))
+	float AlignmentFactor = 1.0f;
+
+	/**	Controls the view angle in degrees from source widget to candidate widget search region. Values 0 to 90 **/
+	UPROPERTY(EditAnywhere, Category = Navigation, meta = (ClampMin = "0.0", ClampMax = "90.0", Units = "Degrees"))
+	float SearchAngleDegrees = 45.0f;
+
+	/**	In slate units, controls the allowed the amount of widget edge overlap widget when considering candidate widget **/
+	UPROPERTY(EditAnywhere, Category = Navigation, meta = (ClampMin = "0.0"))
+	float OverlapThreshold = 20.0f;
+
+private:
+
+	struct FCellSearchInfo
+	{
+		FCellSearchInfo(const FNavigationMethodProximity* NavigationMethod, const FSlateRect& CurrentSearchRect, const FVector2f& SamplePoint, bool bWrapped)
+			: NavigationMethod(NavigationMethod)
+			, CurrentSearchRect(CurrentSearchRect)
+			, CurrentCellPoint(FMath::RoundToInt(SamplePoint.X), FMath::RoundToInt(SamplePoint.Y))
+			, StartingCellPoint(FMath::RoundToInt(SamplePoint.X), FMath::RoundToInt(SamplePoint.Y))
+			, bWrapped(bWrapped)
+		{}
+
+		const FNavigationMethodProximity* NavigationMethod;
+		FSlateRect CurrentSearchRect;
+		FIntPoint CurrentCellPoint;
+		FIntPoint StartingCellPoint;
+		bool bWrapped : 1 ;
+	};
+
+	enum class ECellSearchResultAction
+	{
+		ContinueSearch,
+		EndSearch,
+		ApplyBoundaryRule
+	};
+
+	struct FCellSearchResult
+	{
+		TSharedPtr<SWidget> Widget;
+		FSlateRect WidgetRect;
+		FSlateRect SearchRect;
+		float Distance = FLT_MAX;
+		ECellSearchResultAction Action;
+	};
+
+	TSharedPtr<SWidget> FindFocusableWidgetFromRect(const FSlateRect WidgetRect, const FSlateRect BoundaryRect, const EUINavigation Direction, const FNavigationReply& NavigationReply, int32 UserIndex);
+	FCellSearchResult FindBestWidgetInCell(const FCellSearchInfo& SearchInfo, const FSlateRect BoundaryRect, const EUINavigation Direction, const FNavigationReply& NavigationReply, int32 UserIndex);	
+	void HandleBoundaryRule(const FSlateRect SourceRect, const FSlateRect BoundaryRect, const EUINavigation Direction, const FNavigationReply& NavigationReply, FCellSearchResult& OutCellResult, FCellSearchInfo& SearchInfo, TArray<FCellSearchInfo>& OutCellsToSearch);
+	void PopulateCellsToSearch(const FSlateRect& SourceRect, EUINavigation Direction, bool bWrapped, TArray<FCellSearchInfo>& OutCellsToSearch);
+	
+	void CalculateSamplePoints(EUINavigation Direction, const FSlateRect& SourceRect, const FSlateRect& DestinationRect, FVector2f& OutSourcePoint, FVector2f& OutDestinationPoint) const;
+	float CalculateDistance(EUINavigation Direction, const FSlateRect& SourceRect, const FSlateRect& DestinationRect) const;
+
+	FVector2f GetSearchDirection(const EUINavigation Direction) const;
+	void GetSearchPoints(const EUINavigation Direction, const FSlateRect& SourceRect, FVector2f& OutCellSearchPointMin, FVector2f& OutCellSearchPointMax) const;
+	float GetSearchDirectionRadians(EUINavigation Direction) const;
+
+	bool IsOutsideNavigationEdge(const EUINavigation Direction, const FSlateRect& SourceRect, const FSlateRect& DestinationRect) const;
+	FSlateRect WrapRectWithinBoundary(const EUINavigation Direction, const FSlateRect& BoundaryRect, const FSlateRect& SourceRect) const;
+
+	TSet<SWidget*> VisitedWidgets;
+
+private:
+#if WITH_SLATE_DEBUGGING
+
+	struct FDrawDebugData
+	{
+		TArray<FVector2f> LabelPositions;
+		TArray<float> Distances;
+		TArray<FVector2f> Points;
+		TArray<FLinearColor> Colors;
+	};
+
+	FDrawDebugData DrawDebugData;
+
+	void PrepareDrawDebugData(const FSlateRect& SourceRect, EUINavigation Direction);
+	void DebugCandidateDistance(EUINavigation Direction, const FSlateRect& SourceRect, const FSlateRect& DestinationRect, float Distance);
+	void DebugSearchCell(const FCellSearchInfo& Cell);
+	void DrawDebug(int32 InLayer, const FGeometry& AllottedGeometry, FSlateWindowElementList& WindowElementList);
+
+#endif //WITH_SLATE_DEBUGGING
+};
\ No newline at end of file
diff --git a/Engine/Source/Runtime/SlateCore/Private/Input/NavigationRouting.cpp b/Engine/Source/Runtime/SlateCore/Private/Input/NavigationRouting.cpp
new file mode 100644
--- /dev/null
+++ b/Engine/Source/Runtime/SlateCore/Private/Input/NavigationRouting.cpp
@@ -0,0 +1,197 @@
+// Copyright Epic Games, Inc. All Rights Reserved.
+
+#include "Input/NavigationRouting.h"
+
+#include "Input/HittestGrid.h"
+#include "Types/NavigationMetaData.h"
+#include "Math/UnrealMathUtility.h"
+
+
+TSharedPtr<SWidget> FNavigationRouting::RouteNavigationFocus(const FNavigationRoutingParams& Params)
+{
+	EWidgetNavigationRoutingPolicy Policy = EWidgetNavigationRoutingPolicy::Default;
+
+	// Find owning policy
+	TSharedPtr<SWidget> PolicyWidget = Params.FocusedWidget;
+	while (PolicyWidget.IsValid())
+	{
+		TSharedPtr<FNavigationMetaData> NavigationMetaData = PolicyWidget->GetMetaData<FNavigationMetaData>();
+		if (NavigationMetaData.IsValid())
+		{
+			Policy = NavigationMetaData->GetNavigationRoutingPolicy();
+			if (Policy != EWidgetNavigationRoutingPolicy::Default)
+			{
+				break;
+			}
+		}
+
+		PolicyWidget = PolicyWidget->GetParentWidget();
+	}
+
+	if (!PolicyWidget.IsValid())
+	{
+		return Params.FocusedWidget;
+	}
+
+	// If starting widget is within policy widget, do not apply routing policy. Only apply when focus enters the policy widget
+	const FSlateRect WidgetRect = Params.HittestGrid->GetWidgetRenderBoundingRect(Params.StartingWidget);
+	const FSlateRect BoundaryRect = Params.HittestGrid->GetWidgetRenderBoundingRect(PolicyWidget);
+	if (FSlateRect::DoRectanglesIntersect(WidgetRect, BoundaryRect))
+	{
+		return Params.FocusedWidget;
+	}
+
+	switch (Policy)
+	{
+	case EWidgetNavigationRoutingPolicy::AcceptFocus:
+	{
+		return Params.FocusedWidget;
+	}
+	case EWidgetNavigationRoutingPolicy::RouteToTopMostChild:
+	{
+		auto GetEdgeFunc = [](const FSlateRect& Rect) { return Rect.Top; };
+		return RouteNavigationFocusToEdge(Params, 1, 1, BoundaryRect, BoundaryRect.GetTopLeft(), GetEdgeFunc);
+	}
+	case EWidgetNavigationRoutingPolicy::RouteToBottomMostChild:
+	{
+		auto GetEdgeFunc = [](const FSlateRect& Rect) { return Rect.Bottom; };
+		return RouteNavigationFocusToEdge(Params, -1, 1, BoundaryRect, BoundaryRect.GetBottomLeft(), GetEdgeFunc);
+	}
+	case EWidgetNavigationRoutingPolicy::RouteToLeftMostChild:
+	{
+		auto GetEdgeFunc = [](const FSlateRect& Rect) { return Rect.Left; };
+		return RouteNavigationFocusToEdge(Params, 1, 0,  BoundaryRect, BoundaryRect.GetTopLeft(), GetEdgeFunc);
+	}
+	case EWidgetNavigationRoutingPolicy::RouteToRightMostChild:
+	{
+		auto GetEdgeFunc = [](const FSlateRect& Rect) { return Rect.Right; };
+		return RouteNavigationFocusToEdge(Params, -1, 0,  BoundaryRect, BoundaryRect.GetTopRight(), GetEdgeFunc);
+	}
+	case EWidgetNavigationRoutingPolicy::RouteToTopLeftChild:
+	{
+		// Routing to top left corner, scan from the top of the boundary, find widget nearest the left edge of rect
+		auto GetEdgeFunc = [](const FSlateRect& Rect) { return Rect.Left; };
+		return RouteNavigationFocusToEdge(Params, 1, 1, BoundaryRect, BoundaryRect.GetTopLeft(), GetEdgeFunc);
+	}
+	case EWidgetNavigationRoutingPolicy::RouteToTopRightChild:
+	{
+		// Routing to top right corner, scan from the top of the boundary, find widget nearest the right edge of rect
+		auto GetEdgeFunc = [](const FSlateRect& Rect) { return Rect.Right; };
+		return RouteNavigationFocusToEdge(Params, 1, 1, BoundaryRect, BoundaryRect.GetTopLeft(), GetEdgeFunc);
+	}
+	case EWidgetNavigationRoutingPolicy::RouteToBottomLeftChild:
+	{
+		// Routing to bottom left corner, scan from the bottom of the boundary, find widget nearest the left edge of rect
+		auto GetEdgeFunc = [](const FSlateRect& Rect) { return Rect.Left; };
+		return RouteNavigationFocusToEdge(Params, -1, 1, BoundaryRect, BoundaryRect.GetBottomLeft(), GetEdgeFunc);
+	}
+	case EWidgetNavigationRoutingPolicy::RouteToBottomRightChild:
+	{
+		// Routing to bottom right corner, scan from the bottom of the boundary, find widget nearest the right edge of rect
+		auto GetEdgeFunc = [](const FSlateRect& Rect) { return Rect.Right; };
+		return RouteNavigationFocusToEdge(Params, -1, 1, BoundaryRect, BoundaryRect.GetBottomLeft(), GetEdgeFunc);
+	}
+	default:
+	{
+		// Fallback to default. Do no routing.
+		return Params.FocusedWidget;
+	}
+	}
+}
+
+
+TSharedPtr<SWidget> FNavigationRouting::RouteNavigationFocusToEdge(const FNavigationRoutingParams& Params, int Direction, int AxisIndex, const FSlateRect& BoundaryRect, const FVector2f& StartingPoint, FGetEdgeFunc GetEdgeFunc)
+{
+	const float BoundaryEdge = GetEdgeFunc(BoundaryRect);
+
+	const FIntPoint NumCells = Params.HittestGrid->NumCells;
+	const FVector2f CellSize = Params.HittestGrid->CellSize;
+	FIntPoint CurrentCellPoint = Params.HittestGrid->GetCellCoordinate(StartingPoint);
+
+	int32 StrideAxis, StrideAxisMin, StrideAxisMax;
+	if (AxisIndex == 0)
+	{
+		StrideAxis = 1;
+		StrideAxisMin = FMath::Min(FMath::Max(FMath::FloorToInt(BoundaryRect.Top / CellSize.Y), 0), NumCells.Y - 1);
+		StrideAxisMax = FMath::Min(FMath::Max(FMath::FloorToInt(BoundaryRect.Bottom / CellSize.Y), 0), NumCells.Y - 1);
+	}
+	else
+	{
+		StrideAxis = 0;
+		StrideAxisMin = FMath::Min(FMath::Max(FMath::FloorToInt(BoundaryRect.Left / CellSize.X), 0), NumCells.X - 1);
+		StrideAxisMax = FMath::Min(FMath::Max(FMath::FloorToInt(BoundaryRect.Right / CellSize.X), 0), NumCells.X - 1);
+	}
+
+	while (CurrentCellPoint[AxisIndex] >= 0 && CurrentCellPoint[AxisIndex] < NumCells[AxisIndex])
+	{
+		FIntPoint StrideCellPoint = CurrentCellPoint;
+
+		// Find first nearest widget within stride
+		float DistanceNearestEdge = FLT_MAX;
+		TSharedPtr<SWidget> WidgetNearestEdge;
+
+		for (StrideCellPoint[StrideAxis] = StrideAxisMin; StrideCellPoint[StrideAxis] <= StrideAxisMax; ++StrideCellPoint[StrideAxis])
+		{
+			FHittestGrid::FCollapsedWidgetsArray WidgetIndexes;
+			Params.HittestGrid->GetCollapsedWidgets(WidgetIndexes, StrideCellPoint.X, StrideCellPoint.Y);
+
+			for (int32 WidgetIndex = WidgetIndexes.Num() - 1; WidgetIndex >= 0; --WidgetIndex)
+			{
+				const FHittestGrid::FWidgetData& CandidateWidgetData = WidgetIndexes[WidgetIndex].GetWidgetData();
+				const TSharedPtr<SWidget> CandidateWidget = CandidateWidgetData.GetWidget();
+				if (!CandidateWidget.IsValid())
+				{
+					continue;
+				}
+
+				if (!Params.HittestGrid->IsCompatibleUserIndex(Params.UserIndex, CandidateWidgetData.UserIndex))
+				{
+					continue;
+				}
+
+				// If we have a non escape boundary condition and this widget isn't a descendant of our boundary condition widget then it's invalid so we keep looking.
+				if (Params.NavigationReply.GetBoundaryRule() != EUINavigationRule::Escape
+					&& Params.NavigationReply.GetHandler().IsValid()
+					&& !Params.HittestGrid->IsDescendantOf(Params.NavigationReply.GetHandler().Get(), CandidateWidgetData))
+				{
+					continue;
+				}
+
+				if (!CandidateWidget->IsEnabled())
+				{
+					continue;
+				}
+
+				if (!CandidateWidget->SupportsKeyboardFocus())
+				{
+					continue;
+				}
+
+				const FSlateRect& WidgetRect = Params.HittestGrid->GetWidgetRenderBoundingRect(CandidateWidget);
+				if (!FSlateRect::DoRectanglesIntersect(BoundaryRect, WidgetRect))
+				{
+					continue;
+				}
+
+				const float Edge = GetEdgeFunc(WidgetRect);
+				const float Distance = FMath::Abs(Edge - BoundaryEdge);
+				if (!WidgetNearestEdge.IsValid() || Distance < DistanceNearestEdge)
+				{
+					DistanceNearestEdge = Distance;
+					WidgetNearestEdge = CandidateWidget;
+				}
+			}
+		}
+
+		// Return the first widget found when searching from the edge
+		if (WidgetNearestEdge.IsValid())
+		{
+			return WidgetNearestEdge;
+		}
+
+		CurrentCellPoint[AxisIndex] += Direction;
+	}
+
+	// No widget found.
+	return Params.FocusedWidget;
+}
diff --git a/Engine/Source/Runtime/SlateCore/Private/SlateCoreClasses.cpp b/Engine/Source/Runtime/SlateCore/Private/SlateCoreClasses.cpp
--- a/Engine/Source/Runtime/SlateCore/Private/SlateCoreClasses.cpp
+++ b/Engine/Source/Runtime/SlateCore/Private/SlateCoreClasses.cpp
@@ -100,6 +100,12 @@ FAutoConsoleVariableRef CVarHitTestGridDebugging(
 	GSlateHitTestGridDebugging,
 	TEXT("Whether to show a visualization of everything in the hit teest grid"));
 
+bool GSlateNavigationDebugging = false;
+FAutoConsoleVariableRef CVarNavigationDebugging(
+	TEXT("Slate.NavigationDebugging"),
+	GSlateNavigationDebugging,
+	TEXT("Whether to show a visualization of navigation from current widget"));
+
 #endif
 
 FSlateWidgetStyle::FSlateWidgetStyle() = default;
diff --git a/Engine/Source/Runtime/SlateCore/Private/Widgets/SWindow.cpp b/Engine/Source/Runtime/SlateCore/Private/Widgets/SWindow.cpp
--- a/Engine/Source/Runtime/SlateCore/Private/Widgets/SWindow.cpp
+++ b/Engine/Source/Runtime/SlateCore/Private/Widgets/SWindow.cpp
@@ -2146,6 +2146,11 @@ int32 SWindow::PaintWindow( double CurrentTime, float DeltaTime, FSlateWindowEle
 		const FGeometry& WindowGeometry = GetWindowGeometryInWindow();
 		HittestGrid->DisplayGrid(INT_MAX, WindowGeometry, OutDrawElements);
 	}
+	if (GSlateNavigationDebugging)
+	{
+		const FGeometry& WindowGeometry = GetWindowGeometryInWindow();
+		HittestGrid->DisplayNavigation(INT_MAX, WindowGeometry, OutDrawElements);
+	}
 #endif
 
 	OutDrawElements.EndDeferredGroup();
diff --git a/Engine/Source/Runtime/SlateCore/Public/Input/HittestGrid.h b/Engine/Source/Runtime/SlateCore/Public/Input/HittestGrid.h
--- a/Engine/Source/Runtime/SlateCore/Public/Input/HittestGrid.h
+++ b/Engine/Source/Runtime/SlateCore/Public/Input/HittestGrid.h
@@ -8,9 +8,11 @@
 #include "Layout/ArrangedWidget.h"
 #include "Layout/Clipping.h"
 #include "Input/Events.h"
+#include "StructUtils/InstancedStruct.h"
 #include "Widgets/SWidget.h"
 
 class FArrangedChildren;
+struct FNavigationMethod;
 
 class ICustomHitTestPath
 {
@@ -26,6 +28,10 @@ public:
 
 class FHittestGrid : public FNoncopyable
 {
+	friend struct FNavigationMethod;
+	friend struct FNavigationMethodOrthogonal;
+	friend struct FNavigationRouting;
+
 public:
 
 	SLATECORE_API FHittestGrid();
@@ -75,6 +81,15 @@ public:
 	UE::Slate::FDeprecateVector2DResult GetGridOrigin() const { return GridOrigin; }
 	UE::Slate::FDeprecateVector2DResult GetGridWindowOrigin() const { return GridWindowOrigin; }
 
+	/** Get the screen space rectangle that corresponds to the painted area of the widget **/
+	SLATECORE_API FSlateRect GetWidgetRenderBoundingRect(const TSharedPtr<SWidget>& Widget) const;
+
+	/** Returns true if the CellCoordinates are within the grid **/
+	SLATECORE_API bool IsValidCellCoord(const FIntPoint& CellCoord) const;
+	SLATECORE_API bool IsValidCellCoord(const int32 XCoord, const int32 YCoord) const;
+	
+	SLATECORE_API const FVector2f& GetCellSize() const { return CellSize; }
+
 	/** Clear the grid */
 	SLATECORE_API void Clear();
 
@@ -135,6 +150,8 @@ public:
 	};
 	SLATECORE_API void DisplayGrid(int32 InLayer, const FGeometry& AllottedGeometry, FSlateWindowElementList& WindowElementList, EDisplayGridFlags DisplayFlags = EDisplayGridFlags::UseFocusBrush) const;
 
+	SLATECORE_API void DisplayNavigation(int32 InLayer, const FGeometry& AllottedGeometry, FSlateWindowElementList& WindowElementList);
+
 	struct FWidgetSortData
 	{
 		const TWeakPtr<SWidget> WeakWidget;
@@ -237,8 +254,6 @@ private:
 	};
 
 	//~ Helper functions
-	SLATECORE_API bool IsValidCellCoord(const FIntPoint& CellCoord) const;
-	SLATECORE_API bool IsValidCellCoord(const int32 XCoord, const int32 YCoord) const;
 	SLATECORE_API void ClearInternal(int32 TotalCells);
 
 	/** Return the Index and distance to a hit given the testing params */
@@ -285,6 +300,16 @@ private:
 	/** Remove appended hittest grid that are not valid anymore. */
 	SLATECORE_API void RemoveStaleAppendedHittestGrid();
 
+	SLATECORE_API static constexpr bool IsCompatibleUserIndex(int32 RequestedUserIndex, int32 TestUserIndex)
+	{
+		// INDEX_NONE means the user index is compatible with all
+		return RequestedUserIndex == INDEX_NONE || TestUserIndex == INDEX_NONE|| RequestedUserIndex == TestUserIndex;
+	}
+
+	TInstancedStruct<FNavigationMethod> GetNavigationMethod(const FArrangedWidget& StartingWidget);
+	TSharedPtr<SWidget> FindNextFocusableWidgetDefault(const FArrangedWidget& StartingWidget, const EUINavigation Direction, const FNavigationReply& NavigationReply, const FArrangedWidget& RuleWidget, int32 UserIndex);
+	TSharedPtr<SWidget> FindNextFocusableWidgetExperimental(const FArrangedWidget& StartingWidget, const EUINavigation Direction, const FNavigationReply& NavigationReply, const FArrangedWidget& RuleWidget, int32 UserIndex);
+
 private:
 	/** Map of all the widgets currently in the hit test grid to their stable index. */
 	TMap<const SWidget*, int32> WidgetMap;
@@ -318,6 +343,14 @@ private:
 
 	/** The current slate user index that should be associated with any added widgets */
 	int32 CurrentUserIndex;
+
+	/** The size in pixels of a cell */
+	static const FVector2f CellSize;
+
+#if WITH_SLATE_DEBUGGING
+	/** Cache the state of the most recent FindNextFocusableWidget request*/
+	TInstancedStruct<FNavigationMethod> LastNavigationMethod;	
+#endif
 };
 
 #if WITH_SLATE_DEBUGGING
diff --git a/Engine/Source/Runtime/SlateCore/Public/Input/NavigationMethod.h b/Engine/Source/Runtime/SlateCore/Public/Input/NavigationMethod.h
new file mode 100644
--- /dev/null
+++ b/Engine/Source/Runtime/SlateCore/Public/Input/NavigationMethod.h
@@ -0,0 +1,58 @@
+// Copyright Epic Games, Inc. All Rights Reserved.
+
+#pragma once 
+
+#include "CoreMinimal.h"
+#include "Input/HittestGrid.h"
+
+#include "NavigationMethod.generated.h"
+
+USTRUCT()
+struct FNavigationMethod
+{
+	GENERATED_BODY()
+
+public:
+	using FDebugWidgetResult = FHittestGrid::FDebuggingFindNextFocusableWidgetArgs::FWidgetResult;
+
+	SLATECORE_API FNavigationMethod() = default;
+	SLATECORE_API virtual ~FNavigationMethod() = default;
+	SLATECORE_API virtual void Initialize(const FHittestGrid* InHittestGrid, TArray<FDebugWidgetResult>* InIntermediateResultsPtr);
+	SLATECORE_API virtual TSharedPtr<SWidget> FindNextFocusableWidget(const FArrangedWidget& StartingWidget, const EUINavigation Direction, const FNavigationReply& NavigationReply, const FArrangedWidget& RuleWidget, int32 InUserIndex);
+	
+protected:
+	SLATECORE_API FIntPoint GetCellCoordinate(FVector2f Position) const;
+	SLATECORE_API bool IsValidCellCoordinate(int32 X, int32 Y) const;
+	SLATECORE_API bool IsParentsEnabled(const SWidget* Widget);
+
+	using FWidgetFunc = TFunction<bool(const TSharedPtr<SWidget>)>;
+	SLATECORE_API void ForEachFocusableWidgetsInCell(int32 X, int32 Y, const FNavigationReply& NavigationReply, int32 UserIndex, FWidgetFunc WidgetFunc);
+	
+	const FHittestGrid* HittestGrid;
+	TArray<FDebugWidgetResult>* IntermediateResultsPtr;
+	TSet<TSharedPtr<SWidget>> DisabledDestinations;
+
+private:
+	static FString DefaultNavigationMethod;
+	static TMap<FString, TSharedPtr<FNavigationMethod>> RegisteredNavigationMethods;
+
+#if WITH_SLATE_DEBUGGING
+public:
+	SLATECORE_API virtual void DrawDebug(int32 InLayer, const FGeometry& AllottedGeometry, FSlateWindowElementList& WindowElementList);
+
+protected:
+	struct FDebuggingText
+	{
+		static FText Valid;
+		static FText NotCompatibleWithUserIndex;
+		static FText DoesNotIntersect;
+		static FText PreviousWidgetIsBetter;
+		static FText NotADescendant;
+		static FText Disabled;
+		static FText ParentDisabled;
+		static FText DoesNotSuportKeyboardFocus;
+	};
+
+	SLATECORE_API void AddDebugIntermediateResult(const TSharedPtr<const SWidget>& InWidget, const FText Result);
+#endif // WITH_SLATE_DEBUGGING
+};
diff --git a/Engine/Source/Runtime/SlateCore/Public/Input/NavigationRouting.h b/Engine/Source/Runtime/SlateCore/Public/Input/NavigationRouting.h
new file mode 100644
--- /dev/null
+++ b/Engine/Source/Runtime/SlateCore/Public/Input/NavigationRouting.h
@@ -0,0 +1,53 @@
+// Copyright Epic Games, Inc. All Rights Reserved.
+
+#pragma once 
+
+#include "CoreMinimal.h"
+
+#include "NavigationRouting.generated.h"
+
+class SWidget;
+class FHittestGrid;
+class FNavigationReply;
+class FSlateRect;
+
+UENUM()
+enum class EWidgetNavigationRoutingPolicy : uint8
+{
+	AcceptFocus UMETA(DisplayName="Accept Focus", Tooltip = "On navigation, this widget, or widgets within this boundary, will receive focus. Focus will not be routed"),
+	RouteToTopMostChild  UMETA(DisplayName="Route To Top Most Child", Tooltip="On navigation, this widget, or widgets within this boundary, will route focus to topmost visibile widget in this container"),
+	RouteToBottomMostChild  UMETA(DisplayName="Route To Bottom Most Child", Tooltip="On navigation, this widget, or widgets within this boundary, will route focus to bottommost visibile widget in this container"),
+	RouteToLeftMostChild UMETA(DisplayName="Route To Left Most Child", Tooltip="On navigation, this widget, or widgets within this boundary, will route focus to leftmost visibile widget in this container"),
+	RouteToRightMostChild UMETA(DisplayName="Route To Right Most Child", Tooltip="On navigation, this widget, or widgets within this boundary, will route focus to rightmost visibile widget in this container"),
+	RouteToTopLeftChild UMETA(DisplayName="Route To Top Left Child", Tooltip="On navigation, this widget, or widgets within this boundary, will route focus to widget in the top left container"),
+	RouteToTopRightChild UMETA(DisplayName="Route To Top Right Child", Tooltip="On navigation, this widget, or widgets within this boundary, will route focus to widget in the top right container"),
+	RouteToBottomLeftChild UMETA(DisplayName="Route To Bottom Left Child", Tooltip="On navigation, this widget, or widgets within this boundary, will route focus to widget in the bottom left container"),
+	RouteToBottomRightChild UMETA(DisplayName="Route To Bottom Right Child", Tooltip="On navigation, this widget, or widgets within this boundary, will route focus to widget in the bottom right container"),
+
+	Default = AcceptFocus,
+};
+
+struct FNavigationRoutingParams
+{
+	FNavigationRoutingParams(TSharedPtr<SWidget> StartingWidget, TSharedPtr<SWidget> FocusedWidget, const FHittestGrid* HittestGrid, const FNavigationReply& NavigationReply, int32 UserIndex)
+		: StartingWidget(StartingWidget)
+		, FocusedWidget(FocusedWidget)
+		, HittestGrid(HittestGrid)
+		, NavigationReply(NavigationReply)
+		, UserIndex(UserIndex)
+	{}
+	TSharedPtr<SWidget> StartingWidget;
+	TSharedPtr<SWidget> FocusedWidget;
+	const FHittestGrid* HittestGrid;
+	const FNavigationReply& NavigationReply;
+	int32 UserIndex;
+};
+
+struct FNavigationRouting
+{
+	static TSharedPtr<SWidget> RouteNavigationFocus(const FNavigationRoutingParams& Params);
+
+private:
+	using FGetEdgeFunc = TFunction<float(const FSlateRect&)>;
+	static TSharedPtr<SWidget> RouteNavigationFocusToEdge(const FNavigationRoutingParams& Params, int Direction, int AxisIndex, const FSlateRect& BoundaryRect, const FVector2f& StartingPoint, FGetEdgeFunc GetEdgeFunc);
+};
\ No newline at end of file
diff --git a/Engine/Source/Runtime/SlateCore/Public/SlateGlobals.h b/Engine/Source/Runtime/SlateCore/Public/SlateGlobals.h
--- a/Engine/Source/Runtime/SlateCore/Public/SlateGlobals.h
+++ b/Engine/Source/Runtime/SlateCore/Public/SlateGlobals.h
@@ -88,6 +88,7 @@ extern SLATECORE_API bool GSlateCheckUObjectShapedGlyphSequence;
 
 #if WITH_SLATE_DEBUGGING
 extern SLATECORE_API bool GSlateHitTestGridDebugging;
+extern SLATECORE_API bool GSlateNavigationDebugging;
 #endif
 /* Forward declarations
 *****************************************************************************/
diff --git a/Engine/Source/Runtime/SlateCore/Public/Types/NavigationMetaData.h b/Engine/Source/Runtime/SlateCore/Public/Types/NavigationMetaData.h
--- a/Engine/Source/Runtime/SlateCore/Public/Types/NavigationMetaData.h
+++ b/Engine/Source/Runtime/SlateCore/Public/Types/NavigationMetaData.h
@@ -3,9 +3,12 @@
 #pragma once
 
 #include "CoreMinimal.h"
-#include "Types/SlateEnums.h"
 #include "Input/NavigationReply.h"
+#include "Input/NavigationRouting.h"
+#include "Input/NavigationMethod.h"
+#include "StructUtils/InstancedStruct.h"
 #include "Types/ISlateMetaData.h"
+#include "Types/SlateEnums.h"
 
 class SWidget;
 
@@ -119,6 +122,27 @@ public:
 		Rules[(uint8)InNavigation].FocusDelegate = nullptr;
 		Rules[(uint8)InNavigation].FocusRecipient = nullptr;
 	}
+
+	void SetNavigationMethodStruct(const TInstancedStruct<FNavigationMethod>& InNavigationMethodStruct) 
+	{
+		NavigationMethodStruct = InNavigationMethodStruct;
+	}
+
+	const TInstancedStruct<FNavigationMethod>& GetNavigationMethodStruct() const
+	{
+		return NavigationMethodStruct;
+	}
+
+	void SetNavigationRoutingPolicy(EWidgetNavigationRoutingPolicy InNavigationRoutingPolicy) 
+	{
+		NavigationRoutingPolicy = InNavigationRoutingPolicy;
+	}
+
+	EWidgetNavigationRoutingPolicy GetNavigationRoutingPolicy() const
+	{
+		return NavigationRoutingPolicy;
+	}
+
 private:
 
 	struct SNavData
@@ -129,6 +153,8 @@ private:
 	};
 	SNavData Rules[(uint8)EUINavigation::Num];
 
+	EWidgetNavigationRoutingPolicy NavigationRoutingPolicy = EWidgetNavigationRoutingPolicy::Default;
+	TInstancedStruct<FNavigationMethod> NavigationMethodStruct;
 };
 
 #ifndef UE_WITH_SLATE_SIMULATEDNAVIGATIONMETADATA
diff --git a/Engine/Source/Runtime/UMG/Private/UWidgetNavigation.cpp b/Engine/Source/Runtime/UMG/Private/UWidgetNavigation.cpp
--- a/Engine/Source/Runtime/UMG/Private/UWidgetNavigation.cpp
+++ b/Engine/Source/Runtime/UMG/Private/UWidgetNavigation.cpp
@@ -4,6 +4,7 @@
 #include "Types/SlateEnums.h"
 #include "Input/NavigationReply.h"
 #include "Types/NavigationMetaData.h"
+#include "Input/NavigationMethod.h"
 #include "Blueprint/WidgetNavigation.h"
 #include "Blueprint/WidgetTree.h"
 
@@ -119,6 +120,9 @@ void UWidgetNavigation::UpdateMetaData(TSharedRef<FNavigationMetaData> MetaData)
 	UpdateMetaDataEntry(MetaData, Right, EUINavigation::Right);
 	UpdateMetaDataEntry(MetaData, Next, EUINavigation::Next);
 	UpdateMetaDataEntry(MetaData, Previous, EUINavigation::Previous);
+
+	MetaData->SetNavigationMethodStruct(NavigationMethod);
+	MetaData->SetNavigationRoutingPolicy(RoutingPolicy);
 }
 
 bool UWidgetNavigation::IsDefaultNavigation() const
@@ -128,7 +132,9 @@ bool UWidgetNavigation::IsDefaultNavigation() const
 		Left.Rule == EUINavigationRule::Escape &&
 		Right.Rule == EUINavigationRule::Escape &&
 		Next.Rule == EUINavigationRule::Escape &&
-		Previous.Rule == EUINavigationRule::Escape;
+		Previous.Rule == EUINavigationRule::Escape &&
+		RoutingPolicy == EWidgetNavigationRoutingPolicy::Default &&
+		!NavigationMethod.IsValid();
 }
 
 void UWidgetNavigation::UpdateMetaDataEntry(TSharedRef<FNavigationMetaData> MetaData, const FWidgetNavigationData & NavData, EUINavigation Nav)
diff --git a/Engine/Source/Runtime/UMG/Public/Blueprint/WidgetNavigation.h b/Engine/Source/Runtime/UMG/Public/Blueprint/WidgetNavigation.h
--- a/Engine/Source/Runtime/UMG/Public/Blueprint/WidgetNavigation.h
+++ b/Engine/Source/Runtime/UMG/Public/Blueprint/WidgetNavigation.h
@@ -3,10 +3,12 @@
 #pragma once
 
 #include "CoreMinimal.h"
+#include "StructUtils/InstancedStruct.h"
+#include "Types/NavigationMetaData.h"
+#include "Input/NavigationMethod.h"
 #include "UObject/ObjectMacros.h"
 #include "UObject/Object.h"
 #include "UObject/WeakObjectPtr.h"
-#include "Types/NavigationMetaData.h"
 
 #include "WidgetNavigation.generated.h"
 
@@ -76,6 +78,12 @@ public:
 	UPROPERTY(EditAnywhere, BlueprintReadOnly, Category="Navigation")
 	FWidgetNavigationData Previous;
 
+	UPROPERTY(EditAnywhere, BlueprintReadOnly, Category="Navigation")
+	EWidgetNavigationRoutingPolicy RoutingPolicy = EWidgetNavigationRoutingPolicy::AcceptFocus;
+
+	UPROPERTY(EditAnywhere, BlueprintReadOnly, Category="Navigation", meta=(ExcludeBaseStruct))
+	TInstancedStruct<FNavigationMethod> NavigationMethod;
+
 public:
 
 #if WITH_EDITOR
