Slate: Fix SStackBox not respecting the Orientation slate arg

diff --git a/Engine/Source/Runtime/SlateCore/Private/Widgets/SBoxPanel.cpp b/Engine/Source/Runtime/SlateCore/Private/Widgets/SBoxPanel.cpp
--- a/Engine/Source/Runtime/SlateCore/Private/Widgets/SBoxPanel.cpp
+++ b/Engine/Source/Runtime/SlateCore/Private/Widgets/SBoxPanel.cpp
@@ -54,6 +54,7 @@ void SVerticalBox::Construct( const SVerticalBox::FArguments& InArgs )
 
 void SStackBox::Construct(const SStackBox::FArguments& InArgs)
 {
+	SetOrientation(InArgs._Orientation);
 	Children.Reserve(InArgs._Slots.Num());
 	for (const FSlot::FSlotArguments& Arg : InArgs._Slots)
 	{
