Changed default button text styles in StarshipCoreStyle to use the Ellipsis OverflowPolicy instead of Clip.

[RN]

diff --git a/Engine/Source/Runtime/SlateCore/Private/Styling/StarshipCoreStyle.cpp b/Engine/Source/Runtime/SlateCore/Private/Styling/StarshipCoreStyle.cpp
--- a/Engine/Source/Runtime/SlateCore/Private/Styling/StarshipCoreStyle.cpp
+++ b/Engine/Source/Runtime/SlateCore/Private/Styling/StarshipCoreStyle.cpp
@@ -1433,8 +1433,9 @@ void FStarshipCoreStyle::SetupButtonStyles(TSharedRef<FStyle>& Style)
 	{
 
 		const FTextBlockStyle& NormalText = Style->GetWidgetStyle<FTextBlockStyle>("NormalText");
+		const FTextBlockStyle ButtonText = FTextBlockStyle(NormalText).SetOverflowPolicy(ETextOverflowPolicy::Ellipsis);
 
-		Style->Set("ButtonText", NormalText);
+		Style->Set("ButtonText", ButtonText);
 
 		Style->Set("PrimaryButton", PrimaryButton);
 		Style->Set("PrimaryButtonLabelAndIconNormalPadding", PrimaryButtonLabelAndIconNormalPadding);
@@ -1444,7 +1445,7 @@ void FStarshipCoreStyle::SetupButtonStyles(TSharedRef<FStyle>& Style)
 		Style->Set("SimpleButton", SimpleButtonIconOnly);
 		Style->Set("SimpleButtonLabelAndIcon", SimpleButtonLabelAndIcon);
 
-		Style->Set("PrimaryButtonText", FTextBlockStyle(NormalText).SetFont(FStyleFonts::Get().NormalBold));
+		Style->Set("PrimaryButtonText", FTextBlockStyle(ButtonText).SetFont(FStyleFonts::Get().NormalBold));
 
 		FSlateFontInfo AddFont = FStyleFonts::Get().SmallBold;
 		AddFont.LetterSpacing = 200;
@@ -1452,11 +1453,11 @@ void FStarshipCoreStyle::SetupButtonStyles(TSharedRef<FStyle>& Style)
 		Style->Set("SecondaryButton", Button);
 
 		//@todo Recolor - This is redundant with button text now. Temporarily here for backwards compat
-		Style->Set("DialogButtonText", NormalText);
+		Style->Set("DialogButtonText", ButtonText);
 
 
 		//@todo Recolor - Rename to simple button text
-		Style->Set("SmallButtonText", FTextBlockStyle(NormalText));
+		Style->Set("SmallButtonText", FTextBlockStyle(ButtonText));
 
 		// A button that only responds to clicks and displays nothing
 		const FButtonStyle InvisibleButton = FButtonStyle()
