Fix the tab flash effect when attempting to open a tab that is already open.

[RN]

diff --git a/Engine/Source/Runtime/SlateCore/Private/Styling/StarshipCoreStyle.cpp b/Engine/Source/Runtime/SlateCore/Private/Styling/StarshipCoreStyle.cpp
--- a/Engine/Source/Runtime/SlateCore/Private/Styling/StarshipCoreStyle.cpp
+++ b/Engine/Source/Runtime/SlateCore/Private/Styling/StarshipCoreStyle.cpp
@@ -1816,7 +1816,7 @@ void FStarshipCoreStyle::SetupDockingStyles(TSharedRef<FStyle>& Style)
 		.SetNormalBrush(FSlateNoResource())
 		.SetHoveredBrush(BOX_BRUSH("/Starship/Docking/DockTab_Foreground", 4.0f / 20.0f, FStyleColors::Panel.GetSpecifiedColor().CopyWithNewOpacity(0.8))) // Panel + 20% Black
 		.SetForegroundBrush(BOX_BRUSH("/Starship/Docking/DockTab_Foreground", 4.0f / 20.0f, FStyleColors::Panel))
-		.SetColorOverlayTabBrush(FSlateNoResource())
+		.SetColorOverlayTabBrush(FSlateColorBrush(FStyleColors::Foreground))
 		.SetColorOverlayIconBrush(FSlateNoResource())
 		.SetContentAreaBrush(FSlateColorBrush(FStyleColors::Panel))
 		.SetTabWellBrush(FSlateColorBrush(FStyleColors::Background))
@@ -1842,7 +1842,7 @@ void FStarshipCoreStyle::SetupDockingStyles(TSharedRef<FStyle>& Style)
 		.SetHoveredBrush(BOX_BRUSH("/Starship/Docking/DockTab_Foreground", 4.0f / 20.0f, FStyleColors::Panel.GetSpecifiedColor().CopyWithNewOpacity(0.8))) // Panel + 20% Black
 		.SetForegroundBrush(BOX_BRUSH("/Starship/Docking/DockTab_Foreground", 4.0f / 20.0f, FStyleColors::Panel))
 
-		.SetColorOverlayTabBrush(FSlateNoResource())
+		.SetColorOverlayTabBrush(FSlateColorBrush(FStyleColors::Foreground))
 		.SetColorOverlayIconBrush(FSlateRoundedBoxBrush(FStyleColors::Foreground, 4.f))
 		.SetContentAreaBrush(FSlateColorBrush(FStyleColors::Background))
 		.SetTabWellBrush(FSlateColorBrush(FStyleColors::Background))
