Replace some usages of FORCEINLINE with inline in Slate modules.

For context about this change, see this doc: https://docs.google.com/document/d/1LjpT7NpzTecfXdFNyv5KM6sCu67827BRljNi1iaWVSE/edit?tab=t.0#heading=h.fvtz8gjukz7x

Includes modules: SlateScriptingCommands, Slate, SlateCore

Reviewers were automatically determined. Please comment if anyone else should be added.

#rnx

diff --git a/Engine/Plugins/Slate/SlateScripting/Source/SlateScriptingCommands/Public/UICommandsScriptingSubsystem.h b/Engine/Plugins/Slate/SlateScripting/Source/SlateScriptingCommands/Public/UICommandsScriptingSubsystem.h
--- a/Engine/Plugins/Slate/SlateScripting/Source/SlateScriptingCommands/Public/UICommandsScriptingSubsystem.h
+++ b/Engine/Plugins/Slate/SlateScripting/Source/SlateScriptingCommands/Public/UICommandsScriptingSubsystem.h
@@ -92,11 +92,11 @@ struct FScriptingCommand
 	/** The delegate to call to check whether the command can be executed */
 	FCanExecuteAction OnCanExecuteAction;
 
-	FORCEINLINE bool operator==(const FScriptingCommand& Other) const { return Other.CommandInfo.Name == CommandInfo.Name; }
-	FORCEINLINE bool operator!=(const FScriptingCommand& Other) const { return !(*this == Other); }
+	inline bool operator==(const FScriptingCommand& Other) const { return Other.CommandInfo.Name == CommandInfo.Name; }
+	inline bool operator!=(const FScriptingCommand& Other) const { return !(*this == Other); }
 
-	FORCEINLINE bool operator==(const FName OtherName) const { return OtherName == CommandInfo.Name; }
-	FORCEINLINE bool operator!=(const FName OtherName) const { return !(*this == OtherName); }
+	inline bool operator==(const FName OtherName) const { return OtherName == CommandInfo.Name; }
+	inline bool operator!=(const FName OtherName) const { return !(*this == OtherName); }
 	
 	/** Builds a full name in the format context.set.command_name to avoid conflicts with commands registered in different contexts/sets  */
 	FName GetFullName() const
diff --git a/Engine/Source/Runtime/Slate/Public/Framework/Application/AnalogCursor.h b/Engine/Source/Runtime/Slate/Public/Framework/Application/AnalogCursor.h
--- a/Engine/Source/Runtime/Slate/Public/Framework/Application/AnalogCursor.h
+++ b/Engine/Source/Runtime/Slate/Public/Framework/Application/AnalogCursor.h
@@ -66,7 +66,7 @@ protected:
 	SLATE_API virtual bool IsRelevantInput(const FPointerEvent& MouseEvent) const;
 
 	/** Getter */
-	FORCEINLINE const FVector2D& GetAnalogValues( EAnalogStick Stick = EAnalogStick::Left ) const
+	inline const FVector2D& GetAnalogValues( EAnalogStick Stick = EAnalogStick::Left ) const
 	{
 		return AnalogValues[ static_cast< uint8 >( Stick ) ];
 	}
@@ -97,7 +97,7 @@ protected:
 
 private:
 
-	FORCEINLINE FVector2D& GetAnalogValue( EAnalogStick Stick )
+	inline FVector2D& GetAnalogValue( EAnalogStick Stick )
 	{
 		return AnalogValues[ static_cast< uint8 >( Stick ) ];
 	}
diff --git a/Engine/Source/Runtime/Slate/Public/Framework/Application/SlateApplication.h b/Engine/Source/Runtime/Slate/Public/Framework/Application/SlateApplication.h
--- a/Engine/Source/Runtime/Slate/Public/Framework/Application/SlateApplication.h
+++ b/Engine/Source/Runtime/Slate/Public/Framework/Application/SlateApplication.h
@@ -1062,19 +1062,19 @@ public:
 	SLATE_API bool TakeHDRScreenshot(const TSharedRef<SWidget>& Widget, const FIntRect& InnerWidgetArea, TArray<FLinearColor>& OutColorData, FIntVector& OutSize);
 
 	/** Gets the user at the given index, null if the user does not exist. */
-	FORCEINLINE TSharedPtr<const FSlateUser> GetUser(int32 UserIndex) const
+	inline TSharedPtr<const FSlateUser> GetUser(int32 UserIndex) const
 	{
 		return Users.IsValidIndex(UserIndex) ? Users[UserIndex] : nullptr;
 	}
-	FORCEINLINE TSharedPtr<FSlateUser> GetUser(int32 UserIndex)
+	inline TSharedPtr<FSlateUser> GetUser(int32 UserIndex)
 	{
 		return Users.IsValidIndex(UserIndex) ? Users[UserIndex] : nullptr;
 	}
 	SLATE_API TSharedPtr<FSlateUser> GetUser(FPlatformUserId PlatformUser);
-	FORCEINLINE TSharedPtr<const FSlateUser> GetUser(const FInputEvent& InputEvent) const { return GetUser(InputEvent.GetUserIndex()); }
-	FORCEINLINE TSharedPtr<FSlateUser> GetUser(const FInputEvent& InputEvent) { return GetUser(InputEvent.GetUserIndex()); }
+	inline TSharedPtr<const FSlateUser> GetUser(const FInputEvent& InputEvent) const { return GetUser(InputEvent.GetUserIndex()); }
+	inline TSharedPtr<FSlateUser> GetUser(const FInputEvent& InputEvent) { return GetUser(InputEvent.GetUserIndex()); }
 	
-	FORCEINLINE TSharedPtr<FSlateUser> GetUserFromControllerId(int32 ControllerId)
+	inline TSharedPtr<FSlateUser> GetUserFromControllerId(int32 ControllerId)
 	{
 		TOptional<int32> UserIndex = GetUserIndexForController(ControllerId);
 		if (UserIndex.IsSet())
@@ -1083,7 +1083,7 @@ public:
 		}
 		return nullptr;
 	}
-	FORCEINLINE TSharedPtr<const FSlateUser> GetUserFromControllerId(int32 ControllerId) const
+	inline TSharedPtr<const FSlateUser> GetUserFromControllerId(int32 ControllerId) const
 	{
 		TOptional<int32> UserIndex = GetUserIndexForController(ControllerId);
 		if (UserIndex.IsSet())
@@ -1098,13 +1098,13 @@ public:
 	SLATE_API TSharedPtr<const FSlateUser> GetUserFromPlatformUser(FPlatformUserId PlatformUser) const;
 
 	/** Get the standard 'default' user (there's always guaranteed to be at least one). */
-	FORCEINLINE TSharedPtr<const FSlateUser> GetCursorUser() const
+	inline TSharedPtr<const FSlateUser> GetCursorUser() const
 	{
 		TSharedPtr<const FSlateUser> SlateUser = GetUser(CursorUserIndex);
 		check(SlateUser.IsValid());
 		return SlateUser;
 	}
-	FORCEINLINE TSharedPtr<FSlateUser> GetCursorUser()
+	inline TSharedPtr<FSlateUser> GetCursorUser()
 	{
 		TSharedPtr<FSlateUser> SlateUser = GetUser(CursorUserIndex);
 		check(SlateUser.IsValid());
@@ -1130,7 +1130,7 @@ public:
 	 * Gets time step in seconds if a fixed delta time is wanted.
 	 * @return Time step in seconds for fixed delta time.
 	 */
-	FORCEINLINE static double GetFixedDeltaTime()
+	inline static double GetFixedDeltaTime()
 	{
 		return FixedDeltaTime;
 	}
@@ -1175,7 +1175,7 @@ protected:
 	 */
 	SLATE_API TSharedRef<FSlateUser> GetOrCreateUser(FInputDeviceId DeviceId);
 	
-	FORCEINLINE TSharedRef<FSlateUser> GetOrCreateUser(const FInputEvent& InputEvent) { return GetOrCreateUser(InputEvent.GetUserIndex()); }
+	inline TSharedRef<FSlateUser> GetOrCreateUser(const FInputEvent& InputEvent) { return GetOrCreateUser(InputEvent.GetUserIndex()); }
 
 	friend class FEventRouter;
 
diff --git a/Engine/Source/Runtime/Slate/Public/Framework/Application/SlateUser.h b/Engine/Source/Runtime/Slate/Public/Framework/Application/SlateUser.h
--- a/Engine/Source/Runtime/Slate/Public/Framework/Application/SlateUser.h
+++ b/Engine/Source/Runtime/Slate/Public/Framework/Application/SlateUser.h
@@ -19,8 +19,8 @@ class FSlateVirtualUserHandle
 public:
 	SLATE_API virtual ~FSlateVirtualUserHandle();
 
-	FORCEINLINE int32 GetUserIndex() const { return UserIndex; }
-	FORCEINLINE int32 GetVirtualUserIndex() const { return VirtualUserIndex; }
+	inline int32 GetUserIndex() const { return UserIndex; }
+	inline int32 GetVirtualUserIndex() const { return VirtualUserIndex; }
 
 SLATE_SCOPE:
 	FSlateVirtualUserHandle(int32 InUserIndex, int32 InVirtualUserIndex);
@@ -42,9 +42,9 @@ class FSlateUser : public TSharedFromThis<FSlateUser>
 public:
 	SLATE_API virtual ~FSlateUser();
 
-	FORCEINLINE int32 GetUserIndex() const { return UserIndex; }
-	FORCEINLINE FPlatformUserId GetPlatformUserId() const { return PlatformUser; }
-	FORCEINLINE bool IsVirtualUser() const { return !Cursor.IsValid(); }
+	inline int32 GetUserIndex() const { return UserIndex; }
+	inline FPlatformUserId GetPlatformUserId() const { return PlatformUser; }
+	inline bool IsVirtualUser() const { return !Cursor.IsValid(); }
 
 	SLATE_API TSharedPtr<SWidget> GetFocusedWidget() const;
 	SLATE_API bool ShouldShowFocus(TSharedPtr<const SWidget> Widget) const;
@@ -130,10 +130,10 @@ SLATE_SCOPE:
 	static TSharedRef<FSlateUser> Create(int32 InUserIndex, TSharedPtr<ICursor> InCursor);
 	static SLATE_API TSharedRef<FSlateUser> Create(FPlatformUserId InPlatformUserId, TSharedPtr<ICursor> InCursor);
 	
-	FORCEINLINE bool HasValidFocusPath() const { return WeakFocusPath.IsValid(); }
-	FORCEINLINE const FWeakWidgetPath& GetWeakFocusPath() const { return WeakFocusPath; }
+	inline bool HasValidFocusPath() const { return WeakFocusPath.IsValid(); }
+	inline const FWeakWidgetPath& GetWeakFocusPath() const { return WeakFocusPath; }
 	
-	FORCEINLINE TSharedRef<FWidgetPath> GetFocusPath() const
+	inline TSharedRef<FWidgetPath> GetFocusPath() const
 	{
 		if (!StrongFocusPath.IsValid())
 		{
diff --git a/Engine/Source/Runtime/Slate/Public/Framework/Docking/WorkspaceItem.h b/Engine/Source/Runtime/Slate/Public/Framework/Docking/WorkspaceItem.h
--- a/Engine/Source/Runtime/Slate/Public/Framework/Docking/WorkspaceItem.h
+++ b/Engine/Source/Runtime/Slate/Public/Framework/Docking/WorkspaceItem.h
@@ -14,7 +14,7 @@ class FWorkspaceItem : public TSharedFromThis<FWorkspaceItem>
 protected:
 	struct FWorkspaceItemSort
 	{
-		FORCEINLINE bool operator()( const TSharedRef<FWorkspaceItem> A, const TSharedRef<FWorkspaceItem> B ) const
+		inline bool operator()( const TSharedRef<FWorkspaceItem> A, const TSharedRef<FWorkspaceItem> B ) const
 		{
 			if ( A->GetChildItems().Num() > 0 )
 			{
diff --git a/Engine/Source/Runtime/Slate/Public/Framework/Text/IRun.h b/Engine/Source/Runtime/Slate/Public/Framework/Text/IRun.h
--- a/Engine/Source/Runtime/Slate/Public/Framework/Text/IRun.h
+++ b/Engine/Source/Runtime/Slate/Public/Framework/Text/IRun.h
@@ -54,13 +54,13 @@ struct FShapedTextContext
 	{
 	}
 
-	FORCEINLINE bool operator==(const FShapedTextContext& Other) const
+	inline bool operator==(const FShapedTextContext& Other) const
 	{
 		return TextShapingMethod == Other.TextShapingMethod
 			&& BaseDirection == Other.BaseDirection;
 	}
 
-	FORCEINLINE bool operator!=(const FShapedTextContext& Other) const
+	inline bool operator!=(const FShapedTextContext& Other) const
 	{
 		return !(*this == Other);
 	}
@@ -102,13 +102,13 @@ struct FLayoutBlockTextContext : public FRunTextContext
 	{
 	}
 
-	FORCEINLINE bool operator==(const FLayoutBlockTextContext& Other) const
+	inline bool operator==(const FLayoutBlockTextContext& Other) const
 	{
 		return FRunTextContext::operator==(Other)
 			&& TextDirection == Other.TextDirection;
 	}
 
-	FORCEINLINE bool operator!=(const FLayoutBlockTextContext& Other) const
+	inline bool operator!=(const FLayoutBlockTextContext& Other) const
 	{
 		return !(*this == Other);
 	}
diff --git a/Engine/Source/Runtime/Slate/Public/Framework/Text/ShapedTextCache.h b/Engine/Source/Runtime/Slate/Public/Framework/Text/ShapedTextCache.h
--- a/Engine/Source/Runtime/Slate/Public/Framework/Text/ShapedTextCache.h
+++ b/Engine/Source/Runtime/Slate/Public/Framework/Text/ShapedTextCache.h
@@ -21,7 +21,7 @@ public:
 	{
 	}
 
-	FORCEINLINE bool operator==(const FCachedShapedTextKey& Other) const
+	inline bool operator==(const FCachedShapedTextKey& Other) const
 	{
 		return TextRange == Other.TextRange
 			&& Scale == Other.Scale
@@ -29,7 +29,7 @@ public:
 			&& FontInfo.IsIdenticalTo(Other.FontInfo);
 	}
 
-	FORCEINLINE bool operator!=(const FCachedShapedTextKey& Other) const
+	inline bool operator!=(const FCachedShapedTextKey& Other) const
 	{
 		return !(*this == Other);
 	}
diff --git a/Engine/Source/Runtime/Slate/Public/Framework/Text/TextLayout.h b/Engine/Source/Runtime/Slate/Public/Framework/Text/TextLayout.h
--- a/Engine/Source/Runtime/Slate/Public/Framework/Text/TextLayout.h
+++ b/Engine/Source/Runtime/Slate/Public/Framework/Text/TextLayout.h
@@ -136,7 +136,7 @@ public:
 	int32 GetOffset() const { return Offset; }
 	bool IsValid() const { return LineIndex != INDEX_NONE && Offset != INDEX_NONE; }
 
-	friend FORCEINLINE uint32 GetTypeHash(const FTextLocation& InSubject)
+	friend inline uint32 GetTypeHash(const FTextLocation& InSubject)
 	{
 		return HashCombine(InSubject.LineIndex, InSubject.Offset);
 	}
@@ -374,8 +374,8 @@ public:
 
 	SLATE_API virtual ~FTextLayout();
 
-	FORCEINLINE const TArray< FTextLayout::FLineView >& GetLineViews() const { return LineViews; }
-	FORCEINLINE const TArray< FTextLayout::FLineModel >& GetLineModels() const { return LineModels; }
+	inline const TArray< FTextLayout::FLineView >& GetLineViews() const { return LineViews; }
+	inline const TArray< FTextLayout::FLineModel >& GetLineModels() const { return LineModels; }
 
 	/**
 	 * Get the size of the text layout, including any lines which extend beyond the wrapping boundaries (eg, lines with lots of trailing whitespace, or lines with no break candidates)
@@ -406,40 +406,40 @@ public:
 	 */
 	SLATE_API FVector2D GetWrappedDrawSize() const;
 
-	FORCEINLINE float GetWrappingWidth() const { return WrappingWidth; }
+	inline float GetWrappingWidth() const { return WrappingWidth; }
 	SLATE_API void SetWrappingWidth( float Value );
 
-	FORCEINLINE ETextWrappingPolicy GetWrappingPolicy() const { return WrappingPolicy; }
+	inline ETextWrappingPolicy GetWrappingPolicy() const { return WrappingPolicy; }
 	SLATE_API void SetWrappingPolicy(ETextWrappingPolicy Value);
 
-	FORCEINLINE float GetLineHeightPercentage() const { return LineHeightPercentage; }
+	inline float GetLineHeightPercentage() const { return LineHeightPercentage; }
 	SLATE_API void SetLineHeightPercentage( float Value );
 
-	FORCEINLINE bool GetApplyLineHeightToBottomLine() const { return ApplyLineHeightToBottomLine; }
+	inline bool GetApplyLineHeightToBottomLine() const { return ApplyLineHeightToBottomLine; }
 	SLATE_API void SetApplyLineHeightToBottomLine( bool Value );
 
-	FORCEINLINE ETextJustify::Type GetJustification() const { return Justification; }
+	inline ETextJustify::Type GetJustification() const { return Justification; }
 	SLATE_API void SetJustification( ETextJustify::Type Value );
 
 	/** Get the visual justification for this document (based on the visual justification used by the first line of text) */
 	SLATE_API ETextJustify::Type GetVisualJustification() const;
 
 	/** @note This option is destructive to the model text, so changing it requires refreshing the text layout from its marshaller */
-	FORCEINLINE ETextTransformPolicy GetTransformPolicy() const { return TransformPolicy; }
+	inline ETextTransformPolicy GetTransformPolicy() const { return TransformPolicy; }
 	SLATE_API void SetTransformPolicy(ETextTransformPolicy Value);
 
-	FORCEINLINE float GetScale() const { return Scale; }
+	inline float GetScale() const { return Scale; }
 	SLATE_API void SetScale( float Value );
 
-	FORCEINLINE ETextShapingMethod GetTextShapingMethod() const { return TextShapingMethod; }
+	inline ETextShapingMethod GetTextShapingMethod() const { return TextShapingMethod; }
 	SLATE_API void SetTextShapingMethod( const ETextShapingMethod InTextShapingMethod );
 
-	FORCEINLINE ETextFlowDirection GetTextFlowDirection() const { return TextFlowDirection; }
+	inline ETextFlowDirection GetTextFlowDirection() const { return TextFlowDirection; }
 	SLATE_API void SetTextFlowDirection( const ETextFlowDirection InTextFlowDirection );
 
 	SLATE_API void SetTextOverflowPolicy(const TOptional<ETextOverflowPolicy> InTextOverflowPolicy);
 
-	FORCEINLINE FMargin GetMargin() const { return Margin; }
+	inline FMargin GetMargin() const { return Margin; }
 	SLATE_API void SetMargin( const FMargin& InMargin );
 
 	SLATE_API void SetVisibleRegion( const FVector2D& InViewSize, const FVector2D& InScrollOffset );
diff --git a/Engine/Source/Runtime/Slate/Public/Widgets/Input/SVirtualJoystick.h b/Engine/Source/Runtime/Slate/Public/Widgets/Input/SVirtualJoystick.h
--- a/Engine/Source/Runtime/Slate/Public/Widgets/Input/SVirtualJoystick.h
+++ b/Engine/Source/Runtime/Slate/Public/Widgets/Input/SVirtualJoystick.h
@@ -164,7 +164,7 @@ protected:
 	/** 
 	 * Return the target opacity to lerp to given the current state
 	 */
-	SLATE_API FORCEINLINE float GetBaseOpacity();
+	SLATE_API inline float GetBaseOpacity();
 
 	/**
 	 * TArray specialization for controls. In the game only few joysticks presented
diff --git a/Engine/Source/Runtime/Slate/Public/Widgets/Layout/SGridPanel.h b/Engine/Source/Runtime/Slate/Public/Widgets/Layout/SGridPanel.h
--- a/Engine/Source/Runtime/Slate/Public/Widgets/Layout/SGridPanel.h
+++ b/Engine/Source/Runtime/Slate/Public/Widgets/Layout/SGridPanel.h
@@ -174,7 +174,7 @@ public:
 			FVector2D NudgeParam;
 
 			/** Notify that the slot was changed */
-			FORCEINLINE void NotifySlotChanged(bool bSlotLayerChanged = false)
+			inline void NotifySlotChanged(bool bSlotLayerChanged = false)
 			{
 				if ( Panel.IsValid() )
 				{
diff --git a/Engine/Source/Runtime/Slate/Public/Widgets/Layout/SScrollBox.h b/Engine/Source/Runtime/Slate/Public/Widgets/Layout/SScrollBox.h
--- a/Engine/Source/Runtime/Slate/Public/Widgets/Layout/SScrollBox.h
+++ b/Engine/Source/Runtime/Slate/Public/Widgets/Layout/SScrollBox.h
@@ -402,7 +402,7 @@ private:
 	void ConstructHorizontalLayout();
 
 	/** Gets the component of a vector in the direction of scrolling based on the Orientation property. */
-	FORCEINLINE float GetScrollComponentFromVector(FVector2f Vector) const
+	inline float GetScrollComponentFromVector(FVector2f Vector) const
 	{
 		return float(Orientation == Orient_Vertical ? Vector.Y : Vector.X);
 	}
diff --git a/Engine/Source/Runtime/Slate/Public/Widgets/Text/SlateEditableTextLayout.h b/Engine/Source/Runtime/Slate/Public/Widgets/Text/SlateEditableTextLayout.h
--- a/Engine/Source/Runtime/Slate/Public/Widgets/Text/SlateEditableTextLayout.h
+++ b/Engine/Source/Runtime/Slate/Public/Widgets/Text/SlateEditableTextLayout.h
@@ -474,13 +474,13 @@ private:
 
 		void CacheWindow();
 
-		FORCEINLINE void KillContext()
+		inline void KillContext()
 		{
 			OwnerLayout = nullptr;
 			bIsComposing = false;
 		}
 
-		FORCEINLINE FTextRange GetCompositionRange() const
+		inline FTextRange GetCompositionRange() const
 		{
 			return FTextRange(CompositionBeginIndex, CompositionBeginIndex + CompositionLength);
 		}
diff --git a/Engine/Source/Runtime/Slate/Public/Widgets/Text/SlateEditableTextTypes.h b/Engine/Source/Runtime/Slate/Public/Widgets/Text/SlateEditableTextTypes.h
--- a/Engine/Source/Runtime/Slate/Public/Widgets/Text/SlateEditableTextTypes.h
+++ b/Engine/Source/Runtime/Slate/Public/Widgets/Text/SlateEditableTextTypes.h
@@ -40,31 +40,31 @@ public:
 	}
 
 	/** Get the literal position of the cursor (note: this may not be the correct place to insert text to, use GetCursorInteractionLocation for that) */
-	FORCEINLINE FTextLocation GetCursorLocation() const
+	inline FTextLocation GetCursorLocation() const
 	{
 		return CursorPosition;
 	}
 
 	/** Get the alignment of the cursor */
-	FORCEINLINE ECursorAlignment GetCursorAlignment() const
+	inline ECursorAlignment GetCursorAlignment() const
 	{
 		return CursorAlignment;
 	}
 
 	/** Get the direction of the text under the cursor */
-	FORCEINLINE TextBiDi::ETextDirection GetCursorTextDirection() const
+	inline TextBiDi::ETextDirection GetCursorTextDirection() const
 	{
 		return CursorTextDirection;
 	}
 
 	/** Get the interaction position of the cursor (where to insert, delete, etc, text from/to) */
-	FORCEINLINE FTextLocation GetCursorInteractionLocation() const
+	inline FTextLocation GetCursorInteractionLocation() const
 	{
 		// If the cursor is right aligned, we treat it as if it were one character further along for insert/delete purposes
 		return FTextLocation(CursorPosition, (CursorAlignment == ECursorAlignment::Right) ? 1 : 0);
 	}
 
-	FORCEINLINE double GetLastCursorInteractionTime() const
+	inline double GetLastCursorInteractionTime() const
 	{
 		return LastCursorInteractionTime;
 	}
diff --git a/Engine/Source/Runtime/Slate/Public/Widgets/Text/SlateTextBlockLayout.h b/Engine/Source/Runtime/Slate/Public/Widgets/Text/SlateTextBlockLayout.h
--- a/Engine/Source/Runtime/Slate/Public/Widgets/Text/SlateTextBlockLayout.h
+++ b/Engine/Source/Runtime/Slate/Public/Widgets/Text/SlateTextBlockLayout.h
@@ -22,7 +22,7 @@ class FSlateTextBlockLayout
 public:
 	struct FWidgetDesiredSizeArgs
 	{
-		FORCEINLINE FWidgetDesiredSizeArgs(
+		inline FWidgetDesiredSizeArgs(
 			const FText& InText,
 			const FText& InHighlightText,
 			const float InWrapTextAt,
@@ -47,7 +47,7 @@ public:
 		{
 		}
 
-		FORCEINLINE FWidgetDesiredSizeArgs(
+		inline FWidgetDesiredSizeArgs(
 			const FText& InText,
 			const FText& InHighlightText,
 			const float InWrapTextAt,
@@ -87,7 +87,7 @@ public:
 		UE_DEPRECATED(5.0, "FWidgetArgs is deprecated. Upgrade to FWidgetDesiredSizeArgs instead.")
 		FWidgetArgs
 	{
-		FORCEINLINE FWidgetArgs(
+		inline FWidgetArgs(
 			const TAttribute<FText>& InText, 
 			const TAttribute<FText>& InHighlightText, 
 			const TAttribute<float>& InWrapTextAt, 
diff --git a/Engine/Source/Runtime/Slate/Public/Widgets/Views/STableViewBase.h b/Engine/Source/Runtime/Slate/Public/Widgets/Views/STableViewBase.h
--- a/Engine/Source/Runtime/Slate/Public/Widgets/Views/STableViewBase.h
+++ b/Engine/Source/Runtime/Slate/Public/Widgets/Views/STableViewBase.h
@@ -449,7 +449,7 @@ protected:
 	FOnTableViewScrolled OnTableViewScrolled;
 
 	/** Gets the component of a vector in the direction of scrolling based on the Orientation property. */
-	FORCEINLINE float GetScrollComponentFromVector(FVector2f Vector) const
+	inline float GetScrollComponentFromVector(FVector2f Vector) const
 	{
 		return float(Orientation == Orient_Vertical ? Vector.Y : Vector.X);
 	}
diff --git a/Engine/Source/Runtime/SlateCore/Public/Application/SlateApplicationBase.h b/Engine/Source/Runtime/SlateCore/Public/Application/SlateApplicationBase.h
--- a/Engine/Source/Runtime/SlateCore/Public/Application/SlateApplicationBase.h
+++ b/Engine/Source/Runtime/SlateCore/Public/Application/SlateApplicationBase.h
@@ -130,7 +130,7 @@ public:
 	 *
 	 * @return The Slate renderer.
 	 */
-	FORCEINLINE FSlateRenderer* GetRenderer() const
+	inline FSlateRenderer* GetRenderer() const
 	{
 		return Renderer.Get();
 	}
diff --git a/Engine/Source/Runtime/SlateCore/Public/Containers/ObservableArray.h b/Engine/Source/Runtime/SlateCore/Public/Containers/ObservableArray.h
--- a/Engine/Source/Runtime/SlateCore/Public/Containers/ObservableArray.h
+++ b/Engine/Source/Runtime/SlateCore/Public/Containers/ObservableArray.h
@@ -311,7 +311,7 @@ public:
 		return 1;
 	}
 	UE_ALLOWSHRINKING_BOOL_DEPRECATED("RemoveSingle")
-	FORCEINLINE SizeType RemoveSingle(const ElementType& Item, bool bAllowShrinking)
+	inline SizeType RemoveSingle(const ElementType& Item, bool bAllowShrinking)
 	{
 		return RemoveSingle(Item, bAllowShrinking ? EAllowShrinking::Yes : EAllowShrinking::No);
 	}
@@ -327,7 +327,7 @@ public:
 		return 1;
 	}
 	UE_ALLOWSHRINKING_BOOL_DEPRECATED("RemoveSingleSwap")
-	FORCEINLINE SizeType RemoveSingleSwap(const ElementType& Item, bool bAllowShrinking)
+	inline SizeType RemoveSingleSwap(const ElementType& Item, bool bAllowShrinking)
 	{
 		return RemoveSingleSwap(Item, bAllowShrinking ? EAllowShrinking::Yes : EAllowShrinking::No);
 	}
@@ -371,7 +371,7 @@ public:
 		}
 	}
 	UE_ALLOWSHRINKING_BOOL_DEPRECATED("RemoveAt")
-	FORCEINLINE void RemoveAt(SizeType Index, SizeType NumToRemove, bool bAllowShrinking)
+	inline void RemoveAt(SizeType Index, SizeType NumToRemove, bool bAllowShrinking)
 	{
 		RemoveAt(Index, NumToRemove, bAllowShrinking ? EAllowShrinking::Yes : EAllowShrinking::No);
 	}
@@ -419,7 +419,7 @@ public:
 		}
 	}
 	UE_ALLOWSHRINKING_BOOL_DEPRECATED("RemoveAtSwap")
-	FORCEINLINE void RemoveAtSwap(SizeType Index, SizeType NumToRemove, bool bAllowShrinking)
+	inline void RemoveAtSwap(SizeType Index, SizeType NumToRemove, bool bAllowShrinking)
 	{
 		RemoveAtSwap(Index, NumToRemove, bAllowShrinking ? EAllowShrinking::Yes : EAllowShrinking::No);
 	}
diff --git a/Engine/Source/Runtime/SlateCore/Public/Fonts/FontCache.h b/Engine/Source/Runtime/SlateCore/Public/Fonts/FontCache.h
--- a/Engine/Source/Runtime/SlateCore/Public/Fonts/FontCache.h
+++ b/Engine/Source/Runtime/SlateCore/Public/Fonts/FontCache.h
@@ -210,7 +210,7 @@ struct FShapedGlyphEntryKey
 public:
 	FShapedGlyphEntryKey(const FShapedGlyphFaceData& InFontFaceData, uint32 InGlyphIndex, const FFontOutlineSettings& InOutlineSettings);
 
-	FORCEINLINE bool operator==(const FShapedGlyphEntryKey& Other) const
+	inline bool operator==(const FShapedGlyphEntryKey& Other) const
 	{
 		return FontFace == Other.FontFace 
 			&& FontRenderSize == Other.FontRenderSize
@@ -221,7 +221,7 @@ public:
 			&& FontSkew == Other.FontSkew;
 	}
 
-	FORCEINLINE bool operator!=(const FShapedGlyphEntryKey& Other) const
+	inline bool operator!=(const FShapedGlyphEntryKey& Other) const
 	{
 		return !(*this == Other);
 	}
@@ -256,7 +256,7 @@ struct FSdfGlyphEntryKey
 public:
 	FSdfGlyphEntryKey(const TWeakPtr<FFreeTypeFace> InFontFace, uint32 InGlyphIndex, ESlateFontAtlasContentType InAtlasContentType, int32 InPpem, float InEmOuterSpread, float InEmInnerSpread);
 
-	FORCEINLINE bool operator==(const FSdfGlyphEntryKey& Other) const
+	inline bool operator==(const FSdfGlyphEntryKey& Other) const
 	{
 		return FontFace == Other.FontFace
 			&& GlyphIndex == Other.GlyphIndex
@@ -265,7 +265,7 @@ public:
 			&& SpreadCategory == Other.SpreadCategory;
 	}
 
-	FORCEINLINE bool operator!=(const FSdfGlyphEntryKey& Other) const
+	inline bool operator!=(const FSdfGlyphEntryKey& Other) const
 	{
 		return !(*this == Other);
 	}
@@ -298,16 +298,16 @@ private:
 struct FSdfGlyphTaskKey
 {
 public:
-	FORCEINLINE FSdfGlyphTaskKey(const FSdfGlyphEntryKey& InSdfGlyphEntryKey, float InEmOuterSpread, float InEmInnerSpread);
+	inline FSdfGlyphTaskKey(const FSdfGlyphEntryKey& InSdfGlyphEntryKey, float InEmOuterSpread, float InEmInnerSpread);
 
-	FORCEINLINE bool operator==(const FSdfGlyphTaskKey& Other) const
+	inline bool operator==(const FSdfGlyphTaskKey& Other) const
 	{
 		return SdfGlyphEntryKey == Other.SdfGlyphEntryKey
 			&& EmOuterSpread == Other.EmOuterSpread
 			&& EmInnerSpread == Other.EmInnerSpread;
 	}
 
-	FORCEINLINE bool operator!=(const FSdfGlyphTaskKey& Other) const
+	inline bool operator!=(const FSdfGlyphTaskKey& Other) const
 	{
 		return !(*this == Other);
 	}
@@ -549,29 +549,29 @@ private:
 			GlyphDataArray.SetNum(InSourceTextRange.TextLen);
 		}
 
-		FORCEINLINE int32 GetSourceTextStartIndex() const
+		inline int32 GetSourceTextStartIndex() const
 		{
 			return SourceTextRange.TextStart;
 		}
 
-		FORCEINLINE int32 GetSourceTextEndIndex() const
+		inline int32 GetSourceTextEndIndex() const
 		{
 			return SourceTextRange.TextStart + SourceTextRange.TextLen;
 		}
 
-		FORCEINLINE FSourceIndexToGlyphData* GetGlyphData(const int32 InSourceTextIndex)
+		inline FSourceIndexToGlyphData* GetGlyphData(const int32 InSourceTextIndex)
 		{
 			const int32 InternalIndex = InSourceTextIndex - SourceTextRange.TextStart;
 			return (GlyphDataArray.IsValidIndex(InternalIndex)) ? &GlyphDataArray[InternalIndex] : nullptr;
 		}
 
-		FORCEINLINE const FSourceIndexToGlyphData* GetGlyphData(const int32 InSourceTextIndex) const
+		inline const FSourceIndexToGlyphData* GetGlyphData(const int32 InSourceTextIndex) const
 		{
 			const int32 InternalIndex = InSourceTextIndex - SourceTextRange.TextStart;
 			return (GlyphDataArray.IsValidIndex(InternalIndex)) ? &GlyphDataArray[InternalIndex] : nullptr;
 		}
 
-		FORCEINLINE SIZE_T GetAllocatedSize() const
+		inline SIZE_T GetAllocatedSize() const
 		{
 			return GlyphDataArray.GetAllocatedSize();
 		}
diff --git a/Engine/Source/Runtime/SlateCore/Public/Fonts/FontTypes.h b/Engine/Source/Runtime/SlateCore/Public/Fonts/FontTypes.h
--- a/Engine/Source/Runtime/SlateCore/Public/Fonts/FontTypes.h
+++ b/Engine/Source/Runtime/SlateCore/Public/Fonts/FontTypes.h
@@ -32,17 +32,17 @@ public:
 		KeyHash = HashCombine(KeyHash, GetTypeHash(Scale));
 	}
 
-	FORCEINLINE const FSlateFontInfo& GetFontInfo() const
+	inline const FSlateFontInfo& GetFontInfo() const
 	{
 		return FontInfo;
 	}
 
-	FORCEINLINE float GetScale() const
+	inline float GetScale() const
 	{
 		return Scale;
 	}
 
-	FORCEINLINE const FFontOutlineSettings& GetFontOutlineSettings() const
+	inline const FFontOutlineSettings& GetFontOutlineSettings() const
 	{
 		return OutlineSettings;
 	}
@@ -77,17 +77,17 @@ struct FSlateFontKeyFuncs : BaseKeyFuncs<TPair<FSlateFontKey, ValueType>, FSlate
 	typedef typename Super::ElementInitType ElementInitType;
 	typedef typename Super::KeyInitType     KeyInitType;
 
-	FORCEINLINE static const FSlateFontKey& GetSetKey(ElementInitType Element)
+	inline static const FSlateFontKey& GetSetKey(ElementInitType Element)
 	{
 		return Element.Key;
 	}
 
-	FORCEINLINE static bool Matches(const FSlateFontKey& A, const FSlateFontKey& B)
+	inline static bool Matches(const FSlateFontKey& A, const FSlateFontKey& B)
 	{
 		return A.IsIdenticalToForCaching(B);
 	}
 
-	FORCEINLINE static uint32 GetKeyHash(const FSlateFontKey& Identifier)
+	inline static uint32 GetKeyHash(const FSlateFontKey& Identifier)
 	{
 		return GetTypeHash(Identifier);
 	}
diff --git a/Engine/Source/Runtime/SlateCore/Public/Input/HittestGrid.h b/Engine/Source/Runtime/SlateCore/Public/Input/HittestGrid.h
--- a/Engine/Source/Runtime/SlateCore/Public/Input/HittestGrid.h
+++ b/Engine/Source/Runtime/SlateCore/Public/Input/HittestGrid.h
@@ -270,14 +270,14 @@ private:
 	SLATECORE_API FIntPoint GetCellCoordinate(UE::Slate::FDeprecateVector2DParameter Position) const;
 
 	/** Access a cell at coordinates X, Y. Coordinates are row and column indexes. */
-	FORCEINLINE_DEBUGGABLE FCell& CellAt(const int32 X, const int32 Y)
+	FORCEINLINE FCell& CellAt(const int32 X, const int32 Y)
 	{
 		checkfSlow((Y*NumCells.X + X) < Cells.Num(), TEXT("HitTestGrid CellAt() failed: X= %d Y= %d NumCells.X= %d NumCells.Y= %d Cells.Num()= %d"), X, Y, NumCells.X, NumCells.Y, Cells.Num());
 		return Cells[Y*NumCells.X + X];
 	}
 
 	/** Access a cell at coordinates X, Y. Coordinates are row and column indexes. */
-	FORCEINLINE_DEBUGGABLE const FCell& CellAt( const int32 X, const int32 Y ) const
+	FORCEINLINE const FCell& CellAt( const int32 X, const int32 Y ) const
 	{
 		checkfSlow((Y*NumCells.X + X) < Cells.Num(), TEXT("HitTestGrid CellAt() failed: X= %d Y= %d NumCells.X= %d NumCells.Y= %d Cells.Num()= %d"), X, Y, NumCells.X, NumCells.Y, Cells.Num());
 		return Cells[Y*NumCells.X + X];
diff --git a/Engine/Source/Runtime/SlateCore/Public/Layout/ArrangedChildren.h b/Engine/Source/Runtime/SlateCore/Public/Layout/ArrangedChildren.h
--- a/Engine/Source/Runtime/SlateCore/Public/Layout/ArrangedChildren.h
+++ b/Engine/Source/Runtime/SlateCore/Public/Layout/ArrangedChildren.h
@@ -63,7 +63,7 @@ class FArrangedChildren
 	 *                             of layout or performance (i.e. prevent redundant call to Widget->GetVisibility())
 	 * @param InWidgetGeometry     The arranged widget (i.e. widget and its geometry)
 	 */
-	FORCEINLINE_DEBUGGABLE void AddWidget(EVisibility VisibilityOverride, const FArrangedWidget& InWidgetGeometry)
+	inline void AddWidget(EVisibility VisibilityOverride, const FArrangedWidget& InWidgetGeometry)
 	{
 		if ( Accepts(VisibilityOverride) )
 		{
@@ -71,7 +71,7 @@ class FArrangedChildren
 		}
 	}
 
-	FORCEINLINE_DEBUGGABLE void InsertWidget(EVisibility VisibilityOverride, const FArrangedWidget& InWidgetGeometry, int32 Index)
+	inline void InsertWidget(EVisibility VisibilityOverride, const FArrangedWidget& InWidgetGeometry, int32 Index)
 	{
 		if (Accepts(VisibilityOverride))
 		{
@@ -80,7 +80,7 @@ class FArrangedChildren
 	}
 
 	/** Reserves sufficient memory for the given number of arranged widgets. */
-	FORCEINLINE_DEBUGGABLE void Reserve(FArrangedWidgetArray::SizeType Number)
+	inline void Reserve(FArrangedWidgetArray::SizeType Number)
 	{
 		Array.Reserve(Number);
 	}
@@ -93,22 +93,22 @@ class FArrangedChildren
 
 	SLATECORE_API void InsertWidget(const FArrangedWidget& InWidgetGeometry, int32 Index);
 
-	FORCEINLINE_DEBUGGABLE void SetFilter(EVisibility InVisibility)
+	inline void SetFilter(EVisibility InVisibility)
 	{
 		VisibilityFilter = InVisibility;
 	}
 
-	FORCEINLINE_DEBUGGABLE EVisibility GetFilter() const
+	inline EVisibility GetFilter() const
 	{
 		return VisibilityFilter;
 	}
 
-	FORCEINLINE_DEBUGGABLE bool Accepts(EVisibility InVisibility) const
+	inline bool Accepts(EVisibility InVisibility) const
 	{
 		return EVisibility::DoesVisibilityPassFilter(InVisibility, VisibilityFilter);
 	}
 
-	FORCEINLINE_DEBUGGABLE bool Allows3DWidgets() const
+	inline bool Allows3DWidgets() const
 	{
 		return bAllow3DWidgets;
 	}
@@ -125,74 +125,74 @@ private:
 	bool bAllow3DWidgets;
 public:
 
-	FORCEINLINE FArrangedWidgetArray& GetInternalArray()
+	inline FArrangedWidgetArray& GetInternalArray()
 	{
 		return Array;
 	}
 
-	FORCEINLINE const FArrangedWidgetArray& GetInternalArray() const
+	inline const FArrangedWidgetArray& GetInternalArray() const
 	{
 		return Array;
 	}
 
-	FORCEINLINE int32 Num() const
+	inline int32 Num() const
 	{
 		return Array.Num();
 	}
 
-	FORCEINLINE const FArrangedWidget& operator[]( int32 Index ) const
+	inline const FArrangedWidget& operator[]( int32 Index ) const
 	{
 		return Array[Index];
 	}
 
-	FORCEINLINE FArrangedWidget& operator[](int32 Index)
+	inline FArrangedWidget& operator[](int32 Index)
 	{
 		return Array[Index];
 	}
 
-	FORCEINLINE const FArrangedWidget& Last() const
+	inline const FArrangedWidget& Last() const
 	{
 		return Array.Last();
 	}
 
-	FORCEINLINE FArrangedWidget& Last()
+	inline FArrangedWidget& Last()
 	{
 		return Array.Last();
 	}
 
-	FORCEINLINE int32 FindItemIndex(const FArrangedWidget& ItemToFind) const
+	inline int32 FindItemIndex(const FArrangedWidget& ItemToFind) const
 	{
 		return Array.Find(ItemToFind);
 	}
 
 	template<typename PredicateType>
-	FORCEINLINE int32 IndexOfByPredicate(const PredicateType& Pred) const
+	inline int32 IndexOfByPredicate(const PredicateType& Pred) const
 	{
 		return Array.IndexOfByPredicate( Pred );
 	}
 
 	template<typename PredicateType>
-	FORCEINLINE int32 FindLastByPredicate(const PredicateType& Pred) const
+	inline int32 FindLastByPredicate(const PredicateType& Pred) const
 	{
 		return Array.FindLastByPredicate(Pred);
 	}
 
-	FORCEINLINE void Remove(int32 Index, int32 Count=1)
+	inline void Remove(int32 Index, int32 Count=1)
 	{
 		Array.RemoveAt(Index, Count);
 	}
 
-	FORCEINLINE void Append(const FArrangedChildren& Source)
+	inline void Append(const FArrangedChildren& Source)
 	{
 		Array.Append( Source.Array );
 	}
 
-	FORCEINLINE void Empty()
+	inline void Empty()
 	{
 		Array.Reset();
 	}
 
-	FORCEINLINE bool IsValidIndex(int32 Index)
+	inline bool IsValidIndex(int32 Index)
 	{
 		return Array.IsValidIndex(Index);
 	}
diff --git a/Engine/Source/Runtime/SlateCore/Public/Layout/Children.h b/Engine/Source/Runtime/SlateCore/Public/Layout/Children.h
--- a/Engine/Source/Runtime/SlateCore/Public/Layout/Children.h
+++ b/Engine/Source/Runtime/SlateCore/Public/Layout/Children.h
@@ -774,7 +774,7 @@ public:
 	}
 
 	/** conversion to "bool" returning true if the iterator has not reached the last element. */
-	FORCEINLINE explicit operator bool() const
+	inline explicit operator bool() const
 	{
 		return Container.IsValidIndex(Index);
 	}
diff --git a/Engine/Source/Runtime/SlateCore/Public/Layout/ChildrenBase.h b/Engine/Source/Runtime/SlateCore/Public/Layout/ChildrenBase.h
--- a/Engine/Source/Runtime/SlateCore/Public/Layout/ChildrenBase.h
+++ b/Engine/Source/Runtime/SlateCore/Public/Layout/ChildrenBase.h
@@ -62,7 +62,7 @@ public:
 
 	/** Applies the predicate to all the widgets contained by the FChildren. */
 	template<typename Predicate>
-	FORCEINLINE void ForEachWidget(Predicate&& Pred)
+	inline void ForEachWidget(Predicate&& Pred)
 	{
 #if UE_WITH_SLATE_CHILDREN_DEBUGGING 
 		Debug_TestDestroyTag();
@@ -79,7 +79,7 @@ public:
 
 	/** Applies the predicate to all the widgets contained by the FChildren. */
 	template<typename Predicate>
-	FORCEINLINE void ForEachWidget(Predicate&& Pred) const
+	inline void ForEachWidget(Predicate&& Pred) const
 	{
 #if UE_WITH_SLATE_CHILDREN_DEBUGGING 
 		Debug_TestDestroyTag();
diff --git a/Engine/Source/Runtime/SlateCore/Public/Layout/Clipping.h b/Engine/Source/Runtime/SlateCore/Public/Layout/Clipping.h
--- a/Engine/Source/Runtime/SlateCore/Public/Layout/Clipping.h
+++ b/Engine/Source/Runtime/SlateCore/Public/Layout/Clipping.h
@@ -74,31 +74,31 @@ public:
 	FSlateClippingZone() {}
 
 	/**  */
-	FORCEINLINE bool GetShouldIntersectParent() const
+	inline bool GetShouldIntersectParent() const
 	{
 		return bIntersect;
 	}
 
 	/**  */
-	FORCEINLINE void SetShouldIntersectParent(bool bValue)
+	inline void SetShouldIntersectParent(bool bValue)
 	{
 		bIntersect = bValue;
 	}
 
 	/**  */
-	FORCEINLINE bool GetAlwaysClip() const
+	inline bool GetAlwaysClip() const
 	{
 		return bAlwaysClip;
 	}
 
 	/**  */
-	FORCEINLINE void SetAlwaysClip(bool bValue)
+	inline void SetAlwaysClip(bool bValue)
 	{
 		bAlwaysClip = bValue;
 	}
 
 	/** Is the clipping rect axis aligned?  If it is, we can use a much cheaper clipping solution. */
-	FORCEINLINE bool IsAxisAligned() const
+	inline bool IsAxisAligned() const
 	{
 		return bIsAxisAligned;
 	}
@@ -107,7 +107,7 @@ public:
 	 * Indicates if this clipping state has a zero size, aka is empty.  We only possibly report true for
 	 * scissor clipping zones.
 	 */
-	FORCEINLINE bool HasZeroArea() const
+	inline bool HasZeroArea() const
 	{
 		if (bIsAxisAligned)
 		{
@@ -171,7 +171,7 @@ private:
 	/** Should this clipping zone always clip, even if another zone wants to ignore intersection? */
 	uint8 bAlwaysClip : 1;
 
-	friend FORCEINLINE uint32 GetTypeHash(const FSlateClippingZone& Zone)
+	friend inline uint32 GetTypeHash(const FSlateClippingZone& Zone)
 	{
 		return Zone.ComputeHash();
 	}
@@ -213,16 +213,16 @@ public:
 
 #if WITH_SLATE_DEBUGGING
 	/** Set the state index that this clipping state originated from.  We just do this for debugging purposes. */
-	FORCEINLINE void SetDebuggingStateIndex(int32 InStateIndex) const
+	inline void SetDebuggingStateIndex(int32 InStateIndex) const
 	{
 		Debugging_StateIndex = InStateIndex;
 		Debugging_StateIndexFromFrame = GFrameNumber;
 	}
 #endif
 
-	FORCEINLINE bool GetAlwaysClip() const { return EnumHasAllFlags(Flags, EClippingFlags::AlwaysClip); }
+	inline bool GetAlwaysClip() const { return EnumHasAllFlags(Flags, EClippingFlags::AlwaysClip); }
 
-	FORCEINLINE bool GetShouldIntersectParent() const
+	inline bool GetShouldIntersectParent() const
 	{
 		if (GetClippingMethod() == EClippingMethod::Scissor)
 		{
@@ -246,7 +246,7 @@ public:
 	 * Gets the type of clipping that is required by this clipping state.  The simpler clipping is
 	 * scissor clipping, but that's only possible if the clipping rect is axis aligned.
 	 */
-	FORCEINLINE EClippingMethod GetClippingMethod() const
+	inline EClippingMethod GetClippingMethod() const
 	{
 		return ScissorRect.IsSet() ? EClippingMethod::Scissor : EClippingMethod::Stencil;
 	}
@@ -255,7 +255,7 @@ public:
 	 * Indicates if this clipping state has a zero size, aka is empty.  We only possibly report true for
 	 * scissor clipping zones.
 	 */
-	FORCEINLINE bool HasZeroArea() const
+	inline bool HasZeroArea() const
 	{
 		if (ScissorRect.IsSet())
 		{
diff --git a/Engine/Source/Runtime/SlateCore/Public/Layout/Geometry.h b/Engine/Source/Runtime/SlateCore/Public/Layout/Geometry.h
--- a/Engine/Source/Runtime/SlateCore/Public/Layout/Geometry.h
+++ b/Engine/Source/Runtime/SlateCore/Public/Layout/Geometry.h
@@ -194,7 +194,7 @@ public:
 	 * @param LayoutTransform	Layout transform of the geometry.
 	 * @return					The new root geometry
 	 */
-	FORCEINLINE_DEBUGGABLE static FGeometry MakeRoot(const UE::Slate::FDeprecateVector2DParameter& InLocalSize, const FSlateLayoutTransform& LayoutTransform)
+	inline static FGeometry MakeRoot(const UE::Slate::FDeprecateVector2DParameter& InLocalSize, const FSlateLayoutTransform& LayoutTransform)
 	{
 		return FGeometry(InLocalSize, LayoutTransform, FSlateLayoutTransform(), FSlateRenderTransform(), false);
 	}
@@ -206,7 +206,7 @@ public:
 	 * @param LayoutTransform	Layout transform of the geometry.
 	 * @return					The new root geometry
 	 */
-	FORCEINLINE_DEBUGGABLE static FGeometry MakeRoot(const UE::Slate::FDeprecateVector2DParameter& InLocalSize, const FSlateLayoutTransform& LayoutTransform, const FSlateRenderTransform& RenderTransform)
+	inline static FGeometry MakeRoot(const UE::Slate::FDeprecateVector2DParameter& InLocalSize, const FSlateLayoutTransform& LayoutTransform, const FSlateRenderTransform& RenderTransform)
 	{
 		return FGeometry(InLocalSize, LayoutTransform, FSlateLayoutTransform(), FSlateRenderTransform(), !RenderTransform.IsIdentity());
 	}
@@ -222,7 +222,7 @@ public:
 	 *
 	 * @return					The new child geometry.
 	 */
-	FORCEINLINE_DEBUGGABLE FGeometry MakeChild(const UE::Slate::FDeprecateVector2DParameter& InLocalSize, const FSlateLayoutTransform& LayoutTransform, const FSlateRenderTransform& RenderTransform, const UE::Slate::FDeprecateVector2DParameter& RenderTransformPivot) const
+	inline FGeometry MakeChild(const UE::Slate::FDeprecateVector2DParameter& InLocalSize, const FSlateLayoutTransform& LayoutTransform, const FSlateRenderTransform& RenderTransform, const UE::Slate::FDeprecateVector2DParameter& RenderTransformPivot) const
 	{
 		return FGeometry(InLocalSize, LayoutTransform, RenderTransform, RenderTransformPivot, GetAccumulatedLayoutTransform(), GetAccumulatedRenderTransform());
 	}
@@ -284,7 +284,7 @@ public:
 	 * @return				The new child geometry.
 	 */
 	UE_DEPRECATED(5.2, "Use FGeometry MakeChild(const FVector2f& InLocalSize, const FSlateLayoutTransform& InLayoutTransform) instead.")
-	FORCEINLINE_DEBUGGABLE FGeometry MakeChild(const FVector2D& ChildOffset, const FVector2D& InLocalSize, float ChildScale = 1.0f) const
+	inline FGeometry MakeChild(const FVector2D& ChildOffset, const FVector2D& InLocalSize, float ChildScale = 1.0f) const
 	{
 		// Since ChildOffset is given as a LocalSpaceOffset, we MUST convert this offset into the space of the parent to construct a valid layout transform.
 		// The extra TransformPoint below does this by converting the local offset to an offset in parent space.
@@ -391,7 +391,7 @@ public:
 	 * 
 	 * @return				The new paint geometry derived from this one.
 	 */
-	FORCEINLINE_DEBUGGABLE FPaintGeometry ToOffsetPaintGeometry(const UE::Slate::FDeprecateVector2DParameter& LocalOffset) const
+	inline FPaintGeometry ToOffsetPaintGeometry(const UE::Slate::FDeprecateVector2DParameter& LocalOffset) const
 	{
 		return ToPaintGeometry(FSlateLayoutTransform(UE::Slate::CastToVector2f(LocalOffset)));
 	}
@@ -404,7 +404,7 @@ public:
 	 * 
 	 * @return				The new paint geometry derived from this one.
 	 */
-	FORCEINLINE_DEBUGGABLE FPaintGeometry ToInflatedPaintGeometry(const UE::Slate::FDeprecateVector2DParameter& InflateAmount) const
+	inline FPaintGeometry ToInflatedPaintGeometry(const UE::Slate::FDeprecateVector2DParameter& InflateAmount) const
 	{
 		FVector2f InflateAmount2f = UE::Slate::CastToVector2f(InflateAmount);
 		// This essentially adds (or subtracts) a border around the widget. We scale the size then offset by the border amount.
@@ -418,7 +418,7 @@ public:
 	 * 
 	 * @return true if the provided location in absolute coordinates is within the bounds of this geometry. 
 	 */
-	FORCEINLINE_DEBUGGABLE bool IsUnderLocation(const UE::Slate::FDeprecateVector2DParameter& AbsoluteCoordinate) const
+	inline bool IsUnderLocation(const UE::Slate::FDeprecateVector2DParameter& AbsoluteCoordinate) const
 	{
 		// this render transform invert is a little expensive. We might consider caching it?
 		FSlateRotatedRect Rect = TransformRect(GetAccumulatedRenderTransform(), FSlateRotatedRect(FSlateRect(FVector2f(0.0f, 0.0f), FVector2f(Size))));
@@ -443,7 +443,7 @@ public:
 	 * 
 	 * @return  Absolute coordinates
 	 */
-	FORCEINLINE_DEBUGGABLE UE::Slate::FDeprecateVector2DResult LocalToAbsolute(UE::Slate::FDeprecateVector2DParameter LocalCoordinate) const
+	inline UE::Slate::FDeprecateVector2DResult LocalToAbsolute(UE::Slate::FDeprecateVector2DParameter LocalCoordinate) const
 	{
 		return UE::Slate::FDeprecateVector2DResult(TransformPoint(GetAccumulatedRenderTransform(), UE::Slate::CastToVector2f(LocalCoordinate)));
 	}
@@ -454,7 +454,7 @@ public:
 	 * and not have the window start on a half pixel, which can cause the contents to jitter in relation to each other as the tooltip 
 	 * or popup moves around.
 	 */
-	FORCEINLINE_DEBUGGABLE UE::Slate::FDeprecateVector2DResult LocalToRoundedLocal(UE::Slate::FDeprecateVector2DParameter LocalCoordinate) const
+	inline UE::Slate::FDeprecateVector2DResult LocalToRoundedLocal(UE::Slate::FDeprecateVector2DParameter LocalCoordinate) const
 	{
 		const FVector2f AbsoluteCoordinate = LocalToAbsolute(UE::Slate::CastToVector2f(LocalCoordinate));
 		const FVector2f AbsoluteCoordinateRounded = FVector2f(FMath::RoundToFloat(AbsoluteCoordinate.X), FMath::RoundToFloat(AbsoluteCoordinate.Y));
@@ -501,19 +501,19 @@ public:
 	 * Absolute coordinates could be either desktop or window space depending on what space the root of the widget hierarchy is in.
 	 *
 	 * @return the size of the geometry in absolute space */
-	FORCEINLINE_DEBUGGABLE UE::Slate::FDeprecateVector2DResult GetDrawSize() const
+	inline UE::Slate::FDeprecateVector2DResult GetDrawSize() const
 	{
 		return UE::Slate::FDeprecateVector2DResult(TransformVector(GetAccumulatedLayoutTransform(), FVector2f(Size)));
 	}
 
 	/** @return the size of the geometry in local space. */
-	FORCEINLINE UE::Slate::FDeprecateVector2DResult GetLocalSize() const { return Size; }
+	inline UE::Slate::FDeprecateVector2DResult GetLocalSize() const { return Size; }
 
 	/** @return the accumulated render transform. Shouldn't be needed in general. */
-	FORCEINLINE const FSlateRenderTransform& GetAccumulatedRenderTransform() const { return AccumulatedRenderTransform; }
+	inline const FSlateRenderTransform& GetAccumulatedRenderTransform() const { return AccumulatedRenderTransform; }
 
 	/** @return the accumulated layout transform. Shouldn't be needed in general. */
-	FORCEINLINE FSlateLayoutTransform GetAccumulatedLayoutTransform() const
+	inline FSlateLayoutTransform GetAccumulatedLayoutTransform() const
 	{
 		return FSlateLayoutTransform(Scale, AbsolutePosition);
 	}
@@ -537,7 +537,7 @@ public:
 	/**
 	 * Get the absolute position in render space.
 	 */
-	FORCEINLINE UE::Slate::FDeprecateVector2DResult GetAbsolutePosition() const
+	inline UE::Slate::FDeprecateVector2DResult GetAbsolutePosition() const
 	{
 		return UE::Slate::FDeprecateVector2DResult(AccumulatedRenderTransform.TransformPoint(FVector2f::ZeroVector));
 	}
@@ -545,7 +545,7 @@ public:
 	/**
 	 * Get the absolute size of the geometry in render space.
 	 */
-	FORCEINLINE UE::Slate::FDeprecateVector2DResult GetAbsoluteSize() const
+	inline UE::Slate::FDeprecateVector2DResult GetAbsoluteSize() const
 	{
 		return UE::Slate::FDeprecateVector2DResult(AccumulatedRenderTransform.TransformVector(FVector2f(GetLocalSize())));
 	}
@@ -557,7 +557,7 @@ public:
 	 *
 	 * Example: Say you wanted to know the center of the widget in absolute space, GetAbsolutePositionAtCoordinates(FVector2f(0.5f, 0.5f));
 	 */
-	FORCEINLINE UE::Slate::FDeprecateVector2DResult GetAbsolutePositionAtCoordinates(const UE::Slate::FDeprecateVector2DParameter& NormalCoordinates) const
+	inline UE::Slate::FDeprecateVector2DResult GetAbsolutePositionAtCoordinates(const UE::Slate::FDeprecateVector2DParameter& NormalCoordinates) const
 	{
 		return UE::Slate::FDeprecateVector2DResult(AccumulatedRenderTransform.TransformPoint(FVector2f(NormalCoordinates) * FVector2f(GetLocalSize())));
 	}
@@ -569,7 +569,7 @@ public:
 	 *
 	 * Example: Say you wanted to know the center of the widget in local space, GetLocalPositionAtCoordinates(FVector2f(0.5f, 0.5f));
 	 */
-	FORCEINLINE UE::Slate::FDeprecateVector2DResult GetLocalPositionAtCoordinates(const UE::Slate::FDeprecateVector2DParameter& NormalCoordinates) const
+	inline UE::Slate::FDeprecateVector2DResult GetLocalPositionAtCoordinates(const UE::Slate::FDeprecateVector2DParameter& NormalCoordinates) const
 	{
 		return UE::Slate::FDeprecateVector2DResult(Position + (FVector2f(NormalCoordinates) * GetLocalSize()));
 	}
diff --git a/Engine/Source/Runtime/SlateCore/Public/Layout/LayoutUtils.h b/Engine/Source/Runtime/SlateCore/Public/Layout/LayoutUtils.h
--- a/Engine/Source/Runtime/SlateCore/Public/Layout/LayoutUtils.h
+++ b/Engine/Source/Runtime/SlateCore/Public/Layout/LayoutUtils.h
@@ -245,7 +245,7 @@ public:
 	}
 
 	/** conversion to "bool" returning true if the iterator has not reached the last element. */
-	FORCEINLINE explicit operator bool() const
+	inline explicit operator bool() const
 	{
 		return Container.IsValidIndex(Index);
 	}
diff --git a/Engine/Source/Runtime/SlateCore/Public/Layout/SlateRect.h b/Engine/Source/Runtime/SlateCore/Public/Layout/SlateRect.h
--- a/Engine/Source/Runtime/SlateCore/Public/Layout/SlateRect.h
+++ b/Engine/Source/Runtime/SlateCore/Public/Layout/SlateRect.h
@@ -59,7 +59,7 @@ public:
 	/**
 	 * Determines if the rectangle has positive dimensions.
 	 */
-	FORCEINLINE bool IsValid() const
+	inline bool IsValid() const
 	{
 		return !(Left == -1 && Right == -1 && Bottom == -1 && Top == -1) && Right >= Left && Bottom >= Top;
 	}
@@ -67,7 +67,7 @@ public:
 	/**
 	 * @return true, if the rectangle has an effective size of 0.
 	 */
-	FORCEINLINE bool IsEmpty() const
+	inline bool IsEmpty() const
 	{
 		return GetArea() <= UE_SMALL_NUMBER;
 	}
@@ -77,11 +77,11 @@ public:
 	 *
 	 * @return The size as a vector.
 	 */
-	FORCEINLINE UE::Slate::FDeprecateVector2DResult GetSize() const
+	inline UE::Slate::FDeprecateVector2DResult GetSize() const
 	{
 		return UE::Slate::FDeprecateVector2DResult(GetSize2f());
 	}
-	FORCEINLINE FVector2f GetSize2f() const
+	inline FVector2f GetSize2f() const
 	{
 		return FVector2f(Right - Left, Bottom - Top);
 	}
@@ -89,7 +89,7 @@ public:
 	/**
 	 * @return the area of the rectangle
 	 */
-	FORCEINLINE float GetArea() const
+	inline float GetArea() const
 	{
 		return (Right - Left) * (Bottom - Top);
 	}
@@ -99,11 +99,11 @@ public:
 	 * 
 	 * @return The center point.
 	 */
-	FORCEINLINE UE::Slate::FDeprecateVector2DResult GetCenter() const
+	inline UE::Slate::FDeprecateVector2DResult GetCenter() const
 	{
 		return UE::Slate::FDeprecateVector2DResult(GetCenter2f());
 	}
-	FORCEINLINE FVector2f GetCenter2f() const
+	inline FVector2f GetCenter2f() const
 	{
 		return FVector2f(Left, Top) + GetSize2f() * 0.5f;
 	}
@@ -113,11 +113,11 @@ public:
 	 * 
 	 * @return The top-left position.
 	 */
-	FORCEINLINE UE::Slate::FDeprecateVector2DResult GetTopLeft() const
+	inline UE::Slate::FDeprecateVector2DResult GetTopLeft() const
 	{
 		return UE::Slate::FDeprecateVector2DResult(GetTopLeft2f());
 	}
-	FORCEINLINE FVector2f GetTopLeft2f() const
+	inline FVector2f GetTopLeft2f() const
 	{
 		return FVector2f(Left, Top);
 	}
@@ -127,11 +127,11 @@ public:
 	 *
 	 * @return The top-right position.
 	 */
-	FORCEINLINE UE::Slate::FDeprecateVector2DResult GetTopRight() const
+	inline UE::Slate::FDeprecateVector2DResult GetTopRight() const
 	{
 		return UE::Slate::FDeprecateVector2DResult(GetTopRight2f());
 	}
-	FORCEINLINE FVector2f GetTopRight2f() const
+	inline FVector2f GetTopRight2f() const
 	{
 		return FVector2f(Right, Top);
 	}
@@ -141,11 +141,11 @@ public:
 	 * 
 	 * @return The bottom-right position.
 	 */
-	FORCEINLINE UE::Slate::FDeprecateVector2DResult GetBottomRight() const
+	inline UE::Slate::FDeprecateVector2DResult GetBottomRight() const
 	{
 		return UE::Slate::FDeprecateVector2DResult(GetBottomRight2f());
 	}
-	FORCEINLINE FVector2f GetBottomRight2f() const
+	inline FVector2f GetBottomRight2f() const
 	{
 		return FVector2f(Right, Bottom);
 	}
@@ -155,11 +155,11 @@ public:
 	 * 
 	 * @return The bottom-left position.
 	 */
-	FORCEINLINE UE::Slate::FDeprecateVector2DResult GetBottomLeft() const
+	inline UE::Slate::FDeprecateVector2DResult GetBottomLeft() const
 	{
 		return UE::Slate::FDeprecateVector2DResult(GetBottomLeft2f());
 	}
-	FORCEINLINE FVector2f GetBottomLeft2f() const
+	inline FVector2f GetBottomLeft2f() const
 	{
 		return FVector2f(Left, Bottom);
 	}
@@ -171,7 +171,7 @@ public:
 	 *
 	 * @return An inset rectangle.
 	 */
-	FORCEINLINE FSlateRect InsetBy(const struct FMargin& InsetAmount) const
+	inline FSlateRect InsetBy(const struct FMargin& InsetAmount) const
 	{
 		return FSlateRect(Left + InsetAmount.Left, Top + InsetAmount.Top, Right - InsetAmount.Right, Bottom - InsetAmount.Bottom);
 	}
@@ -183,7 +183,7 @@ public:
 	 *
 	 * @return An extended rectangle.
 	 */
-	FORCEINLINE FSlateRect ExtendBy(const FMargin& ExtendAmount) const
+	inline FSlateRect ExtendBy(const FMargin& ExtendAmount) const
 	{
 		return FSlateRect(Left - ExtendAmount.Left, Top - ExtendAmount.Top, Right + ExtendAmount.Right, Bottom + ExtendAmount.Bottom);
 	}
@@ -195,7 +195,7 @@ public:
 	 *
 	 * @return An offset rectangle.
 	 */
-	FORCEINLINE FSlateRect OffsetBy( const UE::Slate::FDeprecateVector2DParameter& OffsetAmount ) const
+	inline FSlateRect OffsetBy( const UE::Slate::FDeprecateVector2DParameter& OffsetAmount ) const
 	{
 		return FSlateRect(GetTopLeft2f() + OffsetAmount, GetBottomRight2f() + OffsetAmount);
 	}
@@ -207,7 +207,7 @@ public:
 	 *
 	 * @return An scaled rectangle.
 	 */
-	FORCEINLINE FSlateRect ScaleBy(float ScaleBy) const
+	inline FSlateRect ScaleBy(float ScaleBy) const
 	{
 		const FVector2f Delta = GetSize() * 0.5f * ScaleBy;
 		return ExtendBy(FMargin(Delta));
@@ -220,7 +220,7 @@ public:
 	 *
 	 * @return	Rectangle that is big enough to fit both rectangles
 	 */
-	FORCEINLINE FSlateRect Expand( const FSlateRect& Other ) const
+	inline FSlateRect Expand( const FSlateRect& Other ) const
 	{
 		return FSlateRect( FMath::Min( Left, Other.Left ), FMath::Min( Top, Other.Top ), FMath::Max( Right, Other.Right ), FMath::Max( Bottom, Other.Bottom ) );
 	}
@@ -228,7 +228,7 @@ public:
 	/**
 	 * Rounds the Left, Top, Right and Bottom fields and returns a new FSlateRect with rounded components.
 	 */
-	FORCEINLINE FSlateRect Round() const
+	inline FSlateRect Round() const
 	{
 		return FSlateRect(
 			FMath::RoundToFloat(Left),
@@ -244,7 +244,7 @@ public:
 	 *
 	 * @return	Rectangle over intersection.
 	 */
-	FORCEINLINE FSlateRect IntersectionWith(const FSlateRect& Other) const
+	inline FSlateRect IntersectionWith(const FSlateRect& Other) const
 	{
 		bool bOverlapping;
 		return IntersectionWith(Other, bOverlapping);
@@ -280,7 +280,7 @@ public:
 	 * @param Point	The point to check
 	 * @return True if the point is inside the rectangle
 	 */
-	FORCEINLINE bool ContainsPoint( const UE::Slate::FDeprecateVector2DParameter& Point ) const
+	inline bool ContainsPoint( const UE::Slate::FDeprecateVector2DParameter& Point ) const
 	{
 		return Point.X >= Left && Point.X <= Right && Point.Y >= Top && Point.Y <= Bottom;
 	}
@@ -335,7 +335,7 @@ public:
 	}
 
 	/** Is rectangle B contained within rectangle A? */
-	FORCEINLINE static bool IsRectangleContained( const FSlateRect& A, const FSlateRect& B )
+	inline static bool IsRectangleContained( const FSlateRect& A, const FSlateRect& B )
 	{
 		return (A.Left <= B.Left) && (A.Right >= B.Right) && (A.Top <= B.Top) && (A.Bottom >= B.Bottom);
 	}
@@ -357,7 +357,7 @@ public:
 	SLATECORE_API bool InitFromString(const FString& InSourceString);
 
 	
-	friend FORCEINLINE uint32 GetTypeHash(const FSlateRect& Key)
+	friend inline uint32 GetTypeHash(const FSlateRect& Key)
 	{
 		uint32 Hash = 0;
 		Hash = HashCombine(Hash, GetTypeHash(Key.Left));
diff --git a/Engine/Source/Runtime/SlateCore/Public/Layout/Visibility.h b/Engine/Source/Runtime/SlateCore/Public/Layout/Visibility.h
--- a/Engine/Source/Runtime/SlateCore/Public/Layout/Visibility.h
+++ b/Engine/Source/Runtime/SlateCore/Public/Layout/Visibility.h
@@ -68,7 +68,7 @@ public:
 		return 0 != (Value & VIS_Visible);
 	}
 
-	static FORCEINLINE bool DoesVisibilityPassFilter(const EVisibility InVisibility, const EVisibility InVisibilityFilter)
+	static inline bool DoesVisibilityPassFilter(const EVisibility InVisibility, const EVisibility InVisibilityFilter)
 	{
 		return 0 != (InVisibility.Value & InVisibilityFilter.Value);
 	}
diff --git a/Engine/Source/Runtime/SlateCore/Public/Rendering/DrawElementCoreTypes.h b/Engine/Source/Runtime/SlateCore/Public/Rendering/DrawElementCoreTypes.h
--- a/Engine/Source/Runtime/SlateCore/Public/Rendering/DrawElementCoreTypes.h
+++ b/Engine/Source/Runtime/SlateCore/Public/Rendering/DrawElementCoreTypes.h
@@ -106,7 +106,7 @@ public:
 		* Assumes that the allocator is currently empty, i.e. memory may be allocated but any existing elements have already been destructed (if necessary).
 		* @param Other - The allocator to move the state from.  This allocator should be left in a valid empty state.
 		*/
-		FORCEINLINE void MoveToEmpty(ForAnyElementType& Other)
+		inline void MoveToEmpty(ForAnyElementType& Other)
 		{
 			Super::MoveToEmpty(Other);
 
diff --git a/Engine/Source/Runtime/SlateCore/Public/Rendering/DrawElementTypes.h b/Engine/Source/Runtime/SlateCore/Public/Rendering/DrawElementTypes.h
--- a/Engine/Source/Runtime/SlateCore/Public/Rendering/DrawElementTypes.h
+++ b/Engine/Source/Runtime/SlateCore/Public/Rendering/DrawElementTypes.h
@@ -121,7 +121,7 @@ public:
 	
 	SLATECORE_API static void MakeText( FSlateWindowElementList& ElementList, uint32 InLayer, const FPaintGeometry& PaintGeometry, const FString& InText, const FSlateFontInfo& InFontInfo, ESlateDrawEffect InDrawEffects = ESlateDrawEffect::None, const FLinearColor& InTint = FLinearColor::White );
 
-	FORCEINLINE static void MakeText(FSlateWindowElementList& ElementList, uint32 InLayer, const FPaintGeometry& PaintGeometry, const FText& InText, const FSlateFontInfo& InFontInfo, ESlateDrawEffect InDrawEffects = ESlateDrawEffect::None, const FLinearColor& InTint = FLinearColor::White)
+	inline static void MakeText(FSlateWindowElementList& ElementList, uint32 InLayer, const FPaintGeometry& PaintGeometry, const FText& InText, const FSlateFontInfo& InFontInfo, ESlateDrawEffect InDrawEffects = ESlateDrawEffect::None, const FLinearColor& InTint = FLinearColor::White)
 	{
 		MakeText(ElementList, InLayer, PaintGeometry, InText.ToString(), InFontInfo, InDrawEffects, InTint);
 	}
@@ -307,42 +307,42 @@ public:
 
 	FSlateDrawElement();
 
-	FORCEINLINE int32 GetLayer() const { return LayerId; }
+	inline int32 GetLayer() const { return LayerId; }
 
-	FORCEINLINE const FSlateRenderTransform& GetRenderTransform() const { return RenderTransform; }
-	FORCEINLINE void SetRenderTransform(const FSlateRenderTransform& InRenderTransform) { RenderTransform = InRenderTransform; }
+	inline const FSlateRenderTransform& GetRenderTransform() const { return RenderTransform; }
+	inline void SetRenderTransform(const FSlateRenderTransform& InRenderTransform) { RenderTransform = InRenderTransform; }
 
-	FORCEINLINE UE::Slate::FDeprecateVector2DResult GetPosition() const
+	inline UE::Slate::FDeprecateVector2DResult GetPosition() const
 	{
 		return UE::Slate::FDeprecateVector2DResult(Position);
 	}
 
-	FORCEINLINE void SetPosition(UE::Slate::FDeprecateVector2DParameter InPosition)
+	inline void SetPosition(UE::Slate::FDeprecateVector2DParameter InPosition)
 	{
 		Position = InPosition;
 	}
 
-	FORCEINLINE UE::Slate::FDeprecateVector2DResult GetLocalSize() const
+	inline UE::Slate::FDeprecateVector2DResult GetLocalSize() const
 	{
 		return UE::Slate::FDeprecateVector2DResult(LocalSize);
 	}
 
-	FORCEINLINE float GetScale() const { return Scale; }
-	FORCEINLINE ESlateDrawEffect GetDrawEffects() const { return DrawEffects; }
-	FORCEINLINE ESlateBatchDrawFlag GetBatchFlags() const { return BatchFlags; }
-	FORCEINLINE bool IsPixelSnapped() const { return !EnumHasAllFlags(DrawEffects, ESlateDrawEffect::NoPixelSnapping); }
+	inline float GetScale() const { return Scale; }
+	inline ESlateDrawEffect GetDrawEffects() const { return DrawEffects; }
+	inline ESlateBatchDrawFlag GetBatchFlags() const { return BatchFlags; }
+	inline bool IsPixelSnapped() const { return !EnumHasAllFlags(DrawEffects, ESlateDrawEffect::NoPixelSnapping); }
 
-	FORCEINLINE int32 GetPrecachedClippingIndex() const { return ClipStateHandle.GetPrecachedClipIndex(); }
-	FORCEINLINE void SetPrecachedClippingIndex(int32 InClippingIndex) { ClipStateHandle.SetPreCachedClipIndex(InClippingIndex); }
+	inline int32 GetPrecachedClippingIndex() const { return ClipStateHandle.GetPrecachedClipIndex(); }
+	inline void SetPrecachedClippingIndex(int32 InClippingIndex) { ClipStateHandle.SetPreCachedClipIndex(InClippingIndex); }
 
-	FORCEINLINE void SetCachedClippingState(const FSlateClippingState* CachedState) { ClipStateHandle.SetCachedClipState(CachedState); }
-	FORCEINLINE const FClipStateHandle& GetClippingHandle() const { return ClipStateHandle; }
-	FORCEINLINE const int8 GetSceneIndex() const { return SceneIndex; }
+	inline void SetCachedClippingState(const FSlateClippingState* CachedState) { ClipStateHandle.SetCachedClipState(CachedState); }
+	inline const FClipStateHandle& GetClippingHandle() const { return ClipStateHandle; }
+	inline const int8 GetSceneIndex() const { return SceneIndex; }
 
-	FORCEINLINE void SetIsCached(bool bInIsCached) { bIsCached = bInIsCached; }
-	FORCEINLINE bool IsCached() const { return bIsCached; }
+	inline void SetIsCached(bool bInIsCached) { bIsCached = bInIsCached; }
+	inline bool IsCached() const { return bIsCached; }
 
-	FORCEINLINE FSlateLayoutTransform GetInverseLayoutTransform() const
+	inline FSlateLayoutTransform GetInverseLayoutTransform() const
 	{
 		return Inverse(FSlateLayoutTransform(Scale, Position));
 	}
@@ -401,8 +401,8 @@ struct FSlateTintableElement
 {
 	FLinearColor Tint;
 
-	FORCEINLINE void SetTint(const FLinearColor& InTint) { Tint = InTint; }
-	FORCEINLINE FLinearColor GetTint() const { return Tint; }
+	inline void SetTint(const FLinearColor& InTint) { Tint = InTint; }
+	inline FLinearColor GetTint() const { return Tint; }
 };
 
 struct FSlateBoxElement : public FSlateDrawElement, public FSlateTintableElement
@@ -450,12 +450,12 @@ struct FSlateRoundedBoxElement : public FSlateBoxElement
 	FVector4f Radius;
 	float OutlineWeight;
 
-	FORCEINLINE void SetRadius(FVector4f InRadius) { Radius = InRadius; }
-	FORCEINLINE FVector4f GetRadius() const { return Radius; }
+	inline void SetRadius(FVector4f InRadius) { Radius = InRadius; }
+	inline FVector4f GetRadius() const { return Radius; }
 
-	FORCEINLINE void SetOutline(const FLinearColor& InOutlineColor, float InOutlineWeight) { OutlineColor = InOutlineColor; OutlineWeight = InOutlineWeight; }
-	FORCEINLINE FLinearColor GetOutlineColor() const { return OutlineColor; }
-	FORCEINLINE float GetOutlineWeight() const { return OutlineWeight; }
+	inline void SetOutline(const FLinearColor& InOutlineColor, float InOutlineWeight) { OutlineColor = InOutlineColor; OutlineWeight = InOutlineWeight; }
+	inline FLinearColor GetOutlineColor() const { return OutlineColor; }
+	inline float GetOutlineWeight() const { return OutlineWeight; }
 };
 
 struct FSlateTextElement : public FSlateDrawElement, public FSlateTintableElement
diff --git a/Engine/Source/Runtime/SlateCore/Public/Rendering/ElementBatcher.h b/Engine/Source/Runtime/SlateCore/Public/Rendering/ElementBatcher.h
--- a/Engine/Source/Runtime/SlateCore/Public/Rendering/ElementBatcher.h
+++ b/Engine/Source/Runtime/SlateCore/Public/Rendering/ElementBatcher.h
@@ -283,7 +283,7 @@ public:
 	 */
 	SLATECORE_API void ResetBatches();
 
-	FORCEINLINE FColor PackVertexColor(const FLinearColor& InLinearColor) const
+	inline FColor PackVertexColor(const FLinearColor& InLinearColor) const
 	{
 		//NOTE: Using pow(x,2) instead of a full sRGB conversion has been tried, but it ended up
 		// causing too much loss of data in the lower levels of black.
@@ -306,7 +306,7 @@ private:
 	 * @param InElementBatchReserver - Functor that reserves vertexes and indicies given an element range and list of elements
 	 */
 	template<typename ElementType, typename ElementAdder, typename ElementBatchParamCreator, typename ElementBatchReserver>
-	FORCEINLINE void GenerateIndexedVertexBatches(const FSlateDrawElementArray<ElementType>& DrawElements
+	inline void GenerateIndexedVertexBatches(const FSlateDrawElementArray<ElementType>& DrawElements
 		, ElementAdder&& InElementAdder
 		, ElementBatchParamCreator&& InElementBatchParamCreator
 		, ElementBatchReserver&& InElementBatchReserver);
diff --git a/Engine/Source/Runtime/SlateCore/Public/Rendering/RenderingCommon.h b/Engine/Source/Runtime/SlateCore/Public/Rendering/RenderingCommon.h
--- a/Engine/Source/Runtime/SlateCore/Public/Rendering/RenderingCommon.h
+++ b/Engine/Source/Runtime/SlateCore/Public/Rendering/RenderingCommon.h
@@ -243,7 +243,7 @@ struct FSlateVertex
 public:
 
 	template<ESlateVertexRounding Rounding>
-	FORCEINLINE static FSlateVertex Make(const FSlateRenderTransform& RenderTransform, const FVector2f InLocalPosition, const FVector2f InTexCoord, const FVector2f InTexCoord2, const FColor InColor, const FColor SecondaryColor = FColor())
+	inline static FSlateVertex Make(const FSlateRenderTransform& RenderTransform, const FVector2f InLocalPosition, const FVector2f InTexCoord, const FVector2f InTexCoord2, const FColor InColor, const FColor SecondaryColor = FColor())
 	{
 		FSlateVertex Vertex;
 		Vertex.TexCoords[0] = InTexCoord.X;
@@ -256,7 +256,7 @@ public:
 	}
 
 	template<ESlateVertexRounding Rounding>
-	FORCEINLINE static FSlateVertex Make(const FSlateRenderTransform& RenderTransform, const FVector2f InLocalPosition, const FVector2f InTexCoord, const FColor& InColor, const FColor SecondaryColor = FColor())
+	inline static FSlateVertex Make(const FSlateRenderTransform& RenderTransform, const FVector2f InLocalPosition, const FVector2f InTexCoord, const FColor& InColor, const FColor SecondaryColor = FColor())
 	{
 		FSlateVertex Vertex;
 		Vertex.TexCoords[0] = InTexCoord.X;
@@ -269,7 +269,7 @@ public:
 	}
 
 	template<ESlateVertexRounding Rounding>
-	FORCEINLINE static FSlateVertex Make(const FSlateRenderTransform& RenderTransform, const FVector2f InLocalPosition, const FVector4f InTexCoords, const FVector2f InMaterialTexCoords, const FColor InColor, const FColor SecondaryColor = FColor())
+	inline static FSlateVertex Make(const FSlateRenderTransform& RenderTransform, const FVector2f InLocalPosition, const FVector4f InTexCoords, const FVector2f InMaterialTexCoords, const FColor InColor, const FColor SecondaryColor = FColor())
 	{
 		FSlateVertex Vertex;
 		Vertex.TexCoords[0] = InTexCoords.X;
@@ -283,7 +283,7 @@ public:
 	}
 
 	template<ESlateVertexRounding Rounding>
-	FORCEINLINE static FSlateVertex Make(const FSlateRenderTransform& RenderTransform, const FVector2f InLocalPosition, const FVector2f InLocalSize, float Scale, const FVector4f InTexCoords, const FColor InColor, const FColor SecondaryColor = FColor())
+	inline static FSlateVertex Make(const FSlateRenderTransform& RenderTransform, const FVector2f InLocalPosition, const FVector2f InLocalSize, float Scale, const FVector4f InTexCoords, const FColor InColor, const FColor SecondaryColor = FColor())
 	{
 		FSlateVertex Vertex;
 		Vertex.TexCoords[0] = InTexCoords.X;
@@ -308,7 +308,7 @@ public:
 		return Vertex;
 	}
 	
-	FORCEINLINE static FSlateVertex Make(const FSlateRenderTransform& RenderTransform, const FVector2f InLocalPosition, const FVector2f InTexCoord, const FVector2f InTexCoord2, const FColor InColor, const FColor SecondaryColor = FColor(), const ESlateVertexRounding InRounding = ESlateVertexRounding::Disabled)
+	inline static FSlateVertex Make(const FSlateRenderTransform& RenderTransform, const FVector2f InLocalPosition, const FVector2f InTexCoord, const FVector2f InTexCoord2, const FColor InColor, const FColor SecondaryColor = FColor(), const ESlateVertexRounding InRounding = ESlateVertexRounding::Disabled)
 	{
 		return InRounding == ESlateVertexRounding::Enabled
 			? FSlateVertex::Make<ESlateVertexRounding::Enabled>(RenderTransform, InLocalPosition, InTexCoord, InTexCoord2, InColor, SecondaryColor)
@@ -316,7 +316,7 @@ public:
 	}
 
 	
-	FORCEINLINE static FSlateVertex Make(const FSlateRenderTransform& RenderTransform, const FVector2f InLocalPosition, const FVector2f InTexCoord, const FColor& InColor, const FColor SecondaryColor = FColor(), const ESlateVertexRounding InRounding = ESlateVertexRounding::Disabled)
+	inline static FSlateVertex Make(const FSlateRenderTransform& RenderTransform, const FVector2f InLocalPosition, const FVector2f InTexCoord, const FColor& InColor, const FColor SecondaryColor = FColor(), const ESlateVertexRounding InRounding = ESlateVertexRounding::Disabled)
 	{
 		return InRounding == ESlateVertexRounding::Enabled
 			? FSlateVertex::Make<ESlateVertexRounding::Enabled>(RenderTransform, InLocalPosition, InTexCoord, InColor, SecondaryColor)
@@ -324,7 +324,7 @@ public:
 	}
 
 	
-	FORCEINLINE static FSlateVertex Make(const FSlateRenderTransform& RenderTransform, const FVector2f InLocalPosition, const FVector4f InTexCoords, const FVector2f InMaterialTexCoords, const FColor InColor, const FColor SecondaryColor = FColor(), const ESlateVertexRounding InRounding = ESlateVertexRounding::Disabled)
+	inline static FSlateVertex Make(const FSlateRenderTransform& RenderTransform, const FVector2f InLocalPosition, const FVector4f InTexCoords, const FVector2f InMaterialTexCoords, const FColor InColor, const FColor SecondaryColor = FColor(), const ESlateVertexRounding InRounding = ESlateVertexRounding::Disabled)
 	{
 		return InRounding == ESlateVertexRounding::Enabled
 			? FSlateVertex::Make<ESlateVertexRounding::Enabled>(RenderTransform, InLocalPosition, InTexCoords, InMaterialTexCoords, InColor, SecondaryColor)
@@ -332,14 +332,14 @@ public:
 	}
 
 	
-	FORCEINLINE static FSlateVertex Make(const FSlateRenderTransform& RenderTransform, const FVector2f InLocalPosition, const FVector2f InLocalSize, float Scale, const FVector4f InTexCoords, const FColor InColor, const FColor SecondaryColor = FColor(), const ESlateVertexRounding InRounding = ESlateVertexRounding::Disabled)
+	inline static FSlateVertex Make(const FSlateRenderTransform& RenderTransform, const FVector2f InLocalPosition, const FVector2f InLocalSize, float Scale, const FVector4f InTexCoords, const FColor InColor, const FColor SecondaryColor = FColor(), const ESlateVertexRounding InRounding = ESlateVertexRounding::Disabled)
 	{
 		return InRounding == ESlateVertexRounding::Enabled
 			? FSlateVertex::Make<ESlateVertexRounding::Enabled>(RenderTransform, InLocalPosition, InLocalSize, Scale, InTexCoords, InColor, SecondaryColor)
 			: FSlateVertex::Make<ESlateVertexRounding::Disabled>(RenderTransform, InLocalPosition, InLocalSize, Scale, InTexCoords, InColor, SecondaryColor);
 	}
 
-	FORCEINLINE void SetTexCoords(const FVector4f InTexCoords)
+	inline void SetTexCoords(const FVector4f InTexCoords)
 	{
 		TexCoords[0] = InTexCoords.X;
 		TexCoords[1] = InTexCoords.Y;
@@ -347,7 +347,7 @@ public:
 		TexCoords[3] = InTexCoords.W;
 	}
 
-	FORCEINLINE void SetPosition(const FVector2f InPosition)
+	inline void SetPosition(const FVector2f InPosition)
 	{
 		Position = InPosition;
 	}
@@ -355,7 +355,7 @@ public:
 private:
 
 	template<ESlateVertexRounding Rounding>
-	FORCEINLINE void InitCommon(const FSlateRenderTransform& RenderTransform, const FVector2f InLocalPosition, const FColor InColor, const FColor InSecondaryColor)
+	inline void InitCommon(const FSlateRenderTransform& RenderTransform, const FVector2f InLocalPosition, const FColor InColor, const FColor InSecondaryColor)
 	{
 		Position = TransformPoint(RenderTransform, InLocalPosition);
 
diff --git a/Engine/Source/Runtime/SlateCore/Public/Rendering/ShaderResourceManager.h b/Engine/Source/Runtime/SlateCore/Public/Rendering/ShaderResourceManager.h
--- a/Engine/Source/Runtime/SlateCore/Public/Rendering/ShaderResourceManager.h
+++ b/Engine/Source/Runtime/SlateCore/Public/Rendering/ShaderResourceManager.h
@@ -30,7 +30,7 @@ struct FNewTextureInfo
 
 struct FCompareFNewTextureInfoByTextureSize
 {
-	FORCEINLINE bool operator()( const FNewTextureInfo& A, const FNewTextureInfo& B ) const
+	inline bool operator()( const FNewTextureInfo& A, const FNewTextureInfo& B ) const
 	{
 		return (B.TextureData->GetWidth()+B.TextureData->GetHeight()) < (A.TextureData->GetWidth()+A.TextureData->GetHeight());
 	}
diff --git a/Engine/Source/Runtime/SlateCore/Public/Rendering/SlateLayoutTransform.h b/Engine/Source/Runtime/SlateCore/Public/Rendering/SlateLayoutTransform.h
--- a/Engine/Source/Runtime/SlateCore/Public/Rendering/SlateLayoutTransform.h
+++ b/Engine/Source/Runtime/SlateCore/Public/Rendering/SlateLayoutTransform.h
@@ -58,12 +58,12 @@ public:
 		return Matrix;
 	}
 
-	FORCEINLINE UE::Slate::FDeprecateVector2DResult TransformPoint(const UE::Slate::FDeprecateVector2DParameter& Point) const
+	inline UE::Slate::FDeprecateVector2DResult TransformPoint(const UE::Slate::FDeprecateVector2DParameter& Point) const
 	{
 		return UE::Slate::FDeprecateVector2DResult(::TransformPoint(Translation, ::TransformPoint(Scale, FVector2f(Point))));
 	}
 
-	FORCEINLINE UE::Slate::FDeprecateVector2DResult TransformVector(const UE::Slate::FDeprecateVector2DParameter& Vector) const
+	inline UE::Slate::FDeprecateVector2DResult TransformVector(const UE::Slate::FDeprecateVector2DParameter& Vector) const
 	{
 		return UE::Slate::FDeprecateVector2DResult(::TransformVector(Translation, ::TransformVector(Scale, FVector2f(Vector))));
 	}
diff --git a/Engine/Source/Runtime/SlateCore/Public/SlotBase.h b/Engine/Source/Runtime/SlateCore/Public/SlotBase.h
--- a/Engine/Source/Runtime/SlateCore/Public/SlotBase.h
+++ b/Engine/Source/Runtime/SlateCore/Public/SlotBase.h
@@ -44,13 +44,13 @@ public:
 	SLATECORE_API void SetOwner(const FChildren& Children);
 
 	/** Attach the child widget the slot now owns. */
-	FORCEINLINE_DEBUGGABLE void AttachWidget(TSharedRef<SWidget>&& InWidget)
+	inline void AttachWidget(TSharedRef<SWidget>&& InWidget)
 	{
 		DetatchParentFromContent();
 		Widget = MoveTemp(InWidget);
 		AfterContentOrOwnerAssigned();
 	}
-	FORCEINLINE_DEBUGGABLE void AttachWidget( const TSharedRef<SWidget>& InWidget )
+	inline void AttachWidget( const TSharedRef<SWidget>& InWidget )
 	{
 		DetatchParentFromContent();
 		Widget = InWidget;
@@ -62,7 +62,7 @@ public:
 	 * There will always be a widget in the slot; sometimes it is
 	 * the SNullWidget instance.
 	 */
-	FORCEINLINE_DEBUGGABLE const TSharedRef<SWidget>& GetWidget() const
+	inline const TSharedRef<SWidget>& GetWidget() const
 	{
 		return Widget;
 	}
diff --git a/Engine/Source/Runtime/SlateCore/Public/Textures/SlateShaderResource.h b/Engine/Source/Runtime/SlateCore/Public/Textures/SlateShaderResource.h
--- a/Engine/Source/Runtime/SlateCore/Public/Textures/SlateShaderResource.h
+++ b/Engine/Source/Runtime/SlateCore/Public/Textures/SlateShaderResource.h
@@ -87,7 +87,7 @@ public:
 	/** Virtual destructor. */
 	virtual ~FSlateShaderResource() { DestroyState = 0x84; }
 #else
-	FORCEINLINE void CheckForStaleResources() const { }
+	inline void CheckForStaleResources() const { }
 	bool Debug_IsDestroyed() const { return false; }
 
 	/** Virtual destructor. */
diff --git a/Engine/Source/Runtime/SlateCore/Public/Textures/TextureAtlas.h b/Engine/Source/Runtime/SlateCore/Public/Textures/TextureAtlas.h
--- a/Engine/Source/Runtime/SlateCore/Public/Textures/TextureAtlas.h
+++ b/Engine/Source/Runtime/SlateCore/Public/Textures/TextureAtlas.h
@@ -207,7 +207,7 @@ protected:
 
 private:
 	/** Returns the amount of padding needed for the current padding style */
-	FORCEINLINE uint8 GetPaddingAmount() const
+	inline uint8 GetPaddingAmount() const
 	{
 		return (PaddingStyle == ESlateTextureAtlasPaddingStyle::NoPadding) ? 0 : 1;
 	}
diff --git a/Engine/Source/Runtime/SlateCore/Public/Types/PaintArgs.h b/Engine/Source/Runtime/SlateCore/Public/Types/PaintArgs.h
--- a/Engine/Source/Runtime/SlateCore/Public/Types/PaintArgs.h
+++ b/Engine/Source/Runtime/SlateCore/Public/Types/PaintArgs.h
@@ -27,7 +27,7 @@ public:
 	SLATECORE_API FPaintArgs(const SWidget* PaintParent, FHittestGrid& InRootHittestGrid, FHittestGrid& InCurrentHitTestGrid, UE::Slate::FDeprecateVector2DParameter InWindowOffset, double InCurrentTime, float InDeltaTime);
 	SLATECORE_API FPaintArgs(const SWidget* PaintParent, FHittestGrid& InRootHittestGrid, UE::Slate::FDeprecateVector2DParameter InWindowOffset, double InCurrentTime, float InDeltaTime);
 
-	[[nodiscard]] FORCEINLINE_DEBUGGABLE FPaintArgs WithNewParent(const SWidget* PaintParent) const
+	[[nodiscard]] inline FPaintArgs WithNewParent(const SWidget* PaintParent) const
 	{
 		FPaintArgs Args(*this);
 		Args.PaintParentPtr = PaintParent;
@@ -35,7 +35,7 @@ public:
 		return Args;
 	}
 
-	[[nodiscard]] FORCEINLINE_DEBUGGABLE FPaintArgs WithNewHitTestGrid(FHittestGrid& NewHitTestGrid) const
+	[[nodiscard]] inline FPaintArgs WithNewHitTestGrid(FHittestGrid& NewHitTestGrid) const
 	{
 		FPaintArgs NewArgs(PaintParentPtr, RootGrid, NewHitTestGrid, WindowOffset, CurrentTime, DeltaTime);
 		return NewArgs;
diff --git a/Engine/Source/Runtime/SlateCore/Public/Types/SlateBox2.h b/Engine/Source/Runtime/SlateCore/Public/Types/SlateBox2.h
--- a/Engine/Source/Runtime/SlateCore/Public/Types/SlateBox2.h
+++ b/Engine/Source/Runtime/SlateCore/Public/Types/SlateBox2.h
@@ -7,7 +7,7 @@
 namespace UE::Slate
 {
 	
-FORCEINLINE static FBox2f CastToBox2f(const FBox2d& InValue)
+inline static FBox2f CastToBox2f(const FBox2d& InValue)
 {
 	return FBox2f(InValue);
 }
diff --git a/Engine/Source/Runtime/SlateCore/Public/Widgets/SWidget.h b/Engine/Source/Runtime/SlateCore/Public/Widgets/SWidget.h
--- a/Engine/Source/Runtime/SlateCore/Public/Widgets/SWidget.h
+++ b/Engine/Source/Runtime/SlateCore/Public/Widgets/SWidget.h
@@ -679,10 +679,10 @@ public:
 	/**  */
 	virtual bool ValidatePathToChild(SWidget* InChild) { return true; }
 
-	FORCEINLINE bool IsParentValid() const { return ParentWidgetPtr.IsValid(); }
-	FORCEINLINE TSharedPtr<SWidget> GetParentWidget() const { return ParentWidgetPtr.Pin(); }
+	inline bool IsParentValid() const { return ParentWidgetPtr.IsValid(); }
+	inline TSharedPtr<SWidget> GetParentWidget() const { return ParentWidgetPtr.Pin(); }
 
-	FORCEINLINE TSharedPtr<SWidget> Advanced_GetPaintParentWidget() const { return PersistentState.PaintParent.Pin(); }
+	inline TSharedPtr<SWidget> Advanced_GetPaintParentWidget() const { return PersistentState.PaintParent.Pin(); }
 
 	/**
 	 * Calculates what if any clipping state changes need to happen when drawing this widget.
@@ -788,7 +788,7 @@ private:
 
 public:
 
-	FORCEINLINE TStatId GetStatID(bool bForDeferredUse = false) const
+	inline TStatId GetStatID(bool bForDeferredUse = false) const
 	{
 #if STATS
 		// this is done to avoid even registering stats for a disabled group (unless we plan on using it later)
@@ -928,7 +928,7 @@ public:
 	}
 
 	/** @return Whether or not this widget is enabled */
-	FORCEINLINE bool IsEnabled() const
+	inline bool IsEnabled() const
 	{
 		return EnabledStateAttribute.Get();
 	}
@@ -989,7 +989,7 @@ public:
 	 * @return is this widget visible, hidden or collapsed.
 	 * @note this widget can be visible but if a parent is hidden or collapsed, it would not show on screen.
 	 */
-	FORCEINLINE EVisibility GetVisibility() const { return VisibilityAttribute.Get(); }
+	inline EVisibility GetVisibility() const { return VisibilityAttribute.Get(); }
 
 	/** @param InVisibility  should this widget be */
 	SLATECORE_API virtual void SetVisibility(TAttribute<EVisibility> InVisibility);
@@ -1059,20 +1059,20 @@ public:
 #endif
 
 	/** Whether or not a widget is volatile and will update every frame without being invalidated */
-	FORCEINLINE bool IsVolatile() const { return bCachedVolatile; }
+	inline bool IsVolatile() const { return bCachedVolatile; }
 
 	/**
 	 * This widget is volatile because its parent or some ancestor is volatile
 	 * @note only valid if the widget is contained by an InvalidationRoot (the proxy is valid).
 	 */
-	FORCEINLINE bool IsVolatileIndirectly() const { return bInheritedVolatility; }
+	inline bool IsVolatileIndirectly() const { return bInheritedVolatility; }
 
 	/**
 	 * Should this widget always appear as volatile for any layout caching host widget.  A volatile
 	 * widget's geometry and layout data will never be cached, and neither will any children.
 	 * @param bForce should we force the widget to be volatile?
 	 */
-	FORCEINLINE void ForceVolatile(bool bForce)
+	inline void ForceVolatile(bool bForce)
 	{
 		if (bForceVolatile != bForce)
 		{
@@ -1081,7 +1081,7 @@ public:
 		}
 	}
 
-	FORCEINLINE bool ShouldInvalidatePrepassDueToVolatility() { return bVolatilityAlwaysInvalidatesPrepass; }
+	inline bool ShouldInvalidatePrepassDueToVolatility() { return bVolatilityAlwaysInvalidatesPrepass; }
 
 	/**
 	 * Invalidates the widget from the view of a layout caching widget that may own this widget.
@@ -1093,7 +1093,7 @@ public:
 	 * Recalculates volatility of the widget and caches the result.  Should be called any time 
 	 * anything examined by your implementation of ComputeVolatility is changed.
 	 */
-	FORCEINLINE void CacheVolatility()
+	inline void CacheVolatility()
 	{
 		bCachedVolatile = bForceVolatile || ComputeVolatility();
 	}
@@ -1114,7 +1114,7 @@ protected:
 	 */
 	SLATECORE_API bool IsChildWidgetCulled(const FSlateRect& MyCullingRect, const FArrangedWidget& ArrangedChild) const;
 #else
-	FORCEINLINE bool IsChildWidgetCulled(const FSlateRect&, const FArrangedWidget&) const { return false; }
+	inline bool IsChildWidgetCulled(const FSlateRect&, const FArrangedWidget&) const { return false; }
 #endif
 
 protected:
@@ -1129,7 +1129,7 @@ protected:
 	/**
 	 * Recalculates and caches volatility and returns 'true' if the volatility changed.
 	 */
-	FORCEINLINE bool Advanced_InvalidateVolatility()
+	inline bool Advanced_InvalidateVolatility()
 	{
 		const bool bWasDirectlyVolatile = IsVolatile();
 		CacheVolatility();
@@ -1138,13 +1138,13 @@ protected:
 public:
 
 	/** @return the render opacity of the widget. */
-	FORCEINLINE float GetRenderOpacity() const
+	inline float GetRenderOpacity() const
 	{
 		return RenderOpacity;
 	}
 
 	/** @param InOpacity The opacity of the widget during rendering. */
-	FORCEINLINE void SetRenderOpacity(float InRenderOpacity)
+	inline void SetRenderOpacity(float InRenderOpacity)
 	{
 		if(RenderOpacity != InRenderOpacity)
 		{
@@ -1153,18 +1153,18 @@ public:
 		}
 	}
 
-	FORCEINLINE void SetTag(FName InTag)
+	inline void SetTag(FName InTag)
 	{
 		Tag = InTag;
 	}
 
 	/** @return the render transform of the widget. */
-	FORCEINLINE const TOptional<FSlateRenderTransform>& GetRenderTransform() const
+	inline const TOptional<FSlateRenderTransform>& GetRenderTransform() const
 	{
 		return RenderTransformAttribute.Get();
 	}
 
-	FORCEINLINE TOptional<FSlateRenderTransform> GetRenderTransformWithRespectToFlowDirection() const
+	inline TOptional<FSlateRenderTransform> GetRenderTransformWithRespectToFlowDirection() const
 	{
 		if (LIKELY(GSlateFlowDirection == EFlowDirection::LeftToRight))
 		{
@@ -1183,7 +1183,7 @@ public:
 		}
 	}
 
-	FORCEINLINE UE::Slate::FDeprecateVector2DResult GetRenderTransformPivotWithRespectToFlowDirection() const
+	inline UE::Slate::FDeprecateVector2DResult GetRenderTransformPivotWithRespectToFlowDirection() const
 	{
 		FVector2f TransformPivot = UE::Slate::CastToVector2f(RenderTransformPivotAttribute.Get());
 		if (LIKELY(GSlateFlowDirection == EFlowDirection::LeftToRight))
@@ -1199,20 +1199,20 @@ public:
 	}
 
 	/** @param InTransform the render transform to set for the widget (transforms from widget's local space). TOptional<> to allow code to skip expensive overhead if there is no render transform applied. */
-	FORCEINLINE void SetRenderTransform(TAttribute<TOptional<FSlateRenderTransform>> InTransform)
+	inline void SetRenderTransform(TAttribute<TOptional<FSlateRenderTransform>> InTransform)
 	{
 		RenderTransformAttribute.Assign(*this, MoveTemp(InTransform));
 	}
 
 	/** @return the pivot point of the render transform. */
-	FORCEINLINE UE::Slate::FDeprecateVector2DResult GetRenderTransformPivot() const
+	inline UE::Slate::FDeprecateVector2DResult GetRenderTransformPivot() const
 	{
 		FVector2f TransformPivot = UE::Slate::CastToVector2f(RenderTransformPivotAttribute.Get());
 		return TransformPivot;
 	}
 
 	/** @param InTransformPivot Sets the pivot point of the widget's render transform (in normalized local space). */
-	FORCEINLINE void SetRenderTransformPivot(TAttribute<FVector2D> InTransformPivot)
+	inline void SetRenderTransformPivot(TAttribute<FVector2D> InTransformPivot)
 	{
 		RenderTransformPivotAttribute.Assign(*this, MoveTemp(InTransformPivot));
 	}
@@ -1223,7 +1223,7 @@ public:
 	SLATECORE_API void SetClipping(EWidgetClipping InClipping);
 
 	/** @return The current clipping rules for this widget. */
-	FORCEINLINE EWidgetClipping GetClipping() const
+	inline EWidgetClipping GetClipping() const
 	{
 		return Clipping;
 	}
@@ -1234,7 +1234,7 @@ public:
 	SLATECORE_API void SetPixelSnapping(EWidgetPixelSnapping InPixelSnappingMethod);
 
 	/** @return The current pixel snapping rules for this widget. */
-	FORCEINLINE EWidgetPixelSnapping GetPixelSnapping() const
+	inline EWidgetPixelSnapping GetPixelSnapping() const
 	{
 		return PixelSnappingMethod;
 	}
@@ -1243,7 +1243,7 @@ public:
 	 * Sets an additional culling padding that is added to a widget to give more leeway when culling widgets.  Useful if 
 	 * several child widgets have rendering beyond their bounds.
 	 */
-	FORCEINLINE void SetCullingBoundsExtension(const FMargin& InCullingBoundsExtension)
+	inline void SetCullingBoundsExtension(const FMargin& InCullingBoundsExtension)
 	{
 		if (CullingBoundsExtension != InCullingBoundsExtension)
 		{
@@ -1254,7 +1254,7 @@ public:
 	}
 
 	/** @return CullingBoundsExtension */
-	FORCEINLINE FMargin GetCullingBoundsExtension() const
+	inline FMargin GetCullingBoundsExtension() const
 	{
 		return CullingBoundsExtension;
 	}
@@ -1896,7 +1896,7 @@ private:
 // FGeometry Arranged Widget Inlined Functions
 //=================================================================
 
-FORCEINLINE_DEBUGGABLE FArrangedWidget FGeometry::MakeChild(const TSharedRef<SWidget>& ChildWidget, const UE::Slate::FDeprecateVector2DParameter& InLocalSize, const FSlateLayoutTransform& LayoutTransform) const
+inline FArrangedWidget FGeometry::MakeChild(const TSharedRef<SWidget>& ChildWidget, const UE::Slate::FDeprecateVector2DParameter& InLocalSize, const FSlateLayoutTransform& LayoutTransform) const
 {
 	// If there is no render transform set, use the simpler MakeChild call that doesn't bother concatenating the render transforms.
 	// This saves a significant amount of overhead since every widget does this, and most children don't have a render transform.
@@ -1912,12 +1912,12 @@ FORCEINLINE_DEBUGGABLE FArrangedWidget FGeometry::MakeChild(const TSharedRef<SWi
 	}
 }
 
-FORCEINLINE_DEBUGGABLE FArrangedWidget FGeometry::MakeChild(const TSharedRef<SWidget>& ChildWidget, const FLayoutGeometry& LayoutGeometry) const
+inline FArrangedWidget FGeometry::MakeChild(const TSharedRef<SWidget>& ChildWidget, const FLayoutGeometry& LayoutGeometry) const
 {
 	return MakeChild(ChildWidget, FVector2f(LayoutGeometry.GetSizeInLocalSpace()), LayoutGeometry.GetLocalToParentTransform());
 }
 
-FORCEINLINE_DEBUGGABLE FArrangedWidget FGeometry::MakeChild(const TSharedRef<SWidget>& ChildWidget, const UE::Slate::FDeprecateVector2DParameter& ChildOffset, const UE::Slate::FDeprecateVector2DParameter& InLocalSize, float ChildScale) const
+inline FArrangedWidget FGeometry::MakeChild(const TSharedRef<SWidget>& ChildWidget, const UE::Slate::FDeprecateVector2DParameter& ChildOffset, const UE::Slate::FDeprecateVector2DParameter& InLocalSize, float ChildScale) const
 {
 	// Since ChildOffset is given as a LocalSpaceOffset, we MUST convert this offset into the space of the parent to construct a valid layout transform.
 	// The extra TransformPoint below does this by converting the local offset to an offset in parent space.
diff --git a/Engine/Source/Runtime/SlateCore/Public/Widgets/SWidgetUtils.h b/Engine/Source/Runtime/SlateCore/Public/Widgets/SWidgetUtils.h
--- a/Engine/Source/Runtime/SlateCore/Public/Widgets/SWidgetUtils.h
+++ b/Engine/Source/Runtime/SlateCore/Public/Widgets/SWidgetUtils.h
@@ -13,7 +13,7 @@ struct FScopeCycleCounterSWidget : public FCycleCounter
 	/**
 	 * Constructor, starts timing
 	 */
-	FORCEINLINE FScopeCycleCounterSWidget(const SWidget* Widget)
+	inline FScopeCycleCounterSWidget(const SWidget* Widget)
 	{
 		if (Widget)
 		{
@@ -42,7 +42,7 @@ struct FScopeCycleCounterSWidget : public FCycleCounter
 	/**
 	 * Updates the stat with the time spent
 	 */
-	FORCEINLINE ~FScopeCycleCounterSWidget()
+	inline ~FScopeCycleCounterSWidget()
 	{
 		Stop();
 	}
@@ -60,7 +60,7 @@ struct FScopeCycleCounterSWidget
 	bool bPop = false;
 #endif
 
-	FORCEINLINE FScopeCycleCounterSWidget(const SWidget* Widget)
+	inline FScopeCycleCounterSWidget(const SWidget* Widget)
 		: ScopeCycleCounter(Widget ? Widget->GetStatID().StatString : nullptr)
 	{
 #if CPUPROFILERTRACE_ENABLED
@@ -76,7 +76,7 @@ struct FScopeCycleCounterSWidget
 #endif
 	}
 
-	FORCEINLINE ~FScopeCycleCounterSWidget()
+	inline ~FScopeCycleCounterSWidget()
 	{
 #if CPUPROFILERTRACE_ENABLED
 		if (bPop)
diff --git a/Engine/Source/Runtime/SlateCore/Public/Widgets/SWindow.h b/Engine/Source/Runtime/SlateCore/Public/Widgets/SWindow.h
--- a/Engine/Source/Runtime/SlateCore/Public/Widgets/SWindow.h
+++ b/Engine/Source/Runtime/SlateCore/Public/Widgets/SWindow.h
@@ -1430,9 +1430,9 @@ private:
 
 struct FScopedSwitchWorldHack
 {
-	FORCEINLINE FScopedSwitchWorldHack(const FWidgetPath& WidgetPath) { }
-	FORCEINLINE FScopedSwitchWorldHack(TSharedPtr<SWindow> InWindow) { }
-	FORCEINLINE ~FScopedSwitchWorldHack() { }
+	inline FScopedSwitchWorldHack(const FWidgetPath& WidgetPath) { }
+	inline FScopedSwitchWorldHack(TSharedPtr<SWindow> InWindow) { }
+	inline ~FScopedSwitchWorldHack() { }
 };
 
 #endif
