[Slate] Change check to checkSlow. It showed in the profiler, and they are less important now that the system has shipped for many versions.

#rnx

diff --git a/Engine/Source/Runtime/SlateCore/Private/Types/SlateAttributeMetaData.cpp b/Engine/Source/Runtime/SlateCore/Private/Types/SlateAttributeMetaData.cpp
--- a/Engine/Source/Runtime/SlateCore/Private/Types/SlateAttributeMetaData.cpp
+++ b/Engine/Source/Runtime/SlateCore/Private/Types/SlateAttributeMetaData.cpp
@@ -76,9 +76,9 @@ FSlateAttributeMetaData* FSlateAttributeMetaData::FindMetaData(const SWidget& Ow
 {
 	if (OwningWidget.HasRegisteredSlateAttribute())
 	{
-		check(OwningWidget.MetaData.Num() > 0);
+		checkSlow(OwningWidget.MetaData.Num() > 0);
 		const TSharedRef<ISlateMetaData>& SlateMetaData = OwningWidget.MetaData[0];
-		check(SlateMetaData->IsOfType<FSlateAttributeMetaData>());
+		checkSlow(SlateMetaData->IsOfType<FSlateAttributeMetaData>());
 		return &(static_cast<FSlateAttributeMetaData&>(SlateMetaData.Get()));
 	}
 #if WITH_SLATE_DEBUGGING
@@ -534,7 +534,7 @@ void FSlateAttributeMetaData::UpdateAttribute(SWidget& OwningWidget, FSlateAttri
 		if (FoundIndex != INDEX_NONE)
 		{
 			FGetterItem& GetterItem = AttributeMetaData->Attributes[FoundIndex];
-			check(GetterItem.Getter.Get());
+			checkSlow(GetterItem.Getter.Get());
 			ISlateAttributeGetter::FUpdateAttributeResult Result = GetterItem.Getter->UpdateAttribute(OwningWidget);
 			if (Result.bInvalidationRequested && OwningWidget.IsConstructed())
 			{
