Slate: Hide Unsuable NavigationPolicy Enums

#rnx

diff --git a/Engine/Source/Runtime/SlateCore/Public/Input/NavigationRouting.h b/Engine/Source/Runtime/SlateCore/Public/Input/NavigationRouting.h
--- a/Engine/Source/Runtime/SlateCore/Public/Input/NavigationRouting.h
+++ b/Engine/Source/Runtime/SlateCore/Public/Input/NavigationRouting.h
@@ -23,8 +23,8 @@ enum class EWidgetNavigationRoutingPolicy : uint8
 	RouteToTopRightChild UMETA(DisplayName="Route To Top Right Child", Tooltip="On navigation, this widget, or widgets within this boundary, will route focus to widget in the top right container"),
 	RouteToBottomLeftChild UMETA(DisplayName="Route To Bottom Left Child", Tooltip="On navigation, this widget, or widgets within this boundary, will route focus to widget in the bottom left container"),
 	RouteToBottomRightChild UMETA(DisplayName="Route To Bottom Right Child", Tooltip="On navigation, this widget, or widgets within this boundary, will route focus to widget in the bottom right container"),
-
-	Default = AcceptFocus,
+	MAX UMETA(Hidden),
+	Default = AcceptFocus UMETA(Hidden)
 };
 
 struct FNavigationRoutingParams
