added the ability to set the min / max width for toolbar buttons

#rnx

diff --git a/Engine/Source/Runtime/Slate/Private/Framework/MultiBox/SToolBarButtonBlock.cpp b/Engine/Source/Runtime/Slate/Private/Framework/MultiBox/SToolBarButtonBlock.cpp
--- a/Engine/Source/Runtime/Slate/Private/Framework/MultiBox/SToolBarButtonBlock.cpp
+++ b/Engine/Source/Runtime/Slate/Private/Framework/MultiBox/SToolBarButtonBlock.cpp
@@ -317,6 +317,8 @@ void SToolBarButtonBlock::BuildMultiBlockWidget(const ISlateStyle* StyleSet, con
 		.AutoWidth()
 		.Padding(ToolBarStyle.LabelPadding)
 		.VAlign(VAlign_Center))
+		.MaxWidth(ToolBarStyle.ButtonTextMaxWidth)
+		.MinWidth(ToolBarStyle.ButtonTextMinWidth)
 		[ TextBlock ];
 	}
 	else
diff --git a/Engine/Source/Runtime/SlateCore/Private/Styling/StarshipCoreStyle.cpp b/Engine/Source/Runtime/SlateCore/Private/Styling/StarshipCoreStyle.cpp
--- a/Engine/Source/Runtime/SlateCore/Private/Styling/StarshipCoreStyle.cpp
+++ b/Engine/Source/Runtime/SlateCore/Private/Styling/StarshipCoreStyle.cpp
@@ -2407,6 +2407,13 @@ void FStarshipCoreStyle::SetupMultiboxStyles(TSharedRef<FStyle>& Style)
 			SlimToolbarStyle.SetShowLabels(true);
 
 			Style->Set("CalloutToolbar", SlimToolbarStyle);
+
+			FToolBarStyle SlimToolbarStyleWidth100 = SlimToolbarStyle;
+			SlimToolbarStyleWidth100.SetButtonTextMaxWidth(100);
+			SlimToolbarStyleWidth100.SetButtonTextMinWidth(100);
+
+			Style->Set("CalloutToolbar-Width100", SlimToolbarStyleWidth100);
+
 		}
 
 	}
diff --git a/Engine/Source/Runtime/SlateCore/Public/Styling/ToolBarStyle.h b/Engine/Source/Runtime/SlateCore/Public/Styling/ToolBarStyle.h
--- a/Engine/Source/Runtime/SlateCore/Public/Styling/ToolBarStyle.h
+++ b/Engine/Source/Runtime/SlateCore/Public/Styling/ToolBarStyle.h
@@ -236,6 +236,14 @@ struct FToolBarStyle : public FSlateWidgetStyle
 	float ButtonContentMaxWidth = 64.0f;
 	FToolBarStyle& SetButtonContentMaxWidth(float InButtonContentMaxWidth) { ButtonContentMaxWidth = InButtonContentMaxWidth; return *this; }
 
+	UPROPERTY(EditAnywhere, BlueprintReadWrite, Category = Appearance)
+	float ButtonTextMinWidth = 0.0f;
+	FToolBarStyle& SetButtonTextMinWidth(float InButtonTextMinWidth) { ButtonTextMinWidth = InButtonTextMinWidth; return *this; }
+
+	UPROPERTY(EditAnywhere, BlueprintReadWrite, Category = Appearance)
+	float ButtonTextMaxWidth = FLT_MAX;
+	FToolBarStyle& SetButtonTextMaxWidth(float InButtonTextMaxWidth) { ButtonTextMaxWidth = InButtonTextMaxWidth; return *this; }
+	
 	UPROPERTY(EditAnywhere, BlueprintReadWrite, Category=Appearance)
 	float ButtonContentFillWidth = 1.0f;
 	FToolBarStyle& SetButtonContentFillWidth(float InButtonContentFillWidth) { ButtonContentFillWidth = InButtonContentFillWidth; return *this; }
