PR #12240: Slate Invalidation Debugger - add prepass color

diff --git a/Engine/Source/Runtime/SlateCore/Private/Debugging/ConsoleSlateDebuggerInvalidate.cpp b/Engine/Source/Runtime/SlateCore/Private/Debugging/ConsoleSlateDebuggerInvalidate.cpp
--- a/Engine/Source/Runtime/SlateCore/Private/Debugging/ConsoleSlateDebuggerInvalidate.cpp
+++ b/Engine/Source/Runtime/SlateCore/Private/Debugging/ConsoleSlateDebuggerInvalidate.cpp
@@ -32,7 +32,7 @@ FConsoleSlateDebuggerInvalidate::FConsoleSlateDebuggerInvalidate()
 	, DrawRootRootColor(FColorList::Red)
 	, DrawRootChildOrderColor(FColorList::Blue)
 	, DrawRootScreenPositionColor(FColorList::Green)
-	, DrawWidgetPrepassColor(FColorList::Magenta)
+	, DrawWidgetPrepassColor(FColorList::Orange)
 	, DrawWidgetLayoutColor(FColorList::Magenta)
 	, DrawWidgetPaintColor(FColorList::Yellow)
 	, DrawWidgetVolatilityColor(FColorList::Grey)
@@ -746,6 +746,7 @@ void FConsoleSlateDebuggerInvalidate::HandlePaintDebugInfo(const FPaintArgs& InA
 		MakeText(TEXT("Invalidation Root - Child Order"), DrawRootChildOrderColor);
 		MakeText(TEXT("Invalidation Root - Screen Position"), DrawRootScreenPositionColor);
 		MakeText(TEXT("Widget - Layout"), DrawWidgetLayoutColor);
+		MakeText(TEXT("Widget - Prepass"), DrawWidgetPrepassColor);
 		MakeText(TEXT("Widget - Paint"), DrawWidgetPaintColor);
 		MakeText(TEXT("Widget - Volatility"), DrawWidgetVolatilityColor);
 		MakeText(TEXT("Widget - Child Order"), DrawWidgetChildOrderColor);
