Add some error handling code in (non sdf) font rendering, for correctness, and hoping it will help understand a rare crash.

#tests Ran the editor to validate the code was executed properly, faked an error to validate error handling was correct

diff --git a/Engine/Source/Runtime/SlateCore/Private/Fonts/SlateFontRenderer.cpp b/Engine/Source/Runtime/SlateCore/Private/Fonts/SlateFontRenderer.cpp
--- a/Engine/Source/Runtime/SlateCore/Private/Fonts/SlateFontRenderer.cpp
+++ b/Engine/Source/Runtime/SlateCore/Private/Fonts/SlateFontRenderer.cpp
@@ -656,12 +656,22 @@ bool FSlateFontRenderer::GetRenderDataInternal(const FFreeTypeFaceGlyphData& InF
 		}
 
 		// Note: in order to render the stroke properly AND to get proper measurements this must be done after rendering the stroke
-		FT_Render_Glyph(Slot, bIsMonochromeRendered ? FT_RENDER_MODE_MONO : FT_RENDER_MODE_NORMAL);
+		FT_Error Error = FT_Render_Glyph(Slot, bIsMonochromeRendered ? FT_RENDER_MODE_MONO : FT_RENDER_MODE_NORMAL);
+		if (Error != 0)
+		{
+			UE_LOG(LogSlate, Warning, TEXT("FreeType error %x while trying to render glyph index %u from face %s %s (from outline path)"), Error, Slot->glyph_index, UTF8_TO_TCHAR(Face->family_name), UTF8_TO_TCHAR(Face->style_name));
+			return false;
+		}
 	}
 	else
 	{
-		// This path renders a standard font with no outline.  This may occur if the outline failed to generate
-		FT_Render_Glyph(Slot, bIsMonochromeRendered ? FT_RENDER_MODE_MONO : FT_RENDER_MODE_NORMAL);
+		// This path renders a standard font with no outline.
+		FT_Error Error = FT_Render_Glyph(Slot, bIsMonochromeRendered ? FT_RENDER_MODE_MONO : FT_RENDER_MODE_NORMAL);
+		if (Error != 0)
+		{
+			UE_LOG(LogSlate, Warning, TEXT("FreeType error %x while trying to render glyph index %u from face %s %s (from standard path)"), Error, Slot->glyph_index, UTF8_TO_TCHAR(Face->family_name), UTF8_TO_TCHAR(Face->style_name));
+			return false;
+		}
 
 		FT_Bitmap* Bitmap = nullptr;
 		FT_Bitmap TmpBitmap;
