-Optimization: take advantage of harfbuzz face cache system by avoiding to recreate full harfbuzz font from scratch for each harf buzz text sequence.
Thanks to that, harfbuzz can keep many font specific information in cache, and the shaping process is a lot faster (around 60 times faster when getting metrics just after EULA display in arabic at boot, but would lead to even better metrics on longer sessions).
-Added profile points for cpu time and memory usage to better monitor text shaping cost.
-Some cache usage simplification / code clean up.

#tests Ran the editor, displayed arabic text, flushed the cache, ran engine tests project auto tests. In engine, Displayed EULA in arabic, there is nearly more freeze on PC. Ran full flow (EULA => lobby => game => back to lobby), validated both perf and memory usage with Insights.

diff --git a/Engine/Source/Runtime/Slate/Private/Framework/Text/ShapedTextCache.cpp b/Engine/Source/Runtime/Slate/Private/Framework/Text/ShapedTextCache.cpp
--- a/Engine/Source/Runtime/Slate/Private/Framework/Text/ShapedTextCache.cpp
+++ b/Engine/Source/Runtime/Slate/Private/Framework/Text/ShapedTextCache.cpp
@@ -4,6 +4,7 @@
 #include "Fonts/FontCache.h"
 #include "Internationalization/IBreakIterator.h"
 #include "Internationalization/BreakIterator.h"
+#include "Trace/SlateMemoryTags.h"
 
 FShapedGlyphSequencePtr FShapedTextCache::FindShapedText(const FCachedShapedTextKey& InKey) const
 {
@@ -236,6 +237,8 @@ int8 ShapedTextCacheUtil::GetShapedGlyphKerning(const FShapedTextCacheRef& InSha
 
 FShapedGlyphSequenceRef ShapedTextCacheUtil::GetShapedTextSubSequence(const FShapedTextCacheRef& InShapedTextCache, const FCachedShapedTextKey& InRunKey, const FTextRange& InTextRange, const TCHAR* InText, const TextBiDi::ETextDirection InTextDirection)
 {
+	LLM_SCOPE_BYTAG(UI_Text);
+
 	// Get the shaped text for the entire run and try and make a sub-sequence from it - this can help minimize the amount of text shaping that needs to be done when drawing text
 	FShapedGlyphSequencePtr ShapedText = InShapedTextCache->FindOrAddShapedText(InRunKey, InText);
 
diff --git a/Engine/Source/Runtime/SlateCore/Private/Fonts/FontCache.cpp b/Engine/Source/Runtime/SlateCore/Private/Fonts/FontCache.cpp
--- a/Engine/Source/Runtime/SlateCore/Private/Fonts/FontCache.cpp
+++ b/Engine/Source/Runtime/SlateCore/Private/Fonts/FontCache.cpp
@@ -16,6 +16,7 @@
 #include "Fonts/FontCacheUtils.h"
 #include "Fonts/SlateSdfGenerator.h"
 #include "Fonts/FontGeometryPreprocessing.h"
+#include "Trace/SlateMemoryTags.h"
 
 #include UE_INLINE_GENERATED_CPP_BY_NAME(FontCache)
 
@@ -847,11 +848,11 @@ const FCharacterEntry* FCharacterList::CacheCharacter(TCHAR Character)
 
 			const bool bIsWhitespace = FTextChar::IsWhitespace(Character);
 			const uint32 GlyphIndex = FT_Get_Char_Index(FaceGlyphData.FaceAndMemory->GetFace(), Character);
-
+			const uint32 FontRenderSize = FreeTypeUtils::ComputeFontPixelSize(FontInfo.Size, FinalFontScale);
 			int16 XAdvance = 0;
 			{
 				FT_Fixed CachedAdvanceData = 0;
-				TSharedRef<FFreeTypeAdvanceCache> AdvanceCache = FontCache.FTCacheDirectory->GetAdvanceCache(FaceGlyphData.FaceAndMemory->GetFace(), GlyphFlags, FontInfo.Size, FinalFontScale);
+				TSharedRef<FFreeTypeAdvanceCache> AdvanceCache = FontCache.FTCacheDirectory->GetAdvanceCache(FaceGlyphData.FaceAndMemory->GetFace(), GlyphFlags, FontRenderSize);
 				if (AdvanceCache->FindOrCache(GlyphIndex, CachedAdvanceData))
 				{
 					XAdvance = FreeTypeUtils::Convert26Dot6ToRoundedPixel<int16>((CachedAdvanceData + (1 << 9)) >> 10);
@@ -883,7 +884,7 @@ const FCharacterEntry* FCharacterList::CacheCharacter(TCHAR Character)
 					CharEntry.Character = Character;
 					CharEntry.GlyphIndex = GlyphIndex;
 					CharEntry.FontData = FontDataPtr;
-					CharEntry.KerningCache = FontCache.FontRenderer->GetKerningCache(*FontDataPtr, FontInfo.Size, FinalFontScale);
+					CharEntry.KerningCache = FontCache.FontRenderer->GetKerningCache(*FontDataPtr, FontRenderSize);
 					CharEntry.FontScale = ShapedGlyphEntry.FontFaceData->FontScale;
 					CharEntry.BitmapRenderScale = ShapedGlyphEntry.FontFaceData->BitmapRenderScale;
 					CharEntry.StartU = ShapedGlyphFontAtlasData.StartU;
@@ -1323,6 +1324,8 @@ static float CeilSdfHalfSpread(float InEmSpread, float Ppem)
 
 FSdfGlyphFontAtlasData FSlateFontCache::GetSdfGlyphFontAtlasData(const FShapedGlyphEntry& InShapedGlyph, const FFontOutlineSettings& InOutlineSettings)
 {
+	LLM_SCOPE_BYTAG(UI_Text);
+
 	if (!InShapedGlyph.bIsVisible
 		|| InShapedGlyph.FontFaceData->RasterizationMode == EFontRasterizationMode::Bitmap
 		|| !IsSlateSdfTextFeatureEnabled())
@@ -1992,6 +1995,8 @@ void FSlateFontCache::FlushData()
 		SdfGenerator->Flush();
 	}
 
+	TextShaper->FlushCache();
+
 	if (GIsEditor || UnloadFreeTypeDataOnFlush)
 	{
 		FTCacheDirectory->FlushCache();
diff --git a/Engine/Source/Runtime/SlateCore/Private/Fonts/FontCacheFreeType.cpp b/Engine/Source/Runtime/SlateCore/Private/Fonts/FontCacheFreeType.cpp
--- a/Engine/Source/Runtime/SlateCore/Private/Fonts/FontCacheFreeType.cpp
+++ b/Engine/Source/Runtime/SlateCore/Private/Fonts/FontCacheFreeType.cpp
@@ -106,25 +106,14 @@ FT_F26Dot6 Round26Dot6(const FT_F26Dot6 InValue)
 	return Floor26Dot6(InValue + 32);
 }
 
-FT_F26Dot6 Determine26Dot6Ppem(const float InFontSize, const float InFontScale, const bool InRoundPpem) //TODO: this function is very similar to ComputeFontPixelSize. Merge both?
-{
-	FT_F26Dot6 Ppem = static_cast<FT_F26Dot6>(::FT_MulFix((ConvertPixelTo26Dot6<FT_Long>(FMath::Max(0.f, InFontSize)) * (FT_Long)FontConstants::RenderDPI + 36) / 72,
-														  ConvertPixelTo16Dot16<FT_Long>(FMath::Max(0.f, InFontScale))));
-	if (InRoundPpem)
-	{
-		Ppem = Round26Dot6(Ppem);
-	}
-	return Ppem;
-}
-
 FT_Fixed DetermineEmScale(const uint16 InEmSize, const FT_F26Dot6 InPpem)
 {
 	return static_cast<FT_Fixed>(::FT_DivFix(InPpem, static_cast<FT_Long>(FMath::Max(uint16{1}, InEmSize))));
 }
 
-FT_Fixed DeterminePpemAndEmScale(const uint16 InEmSize, const float InFontSize, const float InFontScale, const bool InRoundPpem)
+FT_Fixed DeterminePpemAndEmScale(const uint16 InEmSize, const uint32 InRenderSize)
 {
-	return DetermineEmScale(InEmSize, Determine26Dot6Ppem(InFontSize, InFontScale, InRoundPpem));
+	return DetermineEmScale(InEmSize, ConvertPixelTo26Dot6<FT_F26Dot6>(InRenderSize));
 }
 
 uint32 ComputeFontPixelSize(float InFontSize, float InFontScale)
@@ -730,10 +719,10 @@ private:
 	FT_Size PreviousSize;
 };
 
-FFreeTypeGlyphCache::FFreeTypeGlyphCache(FT_Face InFace, const int32 InLoadFlags, const float InFontSize, const float InFontScale)
+FFreeTypeGlyphCache::FFreeTypeGlyphCache(FT_Face InFace, const int32 InLoadFlags, const uint32 InFontRenderSize)
 	: Face(InFace)
 	, LoadFlags(InLoadFlags)
-	, FontRenderSize(FreeTypeUtils::ComputeFontPixelSize(InFontSize,InFontScale))
+	, FontRenderSize(InFontRenderSize)
 	, GlyphDataMap()
 {
 }
@@ -779,18 +768,18 @@ bool FFreeTypeGlyphCache::FindOrCache(const uint32 InGlyphIndex, FCachedGlyphDat
 FFreeTypeAdvanceCache::FFreeTypeAdvanceCache()
 	: Face(0)
 	, LoadFlags(0)
-	, FontRenderSize(0.f)
+	, FontRenderSize(0)
 	, AdvanceMap()
 {
 	// invalid cache always returns false
 }
 
-FFreeTypeAdvanceCache::FFreeTypeAdvanceCache(FT_Face InFace, const int32 InLoadFlags, const float InFontSize, const float InFontScale)
+FFreeTypeAdvanceCache::FFreeTypeAdvanceCache(FT_Face InFace, const int32 InLoadFlags, const uint32 InFontRenderSize)
 	: Face(InFace)
 	, LoadFlags(InLoadFlags)
 	, FontRenderSize((0 != (InLoadFlags & FT_LOAD_NO_SCALE))
-			   ? static_cast<int32>(FreeTypeUtils::DeterminePpemAndEmScale(InFace->units_per_EM, InFontSize, InFontScale, true))
-			   : FreeTypeUtils::ComputeFontPixelSize(InFontSize, InFontScale)) // compute once the EmScale and keep in FontSize when caching unscaled advances
+			   ? static_cast<uint32>(FreeTypeUtils::DeterminePpemAndEmScale(InFace->units_per_EM, InFontRenderSize))
+			   : InFontRenderSize) // compute once the EmScale and keep in FontSize when caching unscaled advances
 	, AdvanceMap()
 {
 	check(((LoadFlags & FT_LOAD_NO_SCALE) == 0) || FreeTypeUtils::IsFaceEligibleForSdf(Face));
@@ -845,12 +834,12 @@ bool FFreeTypeAdvanceCache::FindOrCache(const uint32 InGlyphIndex, FT_Fixed& Out
 }
 
 
-FFreeTypeKerningCache::FFreeTypeKerningCache(FT_Face InFace, const int32 InKerningFlags, const float InFontSize, const float InFontScale)
+FFreeTypeKerningCache::FFreeTypeKerningCache(FT_Face InFace, const int32 InKerningFlags, const uint32 InFontRenderSize)
 	: Face(InFace)
 	, KerningFlags(InKerningFlags)
 	, FontRenderSize(InKerningFlags == FT_KERNING_UNSCALED ?
-			   static_cast<int32>(FreeTypeUtils::DeterminePpemAndEmScale(InFace->units_per_EM, InFontSize, InFontScale, true)) : 
-			   FreeTypeUtils::ComputeFontPixelSize(InFontSize, InFontScale)) // compute once the EmScale and keep in FontSize when caching unscaled advances
+			   static_cast<uint32>(FreeTypeUtils::DeterminePpemAndEmScale(InFace->units_per_EM, InFontRenderSize)) :
+			   InFontRenderSize) // compute once the EmScale and keep in FontSize when caching unscaled advances
 	, KerningMap()
 {
 	check(Face);
@@ -906,34 +895,34 @@ bool FFreeTypeKerningCache::FindOrCache(const uint32 InFirstGlyphIndex, const ui
 	return false;
 }
 
-TSharedRef<FFreeTypeGlyphCache> FFreeTypeCacheDirectory::GetGlyphCache(FT_Face InFace, const int32 InLoadFlags, const float InFontSize, const float InFontScale)
+TSharedRef<FFreeTypeGlyphCache> FFreeTypeCacheDirectory::GetGlyphCache(FT_Face InFace, const int32 InLoadFlags, const uint32 InFontRenderSize)
 {
-	const FFontKey Key(InFace, InLoadFlags, InFontSize, InFontScale);
+	const FFontKey Key(InFace, InLoadFlags, InFontRenderSize);
 	TSharedPtr<FFreeTypeGlyphCache>& Entry = GlyphCacheMap.FindOrAdd(Key);
 	if (!Entry)
 	{
-		Entry = MakeShared<FFreeTypeGlyphCache>(InFace, InLoadFlags, InFontSize, InFontScale);
+		Entry = MakeShared<FFreeTypeGlyphCache>(InFace, InLoadFlags, InFontRenderSize);
 	}
 	return Entry.ToSharedRef();
 }
 
-TSharedRef<FFreeTypeAdvanceCache> FFreeTypeCacheDirectory::GetAdvanceCache(FT_Face InFace, const int32 InLoadFlags, const float InFontSize, const float InFontScale)
+TSharedRef<FFreeTypeAdvanceCache> FFreeTypeCacheDirectory::GetAdvanceCache(FT_Face InFace, const int32 InLoadFlags, const uint32 InFontRenderSize)
 {
 	if (!InFace || ((InLoadFlags & FT_LOAD_NO_SCALE) != 0 && !FreeTypeUtils::IsFaceEligibleForSdf(InFace)))
 	{
 		checkNoEntry(); // the slate clients can't reach these caches. Reaching this point is an internal error.
 		return InvalidAdvanceCache.ToSharedRef();
 	}
-	const FFontKey Key(InFace, InLoadFlags, InFontSize, InFontScale);
+	const FFontKey Key(InFace, InLoadFlags, InFontRenderSize);
 	TSharedPtr<FFreeTypeAdvanceCache>& Entry = AdvanceCacheMap.FindOrAdd(Key);
 	if (!Entry)
 	{
-		Entry = MakeShared<FFreeTypeAdvanceCache>(InFace, InLoadFlags, InFontSize, InFontScale);
+		Entry = MakeShared<FFreeTypeAdvanceCache>(InFace, InLoadFlags, InFontRenderSize);
 	}
 	return Entry.ToSharedRef();
 }
 
-TSharedPtr<FFreeTypeKerningCache> FFreeTypeCacheDirectory::GetKerningCache(FT_Face InFace, const int32 InKerningFlags, const float InFontSize, const float InFontScale)
+TSharedPtr<FFreeTypeKerningCache> FFreeTypeCacheDirectory::GetKerningCache(FT_Face InFace, const int32 InKerningFlags, const uint32 InFontRenderSize)
 {
 	// Check if this font has kerning as not all fonts do.
 	// We also can't perform kerning between two separate font faces
@@ -941,11 +930,11 @@ TSharedPtr<FFreeTypeKerningCache> FFreeTypeCacheDirectory::GetKerningCache(FT_Fa
 	{
 		if ((InKerningFlags != FT_KERNING_UNSCALED || FreeTypeUtils::IsFaceEligibleForSdf(InFace)))
 		{
-			const FFontKey Key(InFace, InKerningFlags, InFontSize, InFontScale);
+			const FFontKey Key(InFace, InKerningFlags, InFontRenderSize);
 			TSharedPtr<FFreeTypeKerningCache>& Entry = KerningCacheMap.FindOrAdd(Key);
 			if (!Entry)
 			{
-				Entry = MakeShared<FFreeTypeKerningCache>(InFace, InKerningFlags, InFontSize, InFontScale);
+				Entry = MakeShared<FFreeTypeKerningCache>(InFace, InKerningFlags, InFontRenderSize);
 			}
 			return Entry;
 		}
diff --git a/Engine/Source/Runtime/SlateCore/Private/Fonts/FontCacheFreeType.h b/Engine/Source/Runtime/SlateCore/Private/Fonts/FontCacheFreeType.h
--- a/Engine/Source/Runtime/SlateCore/Private/Fonts/FontCacheFreeType.h
+++ b/Engine/Source/Runtime/SlateCore/Private/Fonts/FontCacheFreeType.h
@@ -74,12 +74,6 @@ FT_F26Dot6 Ceil26Dot6(const FT_F26Dot6 InValue);
 */
 FT_F26Dot6 Round26Dot6(const FT_F26Dot6 InValue);
 
-/**
- * Determine the (optionally rounded) pixel size (the number of pixels per em square dimensions) 
- * from a font size in points (72 points per inch) and an arbitrary ui scaling at a resolution of 96 dpi.
- */
-FT_F26Dot6 Determine26Dot6Ppem(const float InFontSize, const float InFontScale, const bool InRoundPpem);
-
 /**
  * The EmScale maps design space distances relative to the em square (with resolution of InEmSize units),
  * to absolute 1/64th pixels distances in the device pixel plane (with a resolution of 96 dpi) with 
@@ -88,9 +82,9 @@ FT_F26Dot6 Determine26Dot6Ppem(const float InFontSize, const float InFontScale,
 FT_Fixed DetermineEmScale(const uint16 InEmSize, const FT_F26Dot6 InPpem);
 
 /**
-* Determine the (optionally rounded) Ppem and then from it determine the EmScale, in one call.
+* Determine the rounded Ppem and then from it determine the EmScale, in one call.
 */
-FT_Fixed DeterminePpemAndEmScale(const uint16 InEmSize, const float InFontSize, const float InFontScale, const bool InRoundPpem);
+FT_Fixed DeterminePpemAndEmScale(const uint16 InEmSize, const float InFontSize, const float InFontScale);
 
 
 /**
@@ -424,7 +418,7 @@ class FFreeTypeGlyphCache
 {
 public:
 #if WITH_FREETYPE
-	FFreeTypeGlyphCache(FT_Face InFace, const int32 InLoadFlags, const float InFontSize, const float InFontScale);
+	FFreeTypeGlyphCache(FT_Face InFace, const int32 InLoadFlags, const uint32 InFontRenderSize);
 
 	struct FCachedGlyphData
 	{
@@ -455,7 +449,7 @@ class FFreeTypeAdvanceCache
 public:
 #if WITH_FREETYPE
 	FFreeTypeAdvanceCache();
-	FFreeTypeAdvanceCache(FT_Face InFace, const int32 InLoadFlags, const float InFontSize, const float InFontScale);
+	FFreeTypeAdvanceCache(FT_Face InFace, const int32 InLoadFlags, const uint32 InFontRenderSize);
 
 	bool FindOrCache(const uint32 InGlyphIndex, FT_Fixed& OutCachedAdvance);
 #endif // WITH_FREETYPE
@@ -479,7 +473,7 @@ class FFreeTypeKerningCache
 {
 public:
 #if WITH_FREETYPE
-	FFreeTypeKerningCache(FT_Face InFace, const int32 InKerningFlags, const float InFontSize, const float InFontScale);
+	FFreeTypeKerningCache(FT_Face InFace, const int32 InKerningFlags, const uint32 InFontRenderSize);
 
 	/**
 	 * Retrieve the kerning vector for a given pair of glyphs.
@@ -526,7 +520,7 @@ private:
 
 	FT_Face Face;
 	const int32 KerningFlags;
-	const int32 FontRenderSize;
+	const uint32 FontRenderSize;
 	TMap<FKerningPair, FT_Vector> KerningMap;
 #endif // WITH_FREETYPE
 };
@@ -544,19 +538,19 @@ public:
 	 * Retrieve the glyph cache for a given set of font parameters.
 	 * @return A reference to the font glyph cache.
 	 */
-	TSharedRef<FFreeTypeGlyphCache> GetGlyphCache(FT_Face InFace, const int32 InLoadFlags, const float InFontSize, const float InFontScale);
+	TSharedRef<FFreeTypeGlyphCache> GetGlyphCache(FT_Face InFace, const int32 InLoadFlags, const uint32 InFontRenderSize);
 
 	/**
 	 * Retrieve the advance cache for a given set of font parameters.
 	 * @return A reference to the font advance cache.
 	 */
-	TSharedRef<FFreeTypeAdvanceCache> GetAdvanceCache(FT_Face InFace, const int32 InLoadFlags, const float InFontSize, const float InFontScale);
+	TSharedRef<FFreeTypeAdvanceCache> GetAdvanceCache(FT_Face InFace, const int32 InLoadFlags, const uint32 InFontRenderSize);
 
 	/**
 	 * Retrieve the kerning cache for a given set of font parameters.
 	 * @return A pointer to the font kerning cache, invalid if the font does not perform kerning.
 	 */
-	TSharedPtr<FFreeTypeKerningCache> GetKerningCache(FT_Face InFace, const int32 InKerningFlags, const float InFontSize, const float InFontScale);
+	TSharedPtr<FFreeTypeKerningCache> GetKerningCache(FT_Face InFace, const int32 InKerningFlags, const uint32 InFontRenderSize);
 #endif // WITH_FREETYPE
 
 	void FlushCache();
@@ -567,10 +561,10 @@ private:
 	class FFontKey
 	{
 	public:
-		FFontKey(FT_Face InFace, const int32 InFlags, const float InFontSize, const float InFontScale)
+		FFontKey(FT_Face InFace, const int32 InFlags, const uint32 InFontRenderSize)
 			: Face(InFace)
 			, Flags(InFlags)
-			, FontRenderSize(FreeTypeUtils::ComputeFontPixelSize(InFontSize, InFontScale))
+			, FontRenderSize(InFontRenderSize)
 			, KeyHash(0)
 		{
 			KeyHash = GetTypeHash(Face);
@@ -598,7 +592,7 @@ private:
 	private:
 		FT_Face Face;
 		const int32 Flags;
-		const int32 FontRenderSize;
+		const uint32 FontRenderSize;
 		uint32 KeyHash;
 	};
 
diff --git a/Engine/Source/Runtime/SlateCore/Private/Fonts/FontCacheHarfBuzz.cpp b/Engine/Source/Runtime/SlateCore/Private/Fonts/FontCacheHarfBuzz.cpp
--- a/Engine/Source/Runtime/SlateCore/Private/Fonts/FontCacheHarfBuzz.cpp
+++ b/Engine/Source/Runtime/SlateCore/Private/Fonts/FontCacheHarfBuzz.cpp
@@ -7,6 +7,17 @@
 #include "Trace/SlateMemoryTags.h"
 #include "Fonts/FontUtils.h"
 
+#if WITH_FREETYPE && WITH_HARFBUZZ
+DECLARE_CYCLE_STAT(TEXT("Create font"), STAT_HarfBuzzFontCache_ApplyFont, STATGROUP_Slate);
+DECLARE_CYCLE_STAT(TEXT("Get_nominal_glyphs"), STAT_HarfBuzzFontFunctions_get_nominal_glyphs, STATGROUP_Slate);
+DECLARE_CYCLE_STAT(TEXT("Get_glyph_h_advances"), STAT_HarfBuzzFontFunctions_get_glyph_h_advances, STATGROUP_Slate);
+DECLARE_CYCLE_STAT(TEXT("Get_glyph_v_advances"), STAT_HarfBuzzFontFunctions_get_glyph_v_advances, STATGROUP_Slate);
+DECLARE_CYCLE_STAT(TEXT("Get_glyph_v_origin"), STAT_HarfBuzzFontFunctions_get_glyph_v_origin, STATGROUP_Slate);
+DECLARE_CYCLE_STAT(TEXT("Get_glyph_h_kerning"), STAT_HarfBuzzFontFunctions_get_glyph_h_kerning, STATGROUP_Slate);
+DECLARE_CYCLE_STAT(TEXT("Get_glyph_extents"), STAT_HarfBuzzFontFunctions_get_glyph_extents, STATGROUP_Slate);
+DECLARE_CYCLE_STAT(TEXT("Get_glyph_contour_point"), STAT_HarfBuzzFontFunctions_get_glyph_contour_point, STATGROUP_Slate);
+#endif // #if WITH_FREETYPE && WITH_HARFBUZZ
+
 #if WITH_HARFBUZZ
 
 extern "C"
@@ -104,23 +115,26 @@ hb_user_data_key_t UserDataKey;
 
 struct FUserData
 {
-	FUserData(const float InFontSize, const float InFontScale, FFreeTypeCacheDirectory* InFTCacheDirectory, const hb_font_extents_t& InHarfBuzzFontExtents)
-		: FontSize(InFontSize)
-		, FontScale(InFontScale)
+	FUserData(uint32 InRenderSize, FFreeTypeCacheDirectory* InFTCacheDirectory, FT_Face InFreeTypeFace, int32 InFreeTypeFlags)
+		: RenderSize(InRenderSize)
 		, FTCacheDirectory(InFTCacheDirectory)
-		, HarfBuzzFontExtents(InHarfBuzzFontExtents)
+		, HarfBuzzFontExtents()
+		, FreeTypeFace(InFreeTypeFace)
+		, FreeTypeFlags(InFreeTypeFlags)
 	{
+		FMemory::Memzero(HarfBuzzFontExtents);
 	}
 
-	float FontSize;
-	float FontScale;
+	uint32 RenderSize;
 	FFreeTypeCacheDirectory* FTCacheDirectory;
 	hb_font_extents_t HarfBuzzFontExtents;
+	FT_Face FreeTypeFace;
+	int32 FreeTypeFlags;
 };
 
-void* CreateUserData(const float InFontSize, const float InFontScale, FFreeTypeCacheDirectory* InFTCacheDirectory, const hb_font_extents_t& InHarfBuzzFontExtents)
+FUserData* CreateUserData(uint32 InRenderSize, FFreeTypeCacheDirectory* InFTCacheDirectory, FT_Face InFreeTypeFace, int32 InFreeTypeFlags)
 {
-	return new FUserData(InFontSize, InFontScale, InFTCacheDirectory, InHarfBuzzFontExtents);
+	return new FUserData(InRenderSize, InFTCacheDirectory, InFreeTypeFace, InFreeTypeFlags);
 }
 
 void DestroyUserData(void* UserData)
@@ -131,21 +145,6 @@ void DestroyUserData(void* UserData)
 
 namespace Internal
 {
-
-FORCEINLINE FT_Face get_ft_face(hb_font_t* InFont)
-{
-	hb_font_t* FontParent = hb_font_get_parent(InFont);
-	check(FontParent);
-	return hb_ft_font_get_face(FontParent);
-}
-
-FORCEINLINE int32 get_ft_flags(hb_font_t* InFont)
-{
-	hb_font_t* FontParent = hb_font_get_parent(InFont);
-	check(FontParent);
-	return hb_ft_font_get_load_flags(FontParent);
-}
-
 hb_bool_t get_font_h_extents(hb_font_t* InFont, void* InFontData, hb_font_extents_t* OutMetrics, void* InUserData)
 {
 	const FUserData* UserDataPtr = static_cast<FUserData*>(hb_font_get_user_data(InFont, &UserDataKey));
@@ -155,7 +154,11 @@ hb_bool_t get_font_h_extents(hb_font_t* InFont, void* InFontData, hb_font_extent
 
 unsigned int get_nominal_glyphs(hb_font_t* InFont, void* InFontData, unsigned int InCount, const hb_codepoint_t *InUnicodeCharBuffer, unsigned int InUnicodeCharBufferStride, hb_codepoint_t *OutGlyphIndexBuffer, unsigned int InGlyphIndexBufferStride, void* InUserData)
 {
-	FT_Face FreeTypeFace = get_ft_face(InFont);
+//Too verbose for now, but will help in subsequent optimization
+//	SCOPE_CYCLE_COUNTER(STAT_HarfBuzzFontFunctions_get_nominal_glyphs);
+
+	const FUserData* UserDataPtr = static_cast<FUserData*>(hb_font_get_user_data(InFont, &UserDataKey));
+	FT_Face FreeTypeFace = UserDataPtr->FreeTypeFace;
 
 	const uint8* UnicodeCharRawBuffer = (const uint8*)InUnicodeCharBuffer;
 	uint8* GlyphIndexRawBuffer = (uint8*)OutGlyphIndexBuffer;
@@ -195,9 +198,12 @@ hb_bool_t get_nominal_glyph(hb_font_t* InFont, void* InFontData, hb_codepoint_t
 
 void get_glyph_h_advances(hb_font_t* InFont, void* InFontData, unsigned int InCount, const hb_codepoint_t* InGlyphIndexBuffer, unsigned int InGlyphIndexBufferStride, hb_position_t* OutAdvanceBuffer, unsigned int InAdvanceBufferStride, void* InUserData)
 {
-	FT_Face FreeTypeFace = get_ft_face(InFont);
-	const int32 FreeTypeFlags = get_ft_flags(InFont);
+//Too verbose for now, but will help in subsequent optimization
+//	SCOPE_CYCLE_COUNTER(STAT_HarfBuzzFontFunctions_get_glyph_h_advances);
+
 	const FUserData* UserDataPtr = static_cast<FUserData*>(hb_font_get_user_data(InFont, &UserDataKey));
+	FT_Face FreeTypeFace = UserDataPtr->FreeTypeFace;
+	const int32 FreeTypeFlags = UserDataPtr->FreeTypeFlags;
 
 	int ScaleMultiplier = 1;
 	{
@@ -213,7 +219,7 @@ void get_glyph_h_advances(hb_font_t* InFont, void* InFontData, unsigned int InCo
 
 	const uint8* GlyphIndexRawBuffer = (const uint8*)InGlyphIndexBuffer;
 	uint8* AdvanceRawBuffer = (uint8*)OutAdvanceBuffer;
-	TSharedRef<FFreeTypeAdvanceCache> AdvanceCache = UserDataPtr->FTCacheDirectory->GetAdvanceCache(FreeTypeFace, FreeTypeFlags, UserDataPtr->FontSize, UserDataPtr->FontScale);
+	TSharedRef<FFreeTypeAdvanceCache> AdvanceCache = UserDataPtr->FTCacheDirectory->GetAdvanceCache(FreeTypeFace, FreeTypeFlags, UserDataPtr->RenderSize);
 
 	for (unsigned int ItemIndex = 0; ItemIndex < InCount; ++ItemIndex)
 	{
@@ -245,9 +251,12 @@ hb_position_t get_glyph_h_advance(hb_font_t* InFont, void* InFontData, hb_codepo
 
 void get_glyph_v_advances(hb_font_t* InFont, void* InFontData, unsigned int InCount, const hb_codepoint_t* InGlyphIndexBuffer, unsigned int InGlyphIndexBufferStride, hb_position_t* OutAdvanceBuffer, unsigned int InAdvanceBufferStride, void* InUserData)
 {
-	FT_Face FreeTypeFace = get_ft_face(InFont);
-	const int32 FreeTypeFlags = get_ft_flags(InFont);
+//Too verbose for now, but will help in subsequent optimization
+//	SCOPE_CYCLE_COUNTER(STAT_HarfBuzzFontFunctions_get_glyph_v_advances);
+
 	const FUserData* UserDataPtr = static_cast<FUserData*>(hb_font_get_user_data(InFont, &UserDataKey));
+	FT_Face FreeTypeFace = UserDataPtr->FreeTypeFace;
+	const int32 FreeTypeFlags = UserDataPtr->FreeTypeFlags;
 
 	int ScaleMultiplier = 1;
 	{
@@ -263,7 +272,7 @@ void get_glyph_v_advances(hb_font_t* InFont, void* InFontData, unsigned int InCo
 
 	const uint8* GlyphIndexRawBuffer = (const uint8*)InGlyphIndexBuffer;
 	uint8* AdvanceRawBuffer = (uint8*)OutAdvanceBuffer;
-	TSharedRef<FFreeTypeAdvanceCache> AdvanceCache = UserDataPtr->FTCacheDirectory->GetAdvanceCache(FreeTypeFace, FreeTypeFlags | FT_LOAD_VERTICAL_LAYOUT, UserDataPtr->FontSize, UserDataPtr->FontScale);
+	TSharedRef<FFreeTypeAdvanceCache> AdvanceCache = UserDataPtr->FTCacheDirectory->GetAdvanceCache(FreeTypeFace, FreeTypeFlags | FT_LOAD_VERTICAL_LAYOUT, UserDataPtr->RenderSize);
 
 	for (unsigned int ItemIndex = 0; ItemIndex < InCount; ++ItemIndex)
 	{
@@ -296,11 +305,14 @@ hb_position_t get_glyph_v_advance(hb_font_t* InFont, void* InFontData, hb_codepo
 
 hb_bool_t get_glyph_v_origin(hb_font_t* InFont, void* InFontData, hb_codepoint_t InGlyphIndex, hb_position_t* OutX, hb_position_t* OutY, void* InUserData)
 {
-	FT_Face FreeTypeFace = get_ft_face(InFont);
-	const int32 FreeTypeFlags = get_ft_flags(InFont);
+//Too verbose for now, but will help in subsequent optimization
+//	SCOPE_CYCLE_COUNTER(STAT_HarfBuzzFontFunctions_get_glyph_v_origin);
+
 	const FUserData* UserDataPtr = static_cast<FUserData*>(hb_font_get_user_data(InFont, &UserDataKey));
+	FT_Face FreeTypeFace = UserDataPtr->FreeTypeFace;
+	const int32 FreeTypeFlags = UserDataPtr->FreeTypeFlags;
 
-	TSharedRef<FFreeTypeGlyphCache> GlyphCache = UserDataPtr->FTCacheDirectory->GetGlyphCache(FreeTypeFace, FreeTypeFlags, UserDataPtr->FontSize, UserDataPtr->FontScale);
+	TSharedRef<FFreeTypeGlyphCache> GlyphCache = UserDataPtr->FTCacheDirectory->GetGlyphCache(FreeTypeFace, FreeTypeFlags, UserDataPtr->RenderSize);
 	FFreeTypeGlyphCache::FCachedGlyphData CachedGlyphData;
 	if (GlyphCache->FindOrCache(InGlyphIndex, CachedGlyphData))
 	{
@@ -330,10 +342,13 @@ hb_bool_t get_glyph_v_origin(hb_font_t* InFont, void* InFontData, hb_codepoint_t
 
 hb_position_t get_glyph_h_kerning(hb_font_t* InFont, void* InFontData, hb_codepoint_t InLeftGlyphIndex, hb_codepoint_t InRightGlyphIndex, void* InUserData)
 {
-	FT_Face FreeTypeFace = get_ft_face(InFont);
+//Too verbose for now, but will help in subsequent optimization
+//	SCOPE_CYCLE_COUNTER(STAT_HarfBuzzFontFunctions_get_glyph_h_kerning);
+
 	const FUserData* UserDataPtr = static_cast<FUserData*>(hb_font_get_user_data(InFont, &UserDataKey));
+	FT_Face FreeTypeFace = UserDataPtr->FreeTypeFace;
 
-	TSharedPtr<FFreeTypeKerningCache> KerningCache = UserDataPtr->FTCacheDirectory->GetKerningCache(FreeTypeFace, FT_KERNING_DEFAULT, UserDataPtr->FontSize, UserDataPtr->FontScale);
+	TSharedPtr<FFreeTypeKerningCache> KerningCache = UserDataPtr->FTCacheDirectory->GetKerningCache(FreeTypeFace, FT_KERNING_DEFAULT, UserDataPtr->RenderSize);
 	if (KerningCache)
 	{
 		FT_Vector KerningVector;
@@ -348,11 +363,14 @@ hb_position_t get_glyph_h_kerning(hb_font_t* InFont, void* InFontData, hb_codepo
 
 hb_bool_t get_glyph_extents(hb_font_t* InFont, void* InFontData, hb_codepoint_t InGlyphIndex, hb_glyph_extents_t* OutExtents, void* InUserData)
 {
-	FT_Face FreeTypeFace = get_ft_face(InFont);
-	const int32 FreeTypeFlags = get_ft_flags(InFont);
+//Too verbose for now, but will help in subsequent optimization
+//	SCOPE_CYCLE_COUNTER(STAT_HarfBuzzFontFunctions_get_glyph_extents);
+
 	const FUserData* UserDataPtr = static_cast<FUserData*>(hb_font_get_user_data(InFont, &UserDataKey));
+	FT_Face FreeTypeFace = UserDataPtr->FreeTypeFace;
+	const int32 FreeTypeFlags = UserDataPtr->FreeTypeFlags;
 
-	TSharedRef<FFreeTypeGlyphCache> GlyphCache = UserDataPtr->FTCacheDirectory->GetGlyphCache(FreeTypeFace, FreeTypeFlags, UserDataPtr->FontSize, UserDataPtr->FontScale);
+	TSharedRef<FFreeTypeGlyphCache> GlyphCache = UserDataPtr->FTCacheDirectory->GetGlyphCache(FreeTypeFace, FreeTypeFlags, UserDataPtr->RenderSize);
 	FFreeTypeGlyphCache::FCachedGlyphData CachedGlyphData;
 	if (GlyphCache->FindOrCache(InGlyphIndex, CachedGlyphData))
 	{
@@ -385,11 +403,14 @@ hb_bool_t get_glyph_extents(hb_font_t* InFont, void* InFontData, hb_codepoint_t
 
 hb_bool_t get_glyph_contour_point(hb_font_t* InFont, void* InFontData, hb_codepoint_t InGlyphIndex, unsigned int InPointIndex, hb_position_t* OutX, hb_position_t* OutY, void* InUserData)
 {
-	FT_Face FreeTypeFace = get_ft_face(InFont);
-	const int32 FreeTypeFlags = get_ft_flags(InFont);
+//Too verbose for now, but will help in subsequent optimization
+//	SCOPE_CYCLE_COUNTER(STAT_HarfBuzzFontFunctions_get_glyph_contour_point);
+
 	const FUserData* UserDataPtr = static_cast<FUserData*>(hb_font_get_user_data(InFont, &UserDataKey));
+	FT_Face FreeTypeFace = UserDataPtr->FreeTypeFace;
+	const int32 FreeTypeFlags = UserDataPtr->FreeTypeFlags;
 
-	TSharedRef<FFreeTypeGlyphCache> GlyphCache = UserDataPtr->FTCacheDirectory->GetGlyphCache(FreeTypeFace, FreeTypeFlags, UserDataPtr->FontSize, UserDataPtr->FontScale);
+	TSharedRef<FFreeTypeGlyphCache> GlyphCache = UserDataPtr->FTCacheDirectory->GetGlyphCache(FreeTypeFace, FreeTypeFlags, UserDataPtr->RenderSize);
 	FFreeTypeGlyphCache::FCachedGlyphData CachedGlyphData;
 	if (GlyphCache->FindOrCache(InGlyphIndex, CachedGlyphData))
 	{
@@ -410,7 +431,7 @@ hb_bool_t get_glyph_contour_point(hb_font_t* InFont, void* InFontData, hb_codepo
 
 #endif // WITH_FREETYPE && WITH_HARFBUZZ
 
-FHarfBuzzFontFactory::FHarfBuzzFontFactory(FFreeTypeCacheDirectory* InFTCacheDirectory)
+FHarfBuzzFontCache::FHarfBuzzFontCache(FFreeTypeCacheDirectory* InFTCacheDirectory)
 	: FTCacheDirectory(InFTCacheDirectory)
 {
 	check(FTCacheDirectory);
@@ -440,8 +461,10 @@ FHarfBuzzFontFactory::FHarfBuzzFontFactory(FFreeTypeCacheDirectory* InFTCacheDir
 #endif // WITH_HARFBUZZ
 }
 
-FHarfBuzzFontFactory::~FHarfBuzzFontFactory()
+FHarfBuzzFontCache::~FHarfBuzzFontCache()
 {
+	Flush();
+
 #if WITH_HARFBUZZ
 	hb_font_funcs_destroy(CustomHarfBuzzFuncs);
 	CustomHarfBuzzFuncs = nullptr;
@@ -450,58 +473,86 @@ FHarfBuzzFontFactory::~FHarfBuzzFontFactory()
 
 #if WITH_HARFBUZZ
 
-hb_font_t* FHarfBuzzFontFactory::CreateFont(const FFreeTypeFace& InFace, const uint32 InGlyphFlags, const FSlateFontInfo& InFontInfo, const float InFontScale) const
+#if WITH_FREETYPE
+FHarfBuzzFontCache::FFontKey::FFontKey(FT_Face InFace, const int32 InFlags)
+: Face(InFace)
+, Flags(InFlags)
+, KeyHash(0)
+{
+	KeyHash = GetTypeHash(Face);
+	KeyHash = HashCombine(KeyHash, GetTypeHash(Flags));
+}
+#endif //WITH_FREETYPE
+
+// Shaping the glyph can be a very complex process, requiring many lookup to take the right decisions. Harfbuzz is using a caching system internally to avoid recomputing everything each time
+// it shapes a text. In order to take advantage of it, it's really important to not recreate hb_font from scratch for each shape process, but unfortunately keeping all the hb_fonts in memory, with their cache,
+// use too much memory. That's why the following function is keeping the hb_face (which is size agnostic and still keep the cache system of harfbuzz) in a cache, then for each call,
+// it create the hb_font from the cached hb_face on the fly. This way, we speed up the shaping a lot thanks to the harfbuzz cache, while keeping memory low (because we don't keep 
+// a font for each face and size combination).
+hb_font_t* FHarfBuzzFontCache::CreateFont(const FFreeTypeFace& InFace, const uint32 InGlyphFlags, const FSlateFontInfo& InFontInfo, const float InFontScale)
 {
+	check(IsInGameThread()); //Would need some lock in map access if we got MT with that.
+
 	hb_font_t* HarfBuzzFont = nullptr;
 
 #if WITH_FREETYPE
-	FT_Face FreeTypeFace = InFace.GetFace();
-	FreeTypeUtils::ApplySizeAndScale(FreeTypeFace, InFontInfo.Size, InFontScale);
+	SCOPE_CYCLE_COUNTER(STAT_HarfBuzzFontCache_ApplyFont);
 
-	hb_font_extents_t HarfBuzzFontExtents;
-	FMemory::Memzero(HarfBuzzFontExtents);
+	FT_Face FreeTypeFace = InFace.GetFace();
+	const uint32 FontRenderSize = FreeTypeUtils::ComputeFontPixelSize(InFontInfo.Size, InFontScale);
+	FreeTypeUtils::ApplySizeAndScale(FreeTypeFace, FontRenderSize);
 
-	// Create a sub-font from the default FreeType implementation so we can override some font functions to provide low-level caching
+	FFontKey FontKey(FreeTypeFace, InGlyphFlags);
+	hb_face_t*& CacheEntry = HarfBuzzFontCacheMap.FindOrAdd(FontKey);
+	if (CacheEntry == nullptr)
 	{
-		hb_font_t* HarfBuzzFTFont = hb_ft_font_create(FreeTypeFace, nullptr);
-		hb_ft_font_set_load_flags(HarfBuzzFTFont, InGlyphFlags);
+		hb_face_t* HarfBuzzFace = hb_ft_face_create(FreeTypeFace,nullptr);
+		CacheEntry = HarfBuzzFace;
+	}
 
-		// The default FreeType implementation doesn't apply the font scale, so we have to do that ourselves (in 16.16 space for maximum precision)
-		{
-			int HarfBuzzFTFontXScale = 0;
-			int HarfBuzzFTFontYScale = 0;
-			hb_font_get_scale(HarfBuzzFTFont, &HarfBuzzFTFontXScale, &HarfBuzzFTFontYScale);
-
-			// Cache the font extents
-			const bool IsAscentDescentOverridenEnabled = UE::Slate::FontUtils::IsAscentDescentOverrideEnabled(InFontInfo.FontObject);
-			HarfBuzzFontExtents.ascender = InFace.GetAscender(IsAscentDescentOverridenEnabled);
-			HarfBuzzFontExtents.descender = InFace.GetDescender(IsAscentDescentOverridenEnabled);
-			HarfBuzzFontExtents.line_gap = InFace.GetScaledHeight(IsAscentDescentOverridenEnabled) - (HarfBuzzFontExtents.ascender - HarfBuzzFontExtents.descender);
-			if (HarfBuzzFTFontYScale < 0)
-			{
-				HarfBuzzFontExtents.ascender = -HarfBuzzFontExtents.ascender;
-				HarfBuzzFontExtents.descender = -HarfBuzzFontExtents.descender;
-				HarfBuzzFontExtents.line_gap = -HarfBuzzFontExtents.line_gap;
-			}
-		}
+	HarfBuzzFont = hb_font_create(CacheEntry);
 
-		HarfBuzzFont = hb_font_create_sub_font(HarfBuzzFTFont);
-		
-		hb_font_destroy(HarfBuzzFTFont);
-	}
+	//The HarfBuzz face doesn't store the size information about the freetype face, so we need to provide the info to the font itself, as it needs it).
+	const int HarfBuzzFontXScale = (int)(((uint64_t)FreeTypeFace->size->metrics.x_scale * (uint64_t)FreeTypeFace->units_per_EM + (1u << 15)) >> 16);
+	const int HarfBuzzFontYScale = (int)(((uint64_t)FreeTypeFace->size->metrics.y_scale * (uint64_t)FreeTypeFace->units_per_EM + (1u << 15)) >> 16);
+	hb_font_set_scale(HarfBuzzFont, HarfBuzzFontXScale, HarfBuzzFontYScale);
 
 	hb_font_set_funcs(HarfBuzzFont, CustomHarfBuzzFuncs, nullptr, nullptr);
 
-	hb_font_set_user_data(
-		HarfBuzzFont, 
-		&HarfBuzzFontFunctions::UserDataKey, 
-		HarfBuzzFontFunctions::CreateUserData(InFontInfo.Size, InFontScale, FTCacheDirectory, HarfBuzzFontExtents),
-		&HarfBuzzFontFunctions::DestroyUserData, 
-		true
-		);
+	HarfBuzzFontFunctions::FUserData* UserData = HarfBuzzFontFunctions::CreateUserData(FontRenderSize, FTCacheDirectory, FreeTypeFace, InGlyphFlags);
+
+	// Apply the current settings of the font to the cache in order to allow harfbuzz to retrieve those values back when shaping the text.
+	hb_font_extents_t& HarfBuzzFontExtents = UserData->HarfBuzzFontExtents;
+
+	// Cache the font extents
+	const bool bIsAscentDescentOverrideEnabled = UE::Slate::FontUtils::IsAscentDescentOverrideEnabled(InFontInfo.FontObject);
+	HarfBuzzFontExtents.ascender = InFace.GetAscender(bIsAscentDescentOverrideEnabled);
+	HarfBuzzFontExtents.descender = InFace.GetDescender(bIsAscentDescentOverrideEnabled);
+	HarfBuzzFontExtents.line_gap = InFace.GetScaledHeight(bIsAscentDescentOverrideEnabled) - (HarfBuzzFontExtents.ascender - HarfBuzzFontExtents.descender);
+	if (HarfBuzzFontYScale < 0)
+	{
+		HarfBuzzFontExtents.ascender = -HarfBuzzFontExtents.ascender;
+		HarfBuzzFontExtents.descender = -HarfBuzzFontExtents.descender;
+		HarfBuzzFontExtents.line_gap = -HarfBuzzFontExtents.line_gap;
+	}
+
+	hb_font_set_user_data(HarfBuzzFont,	&HarfBuzzFontFunctions::UserDataKey, UserData, &HarfBuzzFontFunctions::DestroyUserData,	true);
+
 #endif // WITH_FREETYPE
 
 	return HarfBuzzFont;
 }
 
 #endif // WITH_HARFBUZZ
+
+void FHarfBuzzFontCache::Flush()
+{
+#if WITH_HARFBUZZ
+	for (TTuple<FFontKey, hb_face_t*>& Entry : HarfBuzzFontCacheMap)
+	{
+		hb_face_destroy(Entry.Value);
+	}
+	HarfBuzzFontCacheMap.Empty();
+#endif // WITH_HARFBUZZ
+}
+
diff --git a/Engine/Source/Runtime/SlateCore/Private/Fonts/FontCacheHarfBuzz.h b/Engine/Source/Runtime/SlateCore/Private/Fonts/FontCacheHarfBuzz.h
--- a/Engine/Source/Runtime/SlateCore/Private/Fonts/FontCacheHarfBuzz.h
+++ b/Engine/Source/Runtime/SlateCore/Private/Fonts/FontCacheHarfBuzz.h
@@ -47,20 +47,53 @@ void AppendStringToBuffer(const FStringView InString, const int32 InStartIndex,
 
 } // namespace HarfBuzzUtils
 
-class FHarfBuzzFontFactory
+class FHarfBuzzFontCache
 {
 public:
-	FHarfBuzzFontFactory(FFreeTypeCacheDirectory* InFTCacheDirectory);
-	~FHarfBuzzFontFactory();
+	FHarfBuzzFontCache(FFreeTypeCacheDirectory* InFTCacheDirectory);
+	~FHarfBuzzFontCache();
 
 #if WITH_HARFBUZZ
 	/** Create a HarfBuzz font from the given face - must be destroyed with hb_font_destroy when done */
-	hb_font_t* CreateFont(const FFreeTypeFace& InFace, const uint32 InGlyphFlags, const FSlateFontInfo& InFontInfo, const float InFontScale) const;
+	hb_font_t* CreateFont(const FFreeTypeFace& InFace, const uint32 InGlyphFlags, const FSlateFontInfo& InFontInfo, const float InFontScale);
+
 #endif // WITH_HARFBUZZ
+	/** Clear the cache containing all harfbuzz fonts */
+	void Flush();
 
 private:
 	FFreeTypeCacheDirectory* FTCacheDirectory;
 
+#if WITH_HARFBUZZ && WITH_FREETYPE
+	class FFontKey
+	{
+	public:
+		FFontKey(FT_Face InFace, const int32 InFlags);
+
+		bool operator==(const FFontKey& Other) const
+		{
+			return Face == Other.Face
+				&& Flags == Other.Flags;
+		}
+
+		bool operator!=(const FFontKey& Other) const
+		{
+			return !(*this == Other);
+		}
+
+		friend inline uint32 GetTypeHash(const FFontKey& Key)
+		{
+			return Key.KeyHash;
+		}
+
+	private:
+		FT_Face Face;
+		const int32 Flags;
+		uint32 KeyHash;
+	};
+	TMap<FFontKey, hb_face_t*> HarfBuzzFontCacheMap;
+#endif // WITH_HARFBUZZ && WITH_FREETYPE
+
 #if WITH_HARFBUZZ
 	hb_font_funcs_t* CustomHarfBuzzFuncs;
 #endif // WITH_HARFBUZZ
diff --git a/Engine/Source/Runtime/SlateCore/Private/Fonts/SlateFontRenderer.cpp b/Engine/Source/Runtime/SlateCore/Private/Fonts/SlateFontRenderer.cpp
--- a/Engine/Source/Runtime/SlateCore/Private/Fonts/SlateFontRenderer.cpp
+++ b/Engine/Source/Runtime/SlateCore/Private/Fonts/SlateFontRenderer.cpp
@@ -303,7 +303,8 @@ int8 FSlateFontRenderer::GetKerning(const FFontData& InFontData, const float InS
 	{
 		const FT_UInt FirstIndex = FT_Get_Char_Index(FontFace, First);
 		const FT_UInt SecondIndex = FT_Get_Char_Index(FontFace, Second);
-		TSharedPtr<FFreeTypeKerningCache> KerningCache = GetKerningCache(InFontData, InSize, InScale);
+		const uint32 FontRenderSize = FreeTypeUtils::ComputeFontPixelSize(InSize, InScale);
+		TSharedPtr<FFreeTypeKerningCache> KerningCache = GetKerningCache(InFontData, FontRenderSize);
 		if (KerningCache)
 		{
 			FT_Vector KerningVector;
@@ -318,11 +319,11 @@ int8 FSlateFontRenderer::GetKerning(const FFontData& InFontData, const float InS
 	return 0;
 }
 
-TSharedPtr<FFreeTypeKerningCache> FSlateFontRenderer::GetKerningCache(const FFontData& InFontData, const float InSize, const float InScale) const
+TSharedPtr<FFreeTypeKerningCache> FSlateFontRenderer::GetKerningCache(const FFontData& InFontData, const uint32 FontRenderSize) const
 {
 #if WITH_FREETYPE
 	FT_Face FontFace = GetFontFace(InFontData);
-	return FTCacheDirectory->GetKerningCache(FontFace, FT_KERNING_DEFAULT, InSize, InScale);
+	return FTCacheDirectory->GetKerningCache(FontFace, FT_KERNING_DEFAULT, FontRenderSize);
 #else
 	return nullptr;
 #endif // WITH_FREETYPE
diff --git a/Engine/Source/Runtime/SlateCore/Private/Fonts/SlateFontRenderer.h b/Engine/Source/Runtime/SlateCore/Private/Fonts/SlateFontRenderer.h
--- a/Engine/Source/Runtime/SlateCore/Private/Fonts/SlateFontRenderer.h
+++ b/Engine/Source/Runtime/SlateCore/Private/Fonts/SlateFontRenderer.h
@@ -105,12 +105,11 @@ public:
 	/**
 	 * Retrieves the kerning cache object used to retrieve kerning amounts between glyphs
 	 *
-	 * @param InFontData	The font that used to draw the string with the first and second characters
-	 * @param InSize		The size of the font to draw
-	 * @param InScale		The scale at which the font is being drawn
+	 * @param InFontData		The font that used to draw the string with the first and second characters
+	 * @param InFontRenderSize	The final size of the font to render (including scaling)
 	 * @return A pointer object referring to the kerning cache matching the given parameters, invalid/null if no kerning is performed
 	 */
-	TSharedPtr<FFreeTypeKerningCache> GetKerningCache(const FFontData& InFontData, const float InSize, const float InScale) const;
+	TSharedPtr<FFreeTypeKerningCache> GetKerningCache(const FFontData& InFontData, const uint32 InFontRenderSize) const;
 
 	/**
 	 * Whether or not the specified character, within the specified font, can be loaded with the specified maximum font fallback level
diff --git a/Engine/Source/Runtime/SlateCore/Private/Fonts/SlateTextShaper.cpp b/Engine/Source/Runtime/SlateCore/Private/Fonts/SlateTextShaper.cpp
--- a/Engine/Source/Runtime/SlateCore/Private/Fonts/SlateTextShaper.cpp
+++ b/Engine/Source/Runtime/SlateCore/Private/Fonts/SlateTextShaper.cpp
@@ -13,6 +13,12 @@
 
 DECLARE_CYCLE_STAT(TEXT("Shape Bidirectional Text"), STAT_SlateShapeBidirectionalText, STATGROUP_Slate);
 DECLARE_CYCLE_STAT(TEXT("Shape Unidirectional Text"), STAT_SlateShapeUnidirectionalText, STATGROUP_Slate);
+DECLARE_CYCLE_STAT(TEXT("HarfBuzz Text Shaping"), STAT_PerformHarfBuzzTextShaping, STATGROUP_Slate);
+DECLARE_CYCLE_STAT(TEXT("Kerning Only Text Shaping"), STAT_PerformKerningOnlyTextShaping, STATGROUP_Slate);
+DECLARE_CYCLE_STAT(TEXT("Split By FontFace"), STAT_SplitByFontFace, STATGROUP_Slate);
+DECLARE_CYCLE_STAT(TEXT("Split By Script"), STAT_SplitByScript, STATGROUP_Slate);
+DECLARE_CYCLE_STAT(TEXT("Text Sequences Processing"), STAT_TextSequencesProcessing, STATGROUP_Slate);
+DECLARE_CYCLE_STAT(TEXT("HarfBuzz Shape"), STAT_HarfBuzzNativeShape, STATGROUP_Slate);
 
 namespace
 {
@@ -98,7 +104,7 @@ FSlateTextShaper::FSlateTextShaper(FFreeTypeCacheDirectory* InFTCacheDirectory,
 	, TextBiDiDetection(TextBiDi::CreateTextBiDi())
 	, GraphemeBreakIterator(GSlateUseSharedBreakIterator ? FBreakIterator::GetCharacterBoundaryIterator() : FBreakIterator::CreateCharacterBoundaryIterator())
 #if WITH_HARFBUZZ
-	, HarfBuzzFontFactory(FTCacheDirectory)
+	, HarfBuzzFontCache(FTCacheDirectory)
 #endif // WITH_HARFBUZZ
 {
 	check(FTCacheDirectory);
@@ -107,7 +113,14 @@ FSlateTextShaper::FSlateTextShaper(FFreeTypeCacheDirectory* InFTCacheDirectory,
 	check(FontCache);
 }
 
-FShapedGlyphSequenceRef FSlateTextShaper::ShapeBidirectionalText(const TCHAR* InText, const int32 InTextStart, const int32 InTextLen, const FSlateFontInfo& InFontInfo, const float InFontScale, const TextBiDi::ETextDirection InBaseDirection, const ETextShapingMethod TextShapingMethod) const
+void FSlateTextShaper::FlushCache()
+{
+#if WITH_HARFBUZZ
+	HarfBuzzFontCache.Flush();
+#endif
+}
+
+FShapedGlyphSequenceRef FSlateTextShaper::ShapeBidirectionalText(const TCHAR* InText, const int32 InTextStart, const int32 InTextLen, const FSlateFontInfo& InFontInfo, const float InFontScale, const TextBiDi::ETextDirection InBaseDirection, const ETextShapingMethod TextShapingMethod)
 {
 	SCOPE_CYCLE_COUNTER(STAT_SlateShapeBidirectionalText);
 
@@ -123,7 +136,7 @@ FShapedGlyphSequenceRef FSlateTextShaper::ShapeBidirectionalText(const TCHAR* In
 	return FinalizeTextShaping(MoveTemp(GlyphsToRender), InFontInfo, InFontScale, FShapedGlyphSequence::FSourceTextRange(InTextStart, InTextLen));
 }
 
-FShapedGlyphSequenceRef FSlateTextShaper::ShapeUnidirectionalText(const TCHAR* InText, const int32 InTextStart, const int32 InTextLen, const FSlateFontInfo& InFontInfo, const float InFontScale, const TextBiDi::ETextDirection InTextDirection, const ETextShapingMethod TextShapingMethod) const
+FShapedGlyphSequenceRef FSlateTextShaper::ShapeUnidirectionalText(const TCHAR* InText, const int32 InTextStart, const int32 InTextLen, const FSlateFontInfo& InFontInfo, const float InFontScale, const TextBiDi::ETextDirection InTextDirection, const ETextShapingMethod TextShapingMethod)
 {
 	SCOPE_CYCLE_COUNTER(STAT_SlateShapeUnidirectionalText);
 
@@ -132,7 +145,7 @@ FShapedGlyphSequenceRef FSlateTextShaper::ShapeUnidirectionalText(const TCHAR* I
 	return FinalizeTextShaping(MoveTemp(GlyphsToRender), InFontInfo, InFontScale, FShapedGlyphSequence::FSourceTextRange(InTextStart, InTextLen));
 }
 
-void FSlateTextShaper::PerformTextShaping(const TCHAR* InText, const int32 InTextStart, const int32 InTextLen, const FSlateFontInfo& InFontInfo, const float InFontScale, const TextBiDi::ETextDirection InTextDirection, const ETextShapingMethod TextShapingMethod, TArray<FShapedGlyphEntry>& OutGlyphsToRender) const
+void FSlateTextShaper::PerformTextShaping(const TCHAR* InText, const int32 InTextStart, const int32 InTextLen, const FSlateFontInfo& InFontInfo, const float InFontScale, const TextBiDi::ETextDirection InTextDirection, const ETextShapingMethod TextShapingMethod, TArray<FShapedGlyphEntry>& OutGlyphsToRender)
 {
 	check(InTextDirection != TextBiDi::ETextDirection::Mixed);
 
@@ -322,11 +335,15 @@ FShapedGlyphSequenceRef FSlateTextShaper::FinalizeTextShaping(TArray<FShapedGlyp
 
 void FSlateTextShaper::PerformKerningOnlyTextShaping(const TCHAR* InText, const int32 InTextStart, const int32 InTextLen, const FSlateFontInfo& InFontInfo, const float InFontScale, TArray<FShapedGlyphEntry>& OutGlyphsToRender) const
 {
+	SCOPE_CYCLE_COUNTER(STAT_PerformKerningOnlyTextShaping);
+
 	// We need to work out the correct FFontData for everything so that we can build accurate FShapedGlyphFaceData for rendering later on
 	TArray<FKerningOnlyTextSequenceEntry, TInlineAllocator<4>> KerningOnlyTextSequence;
 
 	// Step 1) Split the text into sections that are using the same font face (composite fonts may contain different faces for different character ranges)
 	{
+		SCOPE_CYCLE_COUNTER(STAT_SplitByFontFace);
+
 		// Data used while detecting font face boundaries
 		int32 SplitStartIndex = InTextStart;
 		int32 RunningTextIndex = InTextStart;
@@ -395,6 +412,8 @@ void FSlateTextShaper::PerformKerningOnlyTextShaping(const TCHAR* InText, const
 
 	// Step 2) Now we use the font cache to get the size for each character, and kerning for each character pair
 	{
+		SCOPE_CYCLE_COUNTER(STAT_TextSequencesProcessing);
+
 		OutGlyphsToRender.Reserve(OutGlyphsToRender.Num() + InTextLen);
 		for (const FKerningOnlyTextSequenceEntry& KerningOnlyTextSequenceEntry : KerningOnlyTextSequence)
 		{
@@ -425,9 +444,10 @@ void FSlateTextShaper::PerformKerningOnlyTextShaping(const TCHAR* InText, const
 
 			FreeTypeUtils::ApplySizeAndScale(KerningOnlyTextSequenceEntry.FaceAndMemory->GetFace(), InFontInfo.Size, FinalFontScale);
 			const FFontRasterizationSettings FontRasterizationSettings = KerningOnlyTextSequenceEntry.FontDataPtr->GetFontRasterizationSettings();
+			const uint32 FontRenderSize = FreeTypeUtils::ComputeFontPixelSize(InFontInfo.Size, FinalFontScale);
 			TSharedRef<FShapedGlyphFaceData> ShapedGlyphFaceData = MakeShared<FShapedGlyphFaceData>(KerningOnlyTextSequenceEntry.FaceAndMemory, GlyphFlags, InFontInfo.Size, FinalFontScale, InFontInfo.GetClampSkew(), FontRasterizationSettings.Mode, FontRasterizationSettings.DistanceFieldPpem);
-			TSharedPtr<FFreeTypeKerningCache> KerningCache = FTCacheDirectory->GetKerningCache(KerningOnlyTextSequenceEntry.FaceAndMemory->GetFace(), FT_KERNING_DEFAULT, InFontInfo.Size, FinalFontScale);
-			TSharedRef<FFreeTypeAdvanceCache> AdvanceCache = FTCacheDirectory->GetAdvanceCache(KerningOnlyTextSequenceEntry.FaceAndMemory->GetFace(), GlyphFlags, InFontInfo.Size, FinalFontScale);
+			TSharedPtr<FFreeTypeKerningCache> KerningCache = FTCacheDirectory->GetKerningCache(KerningOnlyTextSequenceEntry.FaceAndMemory->GetFace(), FT_KERNING_DEFAULT, FontRenderSize);
+			TSharedRef<FFreeTypeAdvanceCache> AdvanceCache = FTCacheDirectory->GetAdvanceCache(KerningOnlyTextSequenceEntry.FaceAndMemory->GetFace(), GlyphFlags, FontRenderSize);
 
 			for (int32 SequenceCharIndex = 0; SequenceCharIndex < KerningOnlyTextSequenceEntry.TextLength; ++SequenceCharIndex)
 			{
@@ -499,11 +519,12 @@ void FSlateTextShaper::PerformKerningOnlyTextShaping(const TCHAR* InText, const
 
 #if WITH_HARFBUZZ
 
-void FSlateTextShaper::PerformHarfBuzzTextShaping(const TCHAR* InText, const int32 InTextStart, const int32 InTextLen, const FSlateFontInfo& InFontInfo, const float InFontScale, const TextBiDi::ETextDirection InTextDirection, TArray<FShapedGlyphEntry>& OutGlyphsToRender) const
+void FSlateTextShaper::PerformHarfBuzzTextShaping(const TCHAR* InText, const int32 InTextStart, const int32 InTextLen, const FSlateFontInfo& InFontInfo, const float InFontScale, const TextBiDi::ETextDirection InTextDirection, TArray<FShapedGlyphEntry>& OutGlyphsToRender)
 {
+	SCOPE_CYCLE_COUNTER(STAT_PerformHarfBuzzTextShaping);
+
 	// HarfBuzz can only shape data that uses the same font face, reads in the same direction, and uses the same script so we need to split the given text...
-	TArray<FHarfBuzzTextSequenceEntry, TInlineAllocator<4>> HarfBuzzTextSequence;
-	hb_unicode_funcs_t* HarfBuzzUnicodeFuncs = hb_unicode_funcs_get_default();
+	TArray<FHarfBuzzTextSequenceEntry, TInlineAllocator<12>> HarfBuzzTextSequence;
 
 	GraphemeBreakIterator->SetString(InText + InTextStart, InTextLen);
 
@@ -512,6 +533,8 @@ void FSlateTextShaper::PerformHarfBuzzTextShaping(const TCHAR* InText, const int
 
 	// Step 1) Split the text into sections that are using the same font face (composite fonts may contain different faces for different character ranges)
 	{
+		SCOPE_CYCLE_COUNTER(STAT_SplitByFontFace);
+
 		// Data used while detecting font face boundaries
 		int32 SplitStartIndex = InTextStart;
 		int32 RunningTextIndex = InTextStart;
@@ -596,71 +619,76 @@ void FSlateTextShaper::PerformHarfBuzzTextShaping(const TCHAR* InText, const int
 		AppendPendingFontDataToSequence();
 	}
 
-	// Step 2) Split the font face sections by their their script code
-	for (FHarfBuzzTextSequenceEntry& HarfBuzzTextSequenceEntry : HarfBuzzTextSequence)
+	// Step 2) Split the font face sections by their script code
 	{
-		// Data used while detecting script code boundaries
-		int32 SplitStartIndex = HarfBuzzTextSequenceEntry.TextStartIndex;
-		int32 RunningTextIndex = HarfBuzzTextSequenceEntry.TextStartIndex;
-		TOptional<hb_script_t> RunningHarfBuzzScript;
+		SCOPE_CYCLE_COUNTER(STAT_SplitByScript);
+		hb_unicode_funcs_t* HarfBuzzUnicodeFuncs = hb_unicode_funcs_get_default();
 
-		auto StartNewPendingTextSequence = [&](const hb_script_t InHarfBuzzScript)
+		for (FHarfBuzzTextSequenceEntry& HarfBuzzTextSequenceEntry : HarfBuzzTextSequence)
 		{
-			SplitStartIndex = RunningTextIndex;
-			RunningHarfBuzzScript = InHarfBuzzScript;
-		};
+			// Data used while detecting script code boundaries
+			int32 SplitStartIndex = HarfBuzzTextSequenceEntry.TextStartIndex;
+			int32 RunningTextIndex = HarfBuzzTextSequenceEntry.TextStartIndex;
+			TOptional<hb_script_t> RunningHarfBuzzScript;
 
-		auto AppendPendingTextToSequence = [&]()
-		{
-			if (RunningHarfBuzzScript.IsSet())
+			auto StartNewPendingTextSequence = [&](const hb_script_t InHarfBuzzScript)
 			{
-				HarfBuzzTextSequenceEntry.SubSequence.Emplace(
-					SplitStartIndex,					// InStartIndex
-					RunningTextIndex - SplitStartIndex,	// InLength
-					RunningHarfBuzzScript.GetValue()	// InHarfBuzzScript
-					);
+				SplitStartIndex = RunningTextIndex;
+				RunningHarfBuzzScript = InHarfBuzzScript;
+			};
 
-				RunningHarfBuzzScript.Reset();
-			}
-		};
+			auto AppendPendingTextToSequence = [&]()
+			{
+				if (RunningHarfBuzzScript.IsSet())
+				{
+					HarfBuzzTextSequenceEntry.SubSequence.Emplace(
+						SplitStartIndex,					// InStartIndex
+						RunningTextIndex - SplitStartIndex,	// InLength
+						RunningHarfBuzzScript.GetValue()	// InHarfBuzzScript
+						);
 
-		auto IsSpecialCharacter = [](const hb_script_t InHarfBuzzScript) -> bool
-		{
-			// Characters in the common, inherited, and unknown scripts are allowed (and in the case of inherited, required) to merge with the script 
-			// of the character(s) that preceded them. This also helps to minimize shaping batches, as spaces are within the common script.
-			return InHarfBuzzScript == HB_SCRIPT_COMMON || InHarfBuzzScript == HB_SCRIPT_INHERITED || InHarfBuzzScript == HB_SCRIPT_UNKNOWN;
-		};
+					RunningHarfBuzzScript.Reset();
+				}
+			};
 
-		const int32 TextEndIndex = HarfBuzzTextSequenceEntry.TextStartIndex + HarfBuzzTextSequenceEntry.TextLength;
-		for (; RunningTextIndex < TextEndIndex; ++RunningTextIndex)
-		{
-			const hb_script_t CharHarfBuzzScript = hb_unicode_script(HarfBuzzUnicodeFuncs, InText[RunningTextIndex]);				
+			auto IsSpecialCharacter = [](const hb_script_t InHarfBuzzScript) -> bool
+			{
+				// Characters in the common, inherited, and unknown scripts are allowed (and in the case of inherited, required) to merge with the script 
+				// of the character(s) that preceded them. This also helps to minimize shaping batches, as spaces are within the common script.
+				return InHarfBuzzScript == HB_SCRIPT_COMMON || InHarfBuzzScript == HB_SCRIPT_INHERITED || InHarfBuzzScript == HB_SCRIPT_UNKNOWN;
+			};
 
-			if (!RunningHarfBuzzScript.IsSet() || RunningHarfBuzzScript.GetValue() != CharHarfBuzzScript)
+			const int32 TextEndIndex = HarfBuzzTextSequenceEntry.TextStartIndex + HarfBuzzTextSequenceEntry.TextLength;
+			for (; RunningTextIndex < TextEndIndex; ++RunningTextIndex)
 			{
-				if (!RunningHarfBuzzScript.IsSet())
-				{
-					// Always start a new run if we're currently un-set
-					StartNewPendingTextSequence(CharHarfBuzzScript);
-				}
-				else if (!IsSpecialCharacter(CharHarfBuzzScript))
+				const hb_script_t CharHarfBuzzScript = hb_unicode_script(HarfBuzzUnicodeFuncs, InText[RunningTextIndex]);
+
+				if (!RunningHarfBuzzScript.IsSet() || RunningHarfBuzzScript.GetValue() != CharHarfBuzzScript)
 				{
-					if (IsSpecialCharacter(RunningHarfBuzzScript.GetValue()))
+					if (!RunningHarfBuzzScript.IsSet())
 					{
-						// If we started our run on a special character, we need to swap the script type to the non-special type as soon as we can
-						RunningHarfBuzzScript = CharHarfBuzzScript;
+						// Always start a new run if we're currently un-set
+						StartNewPendingTextSequence(CharHarfBuzzScript);
 					}
-					else
+					else if (!IsSpecialCharacter(CharHarfBuzzScript))
 					{
-						// Transitioned a non-special character; end the current run and create a new one
-						AppendPendingTextToSequence();
-						StartNewPendingTextSequence(CharHarfBuzzScript);
+						if (IsSpecialCharacter(RunningHarfBuzzScript.GetValue()))
+						{
+							// If we started our run on a special character, we need to swap the script type to the non-special type as soon as we can
+							RunningHarfBuzzScript = CharHarfBuzzScript;
+						}
+						else
+						{
+							// Transitioned a non-special character; end the current run and create a new one
+							AppendPendingTextToSequence();
+							StartNewPendingTextSequence(CharHarfBuzzScript);
+						}
 					}
 				}
 			}
-		}
 
-		AppendPendingTextToSequence();
+			AppendPendingTextToSequence();
+		}
 	}
 
 	if (InTextDirection == TextBiDi::ETextDirection::RightToLeft)
@@ -672,8 +700,11 @@ void FSlateTextShaper::PerformHarfBuzzTextShaping(const TCHAR* InText, const int
 
 	// Step 3) Now we use HarfBuzz to shape each font data sequence using its FreeType glyph
 	{
+		SCOPE_CYCLE_COUNTER(STAT_TextSequencesProcessing);
+
 		hb_buffer_t* HarfBuzzTextBuffer = hb_buffer_create();
 		const FStringView InTextView = InText; // This will do a strlen, so do it once at the start
+		const bool bUseLetterSpacing = !bBypassLetterSpacing && InFontInfo.LetterSpacing != 0;
 
 		for (const FHarfBuzzTextSequenceEntry& HarfBuzzTextSequenceEntry : HarfBuzzTextSequence)
 		{
@@ -683,15 +714,15 @@ void FSlateTextShaper::PerformHarfBuzzTextShaping(const TCHAR* InText, const int
 			}
 
 #if WITH_FREETYPE
-			const bool bHasKerning = FT_HAS_KERNING(HarfBuzzTextSequenceEntry.FaceAndMemory->GetFace()) != 0 || (!bBypassLetterSpacing && InFontInfo.LetterSpacing != 0);
+			const bool bHasKerning = FT_HAS_KERNING(HarfBuzzTextSequenceEntry.FaceAndMemory->GetFace()) != 0 || bUseLetterSpacing;
 #else  // WITH_FREETYPE
 			const bool bHasKerning = false;
 #endif // WITH_FREETYPE
 
 			// Letter spacing should scale proportional to font size / 1000 (to roughly mimic Photoshop tracking)
-			const int32 LetterSpacingScaledAsInt = (!bBypassLetterSpacing && InFontInfo.LetterSpacing != 0) ? InFontInfo.LetterSpacing * InFontInfo.Size / 1000 : 0;
+			const int32 LetterSpacingScaledAsInt = bUseLetterSpacing ? InFontInfo.LetterSpacing * InFontInfo.Size / 1000 : 0;
 			ensure(LetterSpacingScaledAsInt <= std::numeric_limits<int16>::max());
-			const int16 LetterSpacingScaled = (!bBypassLetterSpacing && InFontInfo.LetterSpacing != 0) ? (int16)(LetterSpacingScaledAsInt) : 0;
+			const int16 LetterSpacingScaled = bUseLetterSpacing ? (int16)(LetterSpacingScaledAsInt) : 0;
 
 			const hb_feature_t HarfBuzzFeatures[] = {
 				{ HB_TAG('k','e','r','n'), bHasKerning, 0, uint32(-1) },
@@ -703,7 +734,7 @@ void FSlateTextShaper::PerformHarfBuzzTextShaping(const TCHAR* InText, const int
 			SlateFontRendererUtils::AppendGlyphFlags(*HarfBuzzTextSequenceEntry.FaceAndMemory, *HarfBuzzTextSequenceEntry.FontDataPtr, GlyphFlags);
 			const float FinalFontScale = InFontScale * HarfBuzzTextSequenceEntry.SubFontScalingFactor;
 
-			hb_font_t* HarfBuzzFont = HarfBuzzFontFactory.CreateFont(*HarfBuzzTextSequenceEntry.FaceAndMemory, GlyphFlags, InFontInfo, FinalFontScale);
+			hb_font_t* HarfBuzzFont = HarfBuzzFontCache.CreateFont(*HarfBuzzTextSequenceEntry.FaceAndMemory, GlyphFlags, InFontInfo, FinalFontScale);
 			if (!HarfBuzzFont)
 			{
 				continue;
@@ -716,14 +747,13 @@ void FSlateTextShaper::PerformHarfBuzzTextShaping(const TCHAR* InText, const int
 																									InFontInfo.GetClampSkew(),
 																									FontRasterizationSettings.Mode,
 																									FontRasterizationSettings.DistanceFieldPpem);
+			const uint32 FontRenderSize = FreeTypeUtils::ComputeFontPixelSize(InFontInfo.Size, FinalFontScale);
 			TSharedPtr<FFreeTypeKerningCache> KerningCache = FTCacheDirectory->GetKerningCache(HarfBuzzTextSequenceEntry.FaceAndMemory->GetFace(), 
 																							   FT_KERNING_DEFAULT,
-																							   InFontInfo.Size, 
-																							   FinalFontScale);
+																							   FontRenderSize);
 			TSharedRef<FFreeTypeAdvanceCache> AdvanceCache = FTCacheDirectory->GetAdvanceCache(HarfBuzzTextSequenceEntry.FaceAndMemory->GetFace(), 
 																							   GlyphFlags,
-																							   InFontInfo.Size, 
-																							   FinalFontScale);
+																							   FontRenderSize);
 
 			for (const FHarfBuzzTextSequenceEntry::FSubSequenceEntry& HarfBuzzTextSubSequenceEntry : HarfBuzzTextSequenceEntry.SubSequence)
 			{
@@ -734,7 +764,10 @@ void FSlateTextShaper::PerformHarfBuzzTextShaping(const TCHAR* InText, const int
 				hb_buffer_set_script(HarfBuzzTextBuffer, HarfBuzzTextSubSequenceEntry.HarfBuzzScript);
 
 				HarfBuzzUtils::AppendStringToBuffer(InTextView, HarfBuzzTextSubSequenceEntry.StartIndex, HarfBuzzTextSubSequenceEntry.Length, HarfBuzzTextBuffer);
-				hb_shape(HarfBuzzFont, HarfBuzzTextBuffer, HarfBuzzFeatures, HarfBuzzFeaturesCount);
+				{
+					SCOPE_CYCLE_COUNTER(STAT_HarfBuzzNativeShape);
+					hb_shape(HarfBuzzFont, HarfBuzzTextBuffer, HarfBuzzFeatures, HarfBuzzFeaturesCount);
+				}
 
 				uint32 HarfBuzzGlyphCount = 0;
 				hb_glyph_info_t* HarfBuzzGlyphInfos = hb_buffer_get_glyph_infos(HarfBuzzTextBuffer, &HarfBuzzGlyphCount);
@@ -907,9 +940,9 @@ void FSlateTextShaper::PerformHarfBuzzTextShaping(const TCHAR* InText, const int
 
 #endif // WITH_HARFBUZZ
 
-bool FSlateTextShaper::InsertSubstituteGlyphs(const TCHAR* InText, const int32 InCharIndex, const TSharedRef<FShapedGlyphFaceData>& InShapedGlyphFaceData, const TSharedRef<FFreeTypeAdvanceCache>& AdvanceCache, TArray<FShapedGlyphEntry>& OutGlyphsToRender, const int16 InLetterSpacingScaled) const
+bool FSlateTextShaper::InsertSubstituteGlyphs(const TCHAR* InText, const int32 InCharIndex, const TSharedRef<FShapedGlyphFaceData>& InShapedGlyphFaceData, const TSharedRef<FFreeTypeAdvanceCache>& AdvanceCache, TArray<FShapedGlyphEntry>& OutGlyphsToRender, const int16 InLetterSpacingScaled)
 {
-	auto GetSpecifiedGlyphIndexAndAdvance = [this, &InShapedGlyphFaceData, &AdvanceCache](TCHAR Char, uint32& OutSpaceGlyphIndex, int16& OutSpaceXAdvance)
+	auto GetSpecifiedGlyphIndexAndAdvance = [&InShapedGlyphFaceData, &AdvanceCache](TCHAR Char, uint32& OutSpaceGlyphIndex, int16& OutSpaceXAdvance)
 	{
 		OutSpaceGlyphIndex = 0;
 		OutSpaceXAdvance = 0;
diff --git a/Engine/Source/Runtime/SlateCore/Private/Fonts/SlateTextShaper.h b/Engine/Source/Runtime/SlateCore/Private/Fonts/SlateTextShaper.h
--- a/Engine/Source/Runtime/SlateCore/Private/Fonts/SlateTextShaper.h
+++ b/Engine/Source/Runtime/SlateCore/Private/Fonts/SlateTextShaper.h
@@ -65,12 +65,15 @@ class FSlateTextShaper
 public:
 	FSlateTextShaper(FFreeTypeCacheDirectory* InFTCacheDirectory, FCompositeFontCache* InCompositeFontCache, FSlateFontRenderer* InFontRenderer, FSlateFontCache* InFontCache);
 
-	FShapedGlyphSequenceRef ShapeBidirectionalText(const TCHAR* InText, const int32 InTextStart, const int32 InTextLen, const FSlateFontInfo& InFontInfo, const float InFontScale, const TextBiDi::ETextDirection InBaseDirection, const ETextShapingMethod TextShapingMethod) const;
+	FShapedGlyphSequenceRef ShapeBidirectionalText(const TCHAR* InText, const int32 InTextStart, const int32 InTextLen, const FSlateFontInfo& InFontInfo, const float InFontScale, const TextBiDi::ETextDirection InBaseDirection, const ETextShapingMethod TextShapingMethod);
 
-	FShapedGlyphSequenceRef ShapeUnidirectionalText(const TCHAR* InText, const int32 InTextStart, const int32 InTextLen, const FSlateFontInfo& InFontInfo, const float InFontScale, const TextBiDi::ETextDirection InTextDirection, const ETextShapingMethod TextShapingMethod) const;
+	FShapedGlyphSequenceRef ShapeUnidirectionalText(const TCHAR* InText, const int32 InTextStart, const int32 InTextLen, const FSlateFontInfo& InFontInfo, const float InFontScale, const TextBiDi::ETextDirection InTextDirection, const ETextShapingMethod TextShapingMethod);
+
+	/** Get rid of the harfbuzz font cache, wich store the harfbuzz fonts. */
+	void FlushCache();
 
 private:
-	void PerformTextShaping(const TCHAR* InText, const int32 InTextStart, const int32 InTextLen, const FSlateFontInfo& InFontInfo, const float InFontScale, const TextBiDi::ETextDirection InTextDirection, const ETextShapingMethod TextShapingMethod, TArray<FShapedGlyphEntry>& OutGlyphsToRender) const;
+	void PerformTextShaping(const TCHAR* InText, const int32 InTextStart, const int32 InTextLen, const FSlateFontInfo& InFontInfo, const float InFontScale, const TextBiDi::ETextDirection InTextDirection, const ETextShapingMethod TextShapingMethod, TArray<FShapedGlyphEntry>& OutGlyphsToRender);
 
 	FShapedGlyphSequenceRef FinalizeTextShaping(TArray<FShapedGlyphEntry> InGlyphsToRender, const FSlateFontInfo& InFontInfo, const float InFontScale, const FShapedGlyphSequence::FSourceTextRange& InSourceTextRange) const;
 
@@ -79,10 +82,10 @@ private:
 #endif // WITH_FREETYPE
 
 #if WITH_HARFBUZZ
-	void PerformHarfBuzzTextShaping(const TCHAR* InText, const int32 InTextStart, const int32 InTextLen, const FSlateFontInfo& InFontInfo, const float InFontScale, const TextBiDi::ETextDirection InTextDirection, TArray<FShapedGlyphEntry>& OutGlyphsToRender) const;
+	void PerformHarfBuzzTextShaping(const TCHAR* InText, const int32 InTextStart, const int32 InTextLen, const FSlateFontInfo& InFontInfo, const float InFontScale, const TextBiDi::ETextDirection InTextDirection, TArray<FShapedGlyphEntry>& OutGlyphsToRender);
 #endif // WITH_HARFBUZZ
 
-	bool InsertSubstituteGlyphs(const TCHAR* InText, const int32 InCharIndex, const TSharedRef<FShapedGlyphFaceData>& InShapedGlyphFaceData, const TSharedRef<FFreeTypeAdvanceCache>& AdvanceCache, TArray<FShapedGlyphEntry>& OutGlyphsToRender, const int16 InLetterSpacingScaled = 0) const;
+	static bool InsertSubstituteGlyphs(const TCHAR* InText, const int32 InCharIndex, const TSharedRef<FShapedGlyphFaceData>& InShapedGlyphFaceData, const TSharedRef<FFreeTypeAdvanceCache>& AdvanceCache, TArray<FShapedGlyphEntry>& OutGlyphsToRender, const int16 InLetterSpacingScaled = 0);
 
 	FFreeTypeCacheDirectory* FTCacheDirectory;
 	FCompositeFontCache* CompositeFontCache;
@@ -96,7 +99,7 @@ private:
 	TSharedRef<IBreakIterator> GraphemeBreakIterator;
 
 #if WITH_HARFBUZZ
-	/** HarfBuzz font factory (using our cached functions rather than FreeType directly) */
-	FHarfBuzzFontFactory HarfBuzzFontFactory;
+	/** HarfBuzz font cache (using our cached functions rather than FreeType directly) */
+	FHarfBuzzFontCache HarfBuzzFontCache;
 #endif // WITH_HARFBUZZ
 };
