Slate event binding -> use delegate bind not create + assign.

Less code and faster. Should've done this when I did the same for attributes but didn't see the events then.

diff --git a/Engine/Source/Runtime/SlateCore/Public/Widgets/DeclarativeSyntaxSupport.h b/Engine/Source/Runtime/SlateCore/Public/Widgets/DeclarativeSyntaxSupport.h
--- a/Engine/Source/Runtime/SlateCore/Public/Widgets/DeclarativeSyntaxSupport.h
+++ b/Engine/Source/Runtime/SlateCore/Public/Widgets/DeclarativeSyntaxSupport.h
@@ -474,7 +474,7 @@ struct NamedSlotProperty
 		template< typename StaticFuncPtr, typename... VarTypes > \
 		WidgetArgsType& EventName##_Static( StaticFuncPtr InFunc, VarTypes... Vars )	\
 		{ \
-			_##EventName = DelegateName::CreateStatic( InFunc, Vars... ); \
+			_##EventName.BindStatic( InFunc, Vars... ); \
 			return static_cast<WidgetArgsType*>(this)->Me(); \
 		} \
 		\
@@ -483,7 +483,7 @@ struct NamedSlotProperty
 		template<typename FunctorType, typename... VarTypes> \
 		WidgetArgsType& EventName##_Lambda(FunctorType&& InFunctor, VarTypes... Vars) \
 		{ \
-			_##EventName = DelegateName::CreateLambda( Forward<FunctorType>(InFunctor), Vars... ); \
+			_##EventName.BindLambda( Forward<FunctorType>(InFunctor), Vars... ); \
 			return static_cast<WidgetArgsType*>(this)->Me(); \
 		} \
 		\
@@ -491,13 +491,13 @@ struct NamedSlotProperty
 		template< class UserClass, typename... VarTypes >	\
 		WidgetArgsType& EventName##_Raw( UserClass* InUserObject, typename DelegateName::template TMethodPtr< UserClass, VarTypes... > InFunc, VarTypes... Vars )	\
 		{ \
-			_##EventName = DelegateName::CreateRaw( InUserObject, InFunc, Vars... ); \
+			_##EventName.BindRaw( InUserObject, InFunc, Vars... ); \
 			return static_cast<WidgetArgsType*>(this)->Me(); \
 		} \
 		template< class UserClass, typename... VarTypes >	\
 		WidgetArgsType& EventName##_Raw( UserClass* InUserObject, typename DelegateName::template TConstMethodPtr< UserClass, VarTypes... > InFunc, VarTypes... Vars )	\
 		{ \
-			_##EventName = DelegateName::CreateRaw( InUserObject, InFunc, Vars... ); \
+			_##EventName.BindRaw( InUserObject, InFunc, Vars... ); \
 			return static_cast<WidgetArgsType*>(this)->Me(); \
 		} \
 		\
@@ -505,13 +505,13 @@ struct NamedSlotProperty
 		template< class UserClass, typename... VarTypes >	\
 		WidgetArgsType& EventName( TSharedRef< UserClass > InUserObjectRef, typename DelegateName::template TMethodPtr< UserClass, VarTypes... > InFunc, VarTypes... Vars )	\
 		{ \
-			_##EventName = DelegateName::CreateSP( InUserObjectRef, InFunc, Vars... ); \
+			_##EventName.BindSP( InUserObjectRef, InFunc, Vars... ); \
 			return static_cast<WidgetArgsType*>(this)->Me(); \
 		} \
 		template< class UserClass, typename... VarTypes >	\
 		WidgetArgsType& EventName( TSharedRef< UserClass > InUserObjectRef, typename DelegateName::template TConstMethodPtr< UserClass, VarTypes... > InFunc, VarTypes... Vars )	\
 		{ \
-			_##EventName = DelegateName::CreateSP( InUserObjectRef, InFunc, Vars... ); \
+			_##EventName.BindSP( InUserObjectRef, InFunc, Vars... ); \
 			return static_cast<WidgetArgsType*>(this)->Me(); \
 		} \
 		\
@@ -519,13 +519,13 @@ struct NamedSlotProperty
 		template< class UserClass, typename... VarTypes >	\
 		WidgetArgsType& EventName( UserClass* InUserObject, typename DelegateName::template TMethodPtr< UserClass, VarTypes... > InFunc, VarTypes... Vars )	\
 		{ \
-			_##EventName = DelegateName::CreateSP( InUserObject, InFunc, Vars... ); \
+			_##EventName.BindSP( InUserObject, InFunc, Vars... ); \
 			return static_cast<WidgetArgsType*>(this)->Me(); \
 		} \
 		template< class UserClass, typename... VarTypes >	\
 		WidgetArgsType& EventName( UserClass* InUserObject, typename DelegateName::template TConstMethodPtr< UserClass, VarTypes... > InFunc, VarTypes... Vars )	\
 		{ \
-			_##EventName = DelegateName::CreateSP( InUserObject, InFunc, Vars... ); \
+			_##EventName.BindSP( InUserObject, InFunc, Vars... ); \
 			return static_cast<WidgetArgsType*>(this)->Me(); \
 		} \
 		\
@@ -533,13 +533,13 @@ struct NamedSlotProperty
 		template< class UserClass, typename... VarTypes >	\
 		WidgetArgsType& EventName##_UObject( UserClass* InUserObject, typename DelegateName::template TMethodPtr< UserClass, VarTypes... > InFunc, VarTypes... Vars )	\
 		{ \
-			_##EventName = DelegateName::CreateUObject( InUserObject, InFunc, Vars... ); \
+			_##EventName.BindUObject( InUserObject, InFunc, Vars... ); \
 			return static_cast<WidgetArgsType*>(this)->Me(); \
 		} \
 		template< class UserClass, typename... VarTypes >	\
 		WidgetArgsType& EventName##_UObject( UserClass* InUserObject, typename DelegateName::template TConstMethodPtr< UserClass, VarTypes... > InFunc, VarTypes... Vars )	\
 		{ \
-			_##EventName = DelegateName::CreateUObject( InUserObject, InFunc, Vars... ); \
+			_##EventName.BindUObject( InUserObject, InFunc, Vars... ); \
 			return static_cast<WidgetArgsType*>(this)->Me(); \
 		} \
 		\
