Slate Paint and Prepass Debug View

This CL provides a clearer Slate Paint and Prepass Debug view allowing the display of the recently painted/prepassed widgets in red, and the older ones in green. The goal is to provide a visual representation of when widgets are painted or prepassed.

These views can be enabled with the CVars `SlateDebugger.Prepass.Enable` and `SlateDebugger.Paint.Enable`. Their respective namespaces have additional configuration options like drawing the name of the widgets that have been updated this frame.

Additionally, these debug views and some selected options can be enabled in the WidgetReflector in the `Debug Options` dropdown.

diff --git a/Engine/Source/Developer/SlateReflector/Private/Widgets/SSlateOptions.cpp b/Engine/Source/Developer/SlateReflector/Private/Widgets/SSlateOptions.cpp
--- a/Engine/Source/Developer/SlateReflector/Private/Widgets/SSlateOptions.cpp
+++ b/Engine/Source/Developer/SlateReflector/Private/Widgets/SSlateOptions.cpp
@@ -33,13 +33,17 @@ void SSlateOptions::Construct( const FArguments& InArgs )
 				TooltipText.AppendLine(FString(CVar->GetHelp()));
 				TooltipText.AppendLine(FString(ConsoleVariable));
 
-
 				MenuBuilder.AddMenuEntry(
 					Label,
 					TooltipText.ToText(),
 					Icon,
 					FUIAction(
-						FExecuteAction::CreateLambda([CVar]() { CVar->Set(!CVar->GetBool(), EConsoleVariableFlags::ECVF_SetByCode); }),
+						FExecuteAction::CreateLambda([CVar]()
+						{
+							// If the variable was changed in the command line, we might not be able to change it anymore from the UI so we remove the flag
+							CVar->ClearFlags(EConsoleVariableFlags::ECVF_SetByCommandline);
+							CVar->Set(!CVar->GetBool(), EConsoleVariableFlags::ECVF_SetByCode);
+						}),
 						FCanExecuteAction::CreateLambda([bCanEdit](){ return bCanEdit; }),
 						FGetActionCheckState::CreateLambda([CVar]() { return CVar->GetBool() ? ECheckBoxState::Checked : ECheckBoxState::Unchecked; })
 					),
@@ -47,6 +51,35 @@ void SSlateOptions::Construct( const FArguments& InArgs )
 					EUserInterfaceActionType::ToggleButton);
 			}
 		}
+		static void AddSubMenuEntry(FMenuBuilder& MenuBuilder, const FSlateIcon& Icon, const FText& Label, const TCHAR* ConsoleVariable, const FNewMenuDelegate& InSubMenu, bool bCanEdit = true)
+		{
+			IConsoleVariable* CVar = IConsoleManager::Get().FindConsoleVariable(ConsoleVariable);
+			if (CVar)
+			{
+				FTextBuilder TooltipText;
+				TooltipText.AppendLine(FString(CVar->GetHelp()));
+				TooltipText.AppendLine(FString(ConsoleVariable));
+
+				MenuBuilder.AddSubMenu(
+					Label,
+					TooltipText.ToText(),
+					InSubMenu,
+					FUIAction(
+					FExecuteAction::CreateLambda([CVar]() {
+							// If the variable was changed in the command line, we might not be able to change it anymore from the UI so we remove the flag
+							CVar->ClearFlags(EConsoleVariableFlags::ECVF_SetByCommandline);
+							CVar->Set(!CVar->GetBool(), EConsoleVariableFlags::ECVF_SetByCode);
+						}),
+						FCanExecuteAction::CreateLambda([bCanEdit](){ return bCanEdit; }),
+						FGetActionCheckState::CreateLambda([CVar]() { return CVar->GetBool() ? ECheckBoxState::Checked : ECheckBoxState::Unchecked; })
+					),
+					NAME_None,
+					EUserInterfaceActionType::ToggleButton,
+					false,
+					Icon
+				);
+			}
+		}
 	};
 
 	struct GlobalLocal : Local
@@ -70,6 +103,23 @@ void SSlateOptions::Construct( const FArguments& InArgs )
 
 	struct DebugLocal : Local
 	{
+		static void FillPaintSubmenu(FMenuBuilder& MenuBuilder)
+		{
+			FSlateIcon Icon(FWidgetReflectorStyle::GetStyleSetName(), "Icon.Empty");
+			AddMenuEntry(MenuBuilder, Icon, LOCTEXT("Paint_OnlyGameWindow", "Only Game Window"), TEXT("SlateDebugger.Paint.OnlyGameWindow"));
+			AddMenuEntry(MenuBuilder, Icon, LOCTEXT("Paint_ToggleWidgetNameList", "Toggle Widget Name List"), TEXT("SlateDebugger.Paint.EnableWidgetNameList"));
+			AddMenuEntry(MenuBuilder, Icon, LOCTEXT("Paint_Draw Border", "Draw Border"), TEXT("SlateDebugger.Paint.DrawBorder"));
+			AddMenuEntry(MenuBuilder, Icon, LOCTEXT("Paint_Draw Fill", "Draw Fill"), TEXT("SlateDebugger.Paint.DrawFill"));
+		}
+		static void FillPrepassSubmenu(FMenuBuilder& MenuBuilder)
+		{
+			FSlateIcon Icon(FWidgetReflectorStyle::GetStyleSetName(), "Icon.Empty");
+			AddMenuEntry(MenuBuilder, Icon, LOCTEXT("Prepass_OnlyGameWindow", "Only Game Window"), TEXT("SlateDebugger.Prepass.OnlyGameWindow"));
+			AddMenuEntry(MenuBuilder, Icon, LOCTEXT("Prepass_ToggleWidgetNameList", "Toggle Widget Name List"), TEXT("SlateDebugger.Prepass.EnableWidgetNameList"));
+			AddMenuEntry(MenuBuilder, Icon, LOCTEXT("Prepass_Draw Border", "Draw Border"), TEXT("SlateDebugger.Prepass.DrawBorder"));
+			AddMenuEntry(MenuBuilder, Icon, LOCTEXT("Prepass_Draw Fill", "Draw Fill"), TEXT("SlateDebugger.Prepass.DrawFill"));
+		}
+		
 		static TSharedRef<SWidget> FillToolbar()
 		{
 			const bool bShouldCloseWindowAfterMenuSelection = true;
@@ -81,7 +131,8 @@ void SSlateOptions::Construct( const FArguments& InArgs )
 			AddMenuEntry(MenuBuilder, Icon, LOCTEXT("InvalidationDebugging", "Show Invalidation"), TEXT("SlateDebugger.Invalidate.Enable"));
 			AddMenuEntry(MenuBuilder, Icon, LOCTEXT("InvalidationRootDebugging", "Show Root Invalidation"), TEXT("SlateDebugger.InvalidationRoot.Enable"));
 			AddMenuEntry(MenuBuilder, Icon, LOCTEXT("UpdateDebugging", "Show Update"), TEXT("SlateDebugger.Update.Enable"));
-			AddMenuEntry(MenuBuilder, Icon, LOCTEXT("PaintDebugging", "Show Paint"), TEXT("SlateDebugger.Paint.Enable"));
+			AddSubMenuEntry(MenuBuilder, Icon, LOCTEXT("PaintDebugging", "Show Paint"), TEXT("SlateDebugger.Paint.Enable"), FNewMenuDelegate::CreateStatic(&FillPaintSubmenu));
+			AddSubMenuEntry(MenuBuilder, Icon, LOCTEXT("PrepassDebugging", "Show Prepass"), TEXT("SlateDebugger.Prepass.Enable"), FNewMenuDelegate::CreateStatic(&FillPrepassSubmenu));
 			AddMenuEntry(MenuBuilder, Icon, LOCTEXT("ShowClipping", "Show Clipping"), TEXT("Slate.ShowClipping"));
 			AddMenuEntry(MenuBuilder, Icon, LOCTEXT("DebugCulling", "Debug Culling"), TEXT("Slate.DebugCulling"));
 			AddMenuEntry(MenuBuilder, Icon, LOCTEXT("ShowHitTestGrid", "Show HitTestGrid"), TEXT("Slate.HitTestGridDebugging"));
diff --git a/Engine/Source/Runtime/SlateCore/Private/Debugging/ConsoleSlateDebuggerPaint.cpp b/Engine/Source/Runtime/SlateCore/Private/Debugging/ConsoleSlateDebuggerPaint.cpp
--- a/Engine/Source/Runtime/SlateCore/Private/Debugging/ConsoleSlateDebuggerPaint.cpp
+++ b/Engine/Source/Runtime/SlateCore/Private/Debugging/ConsoleSlateDebuggerPaint.cpp
@@ -1,38 +1,25 @@
 // Copyright Epic Games, Inc. All Rights Reserved.
 
 #include "ConsoleSlateDebuggerPaint.h"
-#include "ConsoleSlateDebugger.h"
 
 #if WITH_SLATE_DEBUGGING
 
-#include "Application/SlateApplicationBase.h"
 #include "CoreGlobals.h"
 #include "Debugging/SlateDebugging.h"
 #include "Layout/WidgetPath.h"
-#include "Types/ReflectionMetadata.h"
 #include "Misc/App.h"
 #include "Misc/ConfigCacheIni.h"
-#include "Styling/CoreStyle.h"
 
 #define LOCTEXT_NAMESPACE "ConsoleSlateDebuggerPaint"
 
 FConsoleSlateDebuggerPaint::FConsoleSlateDebuggerPaint()
-	: bEnabled(false)
-	, bEnabledCVarValue(false)
-	, bDisplayWidgetsNameList(false)
-	, bUseWidgetPathAsName(false)
-	, bDrawBox(false)
-	, bDrawQuad(true)
-	, bLogWidgetName(false)
-	, bLogWidgetNameOnce(false)
+	: FConsoleSlateDebuggerPassBase()
 	, bLogWarningIfWidgetIsPaintedMoreThanOnce(true)
-	, bDebugGameWindowOnly(true)
-	, DrawBoxColor(1.0f, 1.0f, 0.0f, 0.2f)
-	, DrawQuadColor(1.0f, 1.0f, 1.0f, 1.0f)
-	, DrawWidgetNameColor(FColorList::SpicyPink)
-	, MaxNumberOfWidgetInList(20)
-	, CacheDuration(2.0f)
-	, PIEWindowTag("PIEWindow")
+	, EnabledRefCVar(
+		TEXT("SlateDebugger.Paint.Enable")
+		, bEnabled
+		, TEXT("Start/Stop the painted widget debug tool. It shows when widgets are painted.")
+		, FConsoleVariableDelegate::CreateRaw(this, &FConsoleSlateDebuggerPaint::HandleEnabled))
 	, ShowPaintWidgetCommand(
 		TEXT("SlateDebugger.Paint.Start")
 		, TEXT("Start the painted widget debug tool. Use to show widget that have been painted this frame.")
@@ -41,276 +28,108 @@ FConsoleSlateDebuggerPaint::FConsoleSlateDebuggerPaint()
 		TEXT("SlateDebugger.Paint.Stop")
 		, TEXT("Stop the painted widget debug tool.")
 		, FConsoleCommandDelegate::CreateRaw(this, &FConsoleSlateDebuggerPaint::StopDebugging))
-	, EnabledRefCVar(
-		TEXT("SlateDebugger.Paint.Enable")
-		, bEnabledCVarValue
-		, TEXT("Start/Stop the painted widget debug tool. It shows when widgets are painted.")
-		, FConsoleVariableDelegate::CreateRaw(this, &FConsoleSlateDebuggerPaint::HandleEnabled))
 	, LogPaintedWidgetOnceCommand(
 		TEXT("SlateDebugger.Paint.LogOnce")
 		, TEXT("Log the names of all widgets that were painted during the last update.")
 		, FConsoleCommandDelegate::CreateRaw(this, &FConsoleSlateDebuggerPaint::HandleLogOnce))
-	, DisplayWidgetsNameListCommand(
-		TEXT("SlateDebugger.Paint.ToggleWidgetNameList"),
-		TEXT("Option to display the name of the widgets that have been painted."),
-		FConsoleCommandDelegate::CreateRaw(this, &FConsoleSlateDebuggerPaint::HandleToggleWidgetNameList))
-	, MaxNumberOfWidgetInListtRefCVar(
+	, EnableWidgetsNameListRefCVar(
+		TEXT("SlateDebugger.Paint.EnableWidgetNameList")
+		, bDisplayWidgetsNameList
+		, TEXT("Start/Stop displaying the name of the widgets that have been painted.")
+		, FConsoleVariableDelegate::CreateRaw(this, &FConsoleSlateDebuggerPaint::SaveConfigOnVariableChanged))
+	, ToggleWidgetsNameListCommand(
+		TEXT("SlateDebugger.Paint.ToggleWidgetNameList")
+		, TEXT("Option displaying the name of the widgets that have been painted.")
+		, FConsoleCommandDelegate::CreateRaw(this, &FConsoleSlateDebuggerPaint::HandleToggleWidgetNameList))
+	, MaxNumberOfWidgetInListRefCVar(
 		TEXT("SlateDebugger.Paint.MaxNumberOfWidgetDisplayedInList")
 		, MaxNumberOfWidgetInList
-		, TEXT("The max number of widgets that will be displayed when DisplayWidgetNameList is active."))
+		, TEXT("The max number of widgets that will be displayed when DisplayWidgetNameList is active.")
+		, FConsoleVariableDelegate::CreateRaw(this, &FConsoleSlateDebuggerPaint::SaveConfigOnVariableChanged))
 	, LogWarningIfWidgetIsPaintedMoreThanOnceRefCVar(
 		TEXT("SlateDebugger.Paint.LogWarningIfWidgetIsPaintedMoreThanOnce")
 		, bLogWarningIfWidgetIsPaintedMoreThanOnce
-		, TEXT("Option to log a warning if a widget is painted more than once in a single frame."))
-	, OnlyGameWindow(
-		TEXT("SlateDebugger.Paint.OnlyGameWindow"),
-		bDebugGameWindowOnly,
-		TEXT("Option to only the debug the game window"))
+		, TEXT("Option to log a warning if a widget is painted more than once in a single frame.")
+		, FConsoleVariableDelegate::CreateRaw(this, &FConsoleSlateDebuggerPaint::SaveConfigOnVariableChanged))
+	, OnlyGameWindowRefCVar(
+		TEXT("SlateDebugger.Paint.OnlyGameWindow")
+		, bDebugGameWindowOnly
+		, TEXT("Option to only debug the game window")
+		, FConsoleVariableDelegate::CreateRaw(this, &FConsoleSlateDebuggerPaint::HandleDebugGameWindowOnlyChanged))
+	, DrawBorderEnabledRefCVar(
+		TEXT("SlateDebugger.Paint.DrawBorder")
+		, bDrawBorder
+		, TEXT("Draw a border around the widgets being painted")
+		, FConsoleVariableDelegate::CreateRaw(this, &FConsoleSlateDebuggerPaint::SaveConfigOnVariableChanged))
+	, DrawFillEnabledRefCVar(
+		TEXT("SlateDebugger.Paint.DrawFill")
+		, bDrawBox
+		, TEXT("Fill the widgets being painted")
+		, FConsoleVariableDelegate::CreateRaw(this, &FConsoleSlateDebuggerPaint::SaveConfigOnVariableChanged))
 {
-	GConfig->GetBool(TEXT("SlateDebugger.Paint"), TEXT("bDisplayWidgetsNameList"), bDisplayWidgetsNameList, *GEditorPerProjectIni);
-	GConfig->GetBool(TEXT("SlateDebugger.Paint"), TEXT("bUseWidgetPathAsName"), bUseWidgetPathAsName, *GEditorPerProjectIni);
-	GConfig->GetBool(TEXT("SlateDebugger.Paint"), TEXT("bDrawBox"), bDrawBox, *GEditorPerProjectIni);
-	GConfig->GetBool(TEXT("SlateDebugger.Paint"), TEXT("bDrawQuad"), bDrawQuad, *GEditorPerProjectIni);
-	GConfig->GetBool(TEXT("SlateDebugger.Paint"), TEXT("bLogWidgetName"), bLogWidgetName, *GEditorPerProjectIni);
-	GConfig->GetBool(TEXT("SlateDebugger.Paint"), TEXT("bLogWarningIfWidgetIsPaintedMoreThanOnce"), bLogWarningIfWidgetIsPaintedMoreThanOnce, *GEditorPerProjectIni);
-	FColor TmpColor;
-	if (GConfig->GetColor(TEXT("SlateDebugger.Paint"), TEXT("DrawBoxColor"), TmpColor, *GEditorPerProjectIni))
-	{
-		DrawBoxColor = TmpColor;
-	}
-	if (GConfig->GetColor(TEXT("SlateDebugger.Paint"), TEXT("DrawQuadColor"), TmpColor, *GEditorPerProjectIni))
-	{
-		DrawQuadColor = TmpColor;
-	}
-	if (GConfig->GetColor(TEXT("SlateDebugger.Paint"), TEXT("DrawWidgetNameColor"), TmpColor, *GEditorPerProjectIni))
-	{
-		DrawWidgetNameColor = TmpColor;
-	}
-	GConfig->GetInt(TEXT("SlateDebugger.Paint"), TEXT("MaxNumberOfWidgetInList"), MaxNumberOfWidgetInList, *GEditorPerProjectIni);
-	GConfig->GetFloat(TEXT("SlateDebugger.Paint"), TEXT("CacheDuration"), CacheDuration, *GEditorPerProjectIni);
 }
 
-FConsoleSlateDebuggerPaint::~FConsoleSlateDebuggerPaint()
+void FConsoleSlateDebuggerPaint::LoadConfig()
 {
-	StopDebugging();
+	FConsoleSlateDebuggerPassBase::LoadConfig();
+	
+	const FString Section = *GetConfigSection();
+	GConfig->GetBool(*Section, TEXT("bLogWarningIfWidgetIsPaintedMoreThanOnce"), bLogWarningIfWidgetIsPaintedMoreThanOnce, *GEditorPerProjectIni);
+
+	OnlyGameWindowRefCVar->Set(bDebugGameWindowOnly, ECVF_SetByProjectSetting);
+	EnableWidgetsNameListRefCVar->Set(bDisplayWidgetsNameList, ECVF_SetByProjectSetting);
+	DrawFillEnabledRefCVar->Set(bDrawBox, ECVF_SetByProjectSetting);
+	DrawBorderEnabledRefCVar->Set(bDrawBorder, ECVF_SetByProjectSetting);
 }
 
 void FConsoleSlateDebuggerPaint::SaveConfig()
 {
-	GConfig->SetBool(TEXT("SlateDebugger.Paint"), TEXT("bDisplayWidgetsNameList"), bDisplayWidgetsNameList, *GEditorPerProjectIni);
-	GConfig->SetBool(TEXT("SlateDebugger.Paint"), TEXT("bUseWidgetPathAsName"), bUseWidgetPathAsName, *GEditorPerProjectIni);
-	GConfig->SetBool(TEXT("SlateDebugger.Paint"), TEXT("bDrawBox"), bDrawBox, *GEditorPerProjectIni);
-	GConfig->SetBool(TEXT("SlateDebugger.Paint"), TEXT("bDrawQuad"), bDrawQuad, *GEditorPerProjectIni);
-	GConfig->SetBool(TEXT("SlateDebugger.Paint"), TEXT("bLogWidgetName"), bLogWidgetName, *GEditorPerProjectIni);
-	GConfig->SetBool(TEXT("SlateDebugger.Paint"), TEXT("bLogWarningIfWidgetIsPaintedMoreThanOnce"), bLogWarningIfWidgetIsPaintedMoreThanOnce, *GEditorPerProjectIni);
-	FColor TmpColor = DrawBoxColor.ToFColor(true);
-	GConfig->SetColor(TEXT("SlateDebugger.Paint"), TEXT("DrawBoxColor"), TmpColor, *GEditorPerProjectIni);
-	TmpColor = DrawQuadColor.ToFColor(true);
-	GConfig->SetColor(TEXT("SlateDebugger.Paint"), TEXT("DrawQuadColor"), TmpColor, *GEditorPerProjectIni);
-	TmpColor = DrawWidgetNameColor.ToFColor(true);
-	GConfig->SetColor(TEXT("SlateDebugger.Paint"), TEXT("DrawWidgetNameColor"), TmpColor, *GEditorPerProjectIni);
-	GConfig->SetInt(TEXT("SlateDebugger.Paint"), TEXT("MaxNumberOfWidgetInList"), MaxNumberOfWidgetInList, *GEditorPerProjectIni);
-	GConfig->SetFloat(TEXT("SlateDebugger.Paint"), TEXT("CacheDuration"), CacheDuration, *GEditorPerProjectIni);
+	FConsoleSlateDebuggerPassBase::SaveConfig();
+	
+	const FString Section = *GetConfigSection();
+	GConfig->SetBool(*Section, TEXT("bLogWarningIfWidgetIsPaintedMoreThanOnce"), bLogWarningIfWidgetIsPaintedMoreThanOnce, *GEditorPerProjectIni);
 }
 
-void FConsoleSlateDebuggerPaint::StartDebugging()
+void FConsoleSlateDebuggerPaint::StartDebugging_Internal()
 {
-	if (!bEnabled)
-	{
-		bEnabled = true;
-		PaintedWidgets.Empty();
-
-		FSlateDebugging::EndWidgetPaint.AddRaw(this, &FConsoleSlateDebuggerPaint::HandleEndWidgetPaint);
-		FSlateDebugging::PaintDebugElements.AddRaw(this, &FConsoleSlateDebuggerPaint::HandlePaintDebugInfo);
-		FCoreDelegates::OnEndFrame.AddRaw(this, &FConsoleSlateDebuggerPaint::HandleEndFrame);
-	}
-	bEnabledCVarValue = bEnabled;
+	FConsoleSlateDebuggerPassBase::StartDebugging_Internal();
+	
+	FSlateDebugging::EndWidgetPaint.AddRaw(this, &FConsoleSlateDebuggerPaint::HandleEndWidgetPaint);
 }
 
-void FConsoleSlateDebuggerPaint::StopDebugging()
+void FConsoleSlateDebuggerPaint::StopDebugging_Internal()
 {
-	if (bEnabled)
-	{
-		FCoreDelegates::OnEndFrame.RemoveAll(this);
-		FSlateDebugging::PaintDebugElements.RemoveAll(this);
-		FSlateDebugging::EndWidgetPaint.RemoveAll(this);
-
-		PaintedWidgets.Empty();
-		bEnabled = false;
-	}
-	bEnabledCVarValue = bEnabled;
+	FSlateDebugging::EndWidgetPaint.RemoveAll(this);
+	
+	FConsoleSlateDebuggerPassBase::StopDebugging_Internal();
 }
 
-void FConsoleSlateDebuggerPaint::HandleEnabled(IConsoleVariable* Variable)
+void FConsoleSlateDebuggerPaint::AddUpdatedWidget(const SWidget& Widget, const FConsoleSlateDebuggerUtility::TSWindowId WindowId, bool bIncrementUpdateCount)
 {
-	if (bEnabledCVarValue)
-	{
-		StartDebugging();
-	}
-	else
+	FWidgetInfo& WidgetInfo = AddUpdatedWidget_Internal(Widget, WindowId, Widget.Debug_GetLastPaintFrame());
+
+	if (bLogWarningIfWidgetIsPaintedMoreThanOnce && WidgetInfo.UpdateCount != 0)
 	{
-		StopDebugging();
+		UE_LOG(LogSlateDebugger, Warning, TEXT("'%s' got painted more than once."), *(WidgetInfo.WidgetName));
 	}
-}
 
-void FConsoleSlateDebuggerPaint::HandleLogOnce()
-{
-	bLogWidgetNameOnce = true;
-}
-
-void FConsoleSlateDebuggerPaint::HandleToggleWidgetNameList()
-{
-	bDisplayWidgetsNameList = !bDisplayWidgetsNameList;
-	SaveConfig();
-}
-
-void FConsoleSlateDebuggerPaint::HandleEndFrame()
-{
-	double LastTime = FSlateApplicationBase::Get().GetCurrentTime() - CacheDuration;
-	for (TPaintedWidgetMap::TIterator It(PaintedWidgets); It; ++It)
+	if (bIncrementUpdateCount)
 	{
-		It.Value().PaintCount = 0;
-		if (It.Value().LastPaint < LastTime)
-		{
-			It.RemoveCurrent();
-		}
+		++WidgetInfo.UpdateCount;
 	}
 }
 
 void FConsoleSlateDebuggerPaint::HandleEndWidgetPaint(const SWidget* Widget, const FSlateWindowElementList& OutDrawElements, int32 LayerId)
 {
-	// Use the Widget pointer for the id.
-	//That may introduce bug when a widget is destroyed and the same memory is reused for another widget. We do not care for this debug tool.
-	//We do not keep the widget alive or reuse it later, cache all the info that we need.
-	const FConsoleSlateDebuggerUtility::TSWidgetId WidgetId = FConsoleSlateDebuggerUtility::GetId(Widget);
-	const FConsoleSlateDebuggerUtility::TSWindowId WindowId = FConsoleSlateDebuggerUtility::GetId(OutDrawElements.GetPaintWindow());
-
 	// Exclude all windows but the game window
-	SWindow* WindowToDrawIn = OutDrawElements.GetPaintWindow();
+	const SWindow* WindowToDrawIn = OutDrawElements.GetPaintWindow();
 	if (bDebugGameWindowOnly && (WindowToDrawIn->GetType() != EWindowType::GameWindow && WindowToDrawIn->GetTag() != PIEWindowTag))
 	{
 		return;
 	}
 
-
-	FPaintInfo* FoundItem = PaintedWidgets.Find(WidgetId);
-	if (FoundItem == nullptr)
-	{
-		FoundItem = &PaintedWidgets.Add(WidgetId);
-		FoundItem->PaintCount = 0;
-		FoundItem->Window = WindowId;
-		FoundItem->WidgetName = bUseWidgetPathAsName ? FReflectionMetaData::GetWidgetPath(Widget) : FReflectionMetaData::GetWidgetDebugInfo(Widget);
-	}
-	else
-	{
-		ensureAlways(FoundItem->Window == WindowId);
-		if (bLogWarningIfWidgetIsPaintedMoreThanOnce && FoundItem->PaintCount != 0)
-		{
-			UE_LOG(LogSlateDebugger, Warning, TEXT("'%s' got painted more than once."), *(FoundItem->WidgetName));
-		}
-	}
-
-	if (bLogWidgetName)
-	{
-		UE_LOG(LogSlateDebugger, Log, TEXT("%s"), *(FoundItem->WidgetName));
-	}
-	FoundItem->PaintLocation = Widget->GetPersistentState().AllottedGeometry.GetAbsolutePosition();
-	FoundItem->PaintSize = Widget->GetPersistentState().AllottedGeometry.GetAbsoluteSize();
-	FoundItem->LastPaint = FSlateApplicationBase::Get().GetCurrentTime();
-	++FoundItem->PaintCount;
-}
-
-void FConsoleSlateDebuggerPaint::HandlePaintDebugInfo(const FPaintArgs& InArgs, const FGeometry& InAllottedGeometry, FSlateWindowElementList& InOutDrawElements, int32& InOutLayerId)
-{
-	++InOutLayerId;
-
-	FConsoleSlateDebuggerUtility::TSWindowId PaintWindow = FConsoleSlateDebuggerUtility::GetId(InOutDrawElements.GetPaintWindow());
-
-	int32 NumberOfWidget = 0;
-	const float TextElementY = 36.f;
-	const FSlateBrush* BoxBrush = bDrawBox ? FCoreStyle::Get().GetBrush("WhiteBrush") : nullptr;
-	const FSlateBrush* QuadBrush = FCoreStyle::Get().GetBrush(TEXT("FocusRectangle"));
-	FSlateFontInfo FontInfo = FCoreStyle::Get().GetFontStyle("SmallFont");
-	FontInfo.OutlineSettings.OutlineSize = 1;
-
-	CacheDuration = FMath::Max(CacheDuration, 0.01f);
-	const double SlateApplicationCurrentTime = FSlateApplicationBase::Get().GetCurrentTime();
-
-	auto MakeText = [&](const FString& Text, const FVector2f& Location, const FLinearColor& Color)
-	{
-		FSlateDrawElement::MakeText(
-			InOutDrawElements
-			, InOutLayerId
-			, InAllottedGeometry.ToPaintGeometry(FVector2f(1.f, 1.f), FSlateLayoutTransform(Location))
-			, Text
-			, FontInfo
-			, ESlateDrawEffect::None
-			, Color);
-	};
-
-	for (const auto& Itt : PaintedWidgets)
-	{
-		if (Itt.Value.Window == PaintWindow)
-		{
-			const float LerpValue = FMath::Clamp((float)(SlateApplicationCurrentTime - Itt.Value.LastPaint) / CacheDuration, 0.0f, 1.0f);
-			const FGeometry Geometry = FGeometry::MakeRoot(Itt.Value.PaintSize, FSlateLayoutTransform(1.f, Itt.Value.PaintLocation));
-			const FPaintGeometry PaintGeometry = Geometry.ToPaintGeometry();
-
-			if (BoxBrush)
-			{
-				FSlateDrawElement::MakeBox(
-					InOutDrawElements,
-					InOutLayerId,
-					PaintGeometry,
-					BoxBrush,
-					ESlateDrawEffect::None,
-					DrawBoxColor.CopyWithNewOpacity(FMath::InterpExpoOut(1.0f, 0.0f, LerpValue)));
-			}
-
-			if (bDrawQuad)
-			{
-				FSlateDrawElement::MakeDebugQuad(
-					InOutDrawElements,
-					InOutLayerId,
-					PaintGeometry,
-					DrawQuadColor);
-
-				FSlateDrawElement::MakeBox(
-					InOutDrawElements,
-					InOutLayerId,
-					PaintGeometry,
-					QuadBrush,
-					ESlateDrawEffect::None,
-					DrawQuadColor.CopyWithNewOpacity(FMath::InterpExpoOut(1.0f, 0.0f, LerpValue)));
-			}
-
-			if (bLogWidgetNameOnce)
-			{
-				UE_LOG(LogSlateDebugger, Log, TEXT("%s"), *(Itt.Value.WidgetName));
-			}
-
-			if (bDisplayWidgetsNameList)
-			{
-				if (NumberOfWidget < MaxNumberOfWidgetInList)
-				{
-					MakeText(Itt.Value.WidgetName, FVector2f(0.f, (12.f * NumberOfWidget) + TextElementY), DrawWidgetNameColor);
-				}
-			}
-			++NumberOfWidget;
-		}
-	}
-	bLogWidgetNameOnce = false;
-
-	{
-		FString NumberOfWidgetDrawn = FString::Printf(TEXT("Number of Widget Painted: %d"), NumberOfWidget);
-		MakeText(NumberOfWidgetDrawn, FVector2f(10.f, 10.f), DrawWidgetNameColor);
-	}
-
-	if (bDisplayWidgetsNameList && NumberOfWidget > MaxNumberOfWidgetInList)
-	{
-		FString WidgetDisplayName = FString::Printf(TEXT("   %d more invalidations"), NumberOfWidget - MaxNumberOfWidgetInList);
-		MakeText(WidgetDisplayName, FVector2f(0.f, (12.f * NumberOfWidget) + TextElementY), FLinearColor::White);
-	}
+	const FConsoleSlateDebuggerUtility::TSWindowId WindowId = FConsoleSlateDebuggerUtility::GetId(OutDrawElements.GetPaintWindow());
+	AddUpdatedWidget(*Widget, WindowId);
 }
 
 #undef LOCTEXT_NAMESPACE
diff --git a/Engine/Source/Runtime/SlateCore/Private/Debugging/ConsoleSlateDebuggerPaint.h b/Engine/Source/Runtime/SlateCore/Private/Debugging/ConsoleSlateDebuggerPaint.h
--- a/Engine/Source/Runtime/SlateCore/Private/Debugging/ConsoleSlateDebuggerPaint.h
+++ b/Engine/Source/Runtime/SlateCore/Private/Debugging/ConsoleSlateDebuggerPaint.h
@@ -8,87 +8,56 @@
 #if WITH_SLATE_DEBUGGING
 
 #include "CoreMinimal.h"
+#include "ConsoleSlateDebuggerPassBase.h"
 #include "Debugging/ConsoleSlateDebuggerUtility.h"
 #include "Delegates/Delegate.h"
-#include "Input/Reply.h"
 #include "HAL/IConsoleManager.h"
 #include "Rendering/DrawElements.h"
 
-struct FGeometry;
-class FPaintArgs;
-class FSlateRect;
-class FSlateWindowElementList;
-class SWidget;
-class SWindow;
-
 /**
  * Allows debugging the behavior of SWidget::Paint from the console.
  * Basics:
  *   Start - SlateDebugger.Paint.Start
  *   Stop  - SlateDebugger.Paint.Stop
  */
-class FConsoleSlateDebuggerPaint
+class FConsoleSlateDebuggerPaint : public FConsoleSlateDebuggerPassBase
 {
 public:
 	FConsoleSlateDebuggerPaint();
-	virtual ~FConsoleSlateDebuggerPaint();
 
-	void StartDebugging();
-	void StopDebugging();
-	bool IsEnabled() const { return bEnabled; }
+	virtual void LoadConfig() override;
+	virtual void SaveConfig() override;
 
-	void SaveConfig();
+protected:
+	virtual FString GetConfigSection() const override { return TEXT("SlateDebugger.Paint"); }
+	virtual FString GetNumberOfWidgetsUpdatedLogString(uint32 NumberOfWidgetsUpdatedThisFrame) const override
+	{
+		return FString::Printf(TEXT("%d widgets painted"), NumberOfWidgetsUpdatedThisFrame);
+	}
+	virtual FAutoConsoleVariableRef& GetEnabledCVar() override { return EnabledRefCVar; }
+	virtual void StartDebugging_Internal() override;
+	virtual void StopDebugging_Internal() override;
+	
+	virtual void AddUpdatedWidget(const SWidget& Widget, const FConsoleSlateDebuggerUtility::TSWindowId WindowId, bool bIncrementUpdateCount = true) override;
 
 private:
-	void HandleEnabled(IConsoleVariable* Variable);
-	void HandleLogOnce();
-	void HandleToggleWidgetNameList();
-	void HandleEndFrame();
 	void HandleEndWidgetPaint(const SWidget* Widget, const FSlateWindowElementList& OutDrawElements, int32 LayerId);
-	void HandlePaintDebugInfo(const FPaintArgs& InArgs, const FGeometry& InAllottedGeometry, FSlateWindowElementList& InOutDrawElements, int32& InOutLayerId);
 
 private:
-	bool bEnabled;
-	bool bEnabledCVarValue;
-
-	//~ Settings
-	bool bDisplayWidgetsNameList;
-	bool bUseWidgetPathAsName;
-	bool bDrawBox;
-	bool bDrawQuad;
-	bool bLogWidgetName;
-	bool bLogWidgetNameOnce;
 	bool bLogWarningIfWidgetIsPaintedMoreThanOnce;
-	bool bDebugGameWindowOnly;
-	FLinearColor DrawBoxColor;
-	FLinearColor DrawQuadColor;
-	FLinearColor DrawWidgetNameColor;
-	int32 MaxNumberOfWidgetInList;
-	float CacheDuration;
-	FName PIEWindowTag;
-
+	
 	//~ Console objects
+	FAutoConsoleVariableRef EnabledRefCVar;
 	FAutoConsoleCommand ShowPaintWidgetCommand;
 	FAutoConsoleCommand HidePaintWidgetCommand;
-	FAutoConsoleVariableRef EnabledRefCVar;
 	FAutoConsoleCommand LogPaintedWidgetOnceCommand;
-	FAutoConsoleCommand DisplayWidgetsNameListCommand;
-	FAutoConsoleVariableRef MaxNumberOfWidgetInListtRefCVar;
+	FAutoConsoleVariableRef EnableWidgetsNameListRefCVar;
+	FAutoConsoleCommand ToggleWidgetsNameListCommand;
+	FAutoConsoleVariableRef MaxNumberOfWidgetInListRefCVar;
 	FAutoConsoleVariableRef LogWarningIfWidgetIsPaintedMoreThanOnceRefCVar;
-	FAutoConsoleVariableRef OnlyGameWindow;
-
-	struct FPaintInfo 
-	{
-		FConsoleSlateDebuggerUtility::TSWindowId Window;
-		FVector2f PaintLocation;
-		FVector2f PaintSize;
-		FString WidgetName;
-		double LastPaint;
-		int32 PaintCount;
-	};
-
-	using TPaintedWidgetMap = TMap<FConsoleSlateDebuggerUtility::TSWidgetId, FPaintInfo>;
-	TPaintedWidgetMap PaintedWidgets;
+	FAutoConsoleVariableRef OnlyGameWindowRefCVar;
+	FAutoConsoleVariableRef DrawBorderEnabledRefCVar;
+	FAutoConsoleVariableRef DrawFillEnabledRefCVar;
 };
 
 #endif //WITH_SLATE_DEBUGGING
\ No newline at end of file
diff --git a/Engine/Source/Runtime/SlateCore/Private/Debugging/ConsoleSlateDebuggerPassBase.cpp b/Engine/Source/Runtime/SlateCore/Private/Debugging/ConsoleSlateDebuggerPassBase.cpp
new file mode 100644
--- /dev/null
+++ b/Engine/Source/Runtime/SlateCore/Private/Debugging/ConsoleSlateDebuggerPassBase.cpp
@@ -0,0 +1,365 @@
+﻿// Copyright Epic Games, Inc. All Rights Reserved.
+
+#include "ConsoleSlateDebuggerPassBase.h"
+
+#if WITH_SLATE_DEBUGGING
+
+#include "ConsoleSlateDebugger.h"
+#include "CoreGlobals.h"
+#include "Application/SlateApplicationBase.h"
+#include "Debugging/SlateDebugging.h"
+#include "Layout/WidgetPath.h"
+#include "Misc/App.h"
+#include "Misc/ConfigCacheIni.h"
+#include "Styling/CoreStyle.h"
+#include "Types/ReflectionMetadata.h"
+
+FConsoleSlateDebuggerPassBase::FConsoleSlateDebuggerPassBase()
+	: bEnabled(false)
+	, bDisplayWidgetsNameList(false)
+	, bUseWidgetPathAsName(false)
+	, bDrawBox(false)
+	, bDrawBorder(true)
+	, bLogWidgetName(false)
+	, bLogWidgetNameOnce(false)
+	, bDebugGameWindowOnly(true)
+	, MostRecentColor(0.75f, 0.0f, 0.0f, 0.02f)
+	, LeastRecentColor(0.0f, 0.75f, 0.0f, 0.5f)
+	, DrawWidgetNameColor(FColorList::SpicyPink)
+	, MaxNumberOfWidgetInList(20)
+	, FadeDuration(2.0f)
+	, PIEWindowTag("PIEWindow")
+{
+}
+
+FConsoleSlateDebuggerPassBase::~FConsoleSlateDebuggerPassBase()
+{
+	StopDebugging();
+}
+
+void FConsoleSlateDebuggerPassBase::LoadConfig()
+{
+	const FString Section = *GetConfigSection();
+	GConfig->GetBool(*Section, TEXT("bDebugGameWindowOnly"), bDebugGameWindowOnly, *GEditorPerProjectIni);
+	GConfig->GetBool(*Section, TEXT("bDisplayWidgetsNameList"), bDisplayWidgetsNameList, *GEditorPerProjectIni);
+	GConfig->GetBool(*Section, TEXT("bUseWidgetPathAsName"), bUseWidgetPathAsName, *GEditorPerProjectIni);
+	GConfig->GetBool(*Section, TEXT("bDrawBox"), bDrawBox, *GEditorPerProjectIni);
+	GConfig->GetBool(*Section, TEXT("bDrawBorder"), bDrawBorder, *GEditorPerProjectIni);
+	GConfig->GetBool(*Section, TEXT("bLogWidgetName"), bLogWidgetName, *GEditorPerProjectIni);
+	FColor TmpColor;
+	if (GConfig->GetColor(*Section, TEXT("MostRecentColor"), TmpColor, *GEditorPerProjectIni))
+	{
+		MostRecentColor = TmpColor;
+	}
+	if (GConfig->GetColor(*Section, TEXT("LeastRecentColor"), TmpColor, *GEditorPerProjectIni))
+	{
+		LeastRecentColor = TmpColor;
+	}
+	if (GConfig->GetColor(*Section, TEXT("DrawWidgetNameColor"), TmpColor, *GEditorPerProjectIni))
+	{
+		DrawWidgetNameColor = TmpColor;
+	}
+	GConfig->GetInt(*Section, TEXT("MaxNumberOfWidgetInList"), MaxNumberOfWidgetInList, *GEditorPerProjectIni);
+	GConfig->GetFloat(*Section, TEXT("FadeDuration"), FadeDuration, *GEditorPerProjectIni);
+}
+
+void FConsoleSlateDebuggerPassBase::SaveConfig()
+{
+	const FString Section = *GetConfigSection();
+	GConfig->SetBool(*Section, TEXT("bDebugGameWindowOnly"), bDebugGameWindowOnly, *GEditorPerProjectIni);
+	GConfig->SetBool(*Section, TEXT("bDisplayWidgetsNameList"), bDisplayWidgetsNameList, *GEditorPerProjectIni);
+	GConfig->SetBool(*Section, TEXT("bUseWidgetPathAsName"), bUseWidgetPathAsName, *GEditorPerProjectIni);
+	GConfig->SetBool(*Section, TEXT("bDrawBox"), bDrawBox, *GEditorPerProjectIni);
+	GConfig->SetBool(*Section, TEXT("bDrawBorder"), bDrawBorder, *GEditorPerProjectIni);
+	GConfig->SetBool(*Section, TEXT("bLogWidgetName"), bLogWidgetName, *GEditorPerProjectIni);
+	FColor TmpColor = MostRecentColor.ToFColor(true);
+	GConfig->SetColor(*Section, TEXT("MostRecentPaintColor"), TmpColor, *GEditorPerProjectIni);
+	TmpColor = LeastRecentColor.ToFColor(true);
+	GConfig->SetColor(*Section, TEXT("LeastRecentPaintColor"), TmpColor, *GEditorPerProjectIni);
+	TmpColor = DrawWidgetNameColor.ToFColor(true);
+	GConfig->SetColor(*Section, TEXT("DrawWidgetNameColor"), TmpColor, *GEditorPerProjectIni);
+	GConfig->SetInt(*Section, TEXT("MaxNumberOfWidgetInList"), MaxNumberOfWidgetInList, *GEditorPerProjectIni);
+	GConfig->SetFloat(*Section, TEXT("FadeDuration"), FadeDuration, *GEditorPerProjectIni);
+}
+
+void FConsoleSlateDebuggerPassBase::StartDebugging()
+{
+	if (!bEnabled)
+	{
+		// This wil end up calling HandleEnabled > StartDebugging_Internal
+		GetEnabledCVar()->Set(true, EConsoleVariableFlags::ECVF_SetByCode);
+	}
+}
+
+void FConsoleSlateDebuggerPassBase::StopDebugging()
+{
+	if (bEnabled)
+	{
+		// This wil end up calling HandleEnabled > StopDebugging_Internal
+		GetEnabledCVar()->Set(false, EConsoleVariableFlags::ECVF_SetByCode);
+	}
+}
+
+void FConsoleSlateDebuggerPassBase::HandleEnabled(IConsoleVariable* Variable)
+{
+	// The value has already been changed by the CVar
+	if (bEnabled)
+	{
+		StartDebugging_Internal();
+	}
+	else
+	{
+		StopDebugging_Internal();
+	}
+}
+
+void FConsoleSlateDebuggerPassBase::StartDebugging_Internal()
+{
+	bEnabled = true;
+	Widgets.Empty();
+
+	BuildInitialWidgetList();
+
+	FSlateDebugging::PaintDebugElements.AddRaw(this, &FConsoleSlateDebuggerPassBase::HandlePaintDebugInfo);
+	FCoreDelegates::OnEndFrame.AddRaw(this, &FConsoleSlateDebuggerPassBase::HandleEndFrame);
+}
+
+void FConsoleSlateDebuggerPassBase::StopDebugging_Internal()
+{
+	FCoreDelegates::OnEndFrame.RemoveAll(this);
+	FSlateDebugging::PaintDebugElements.RemoveAll(this);
+
+	Widgets.Empty();
+	bEnabled = false;
+}
+
+void FConsoleSlateDebuggerPassBase::HandleLogOnce()
+{
+	bLogWidgetNameOnce = true;
+}
+
+void FConsoleSlateDebuggerPassBase::SaveConfigOnVariableChanged(IConsoleVariable* Variable)
+{
+	SaveConfig();
+}
+
+void FConsoleSlateDebuggerPassBase::HandleDebugGameWindowOnlyChanged(IConsoleVariable* Variable)
+{
+	SaveConfig();
+	if (bEnabled)
+	{
+		// Recreate the list
+		Widgets.Empty();
+		BuildInitialWidgetList();
+	}
+}
+
+void FConsoleSlateDebuggerPassBase::HandleToggleWidgetNameList()
+{
+	bDisplayWidgetsNameList = !bDisplayWidgetsNameList;
+	SaveConfig();
+}
+
+void FConsoleSlateDebuggerPassBase::HandleEndFrame()
+{
+	for (TUpdatedWidgetMap::TIterator It(Widgets); It; ++It)
+	{
+		if (!It.Value().Widget.IsValid())
+		{
+			It.RemoveCurrent();
+		}
+		else
+		{
+			It.Value().UpdateCount = 0;
+		}
+	}
+}
+
+void FConsoleSlateDebuggerPassBase::HandlePaintDebugInfo(const FPaintArgs& InArgs, const FGeometry& InAllottedGeometry, FSlateWindowElementList& InOutDrawElements, int32& InOutLayerId)
+{
+	if (bDebugGameWindowOnly && (InOutDrawElements.GetPaintWindow()->GetType() != EWindowType::GameWindow && InOutDrawElements.GetPaintWindow()->GetTag() != PIEWindowTag))
+	{
+		return;
+	}
+
+	++InOutLayerId;
+
+	FConsoleSlateDebuggerUtility::TSWindowId PaintWindow = FConsoleSlateDebuggerUtility::GetId(InOutDrawElements.GetPaintWindow());
+
+	uint32 NumberOfWidgetsUpdatedThisFrame = 0;
+	uint32 NumberOfWidgetsLoggedThisFrame = 0; // We might log widgets updated on previous frames
+	const float TextElementY = 36.f;
+	const FSlateBrush* BoxBrush = FCoreStyle::Get().GetBrush("WhiteBrush");
+	const FSlateBrush* QuadBrush = FCoreStyle::Get().GetBrush("Border");
+	FSlateFontInfo FontInfo = FCoreStyle::Get().GetFontStyle("SmallFont");
+	FontInfo.OutlineSettings.OutlineSize = 1;
+
+	FadeDuration = FMath::Max(FadeDuration, 0.01f);
+	const double SlateApplicationCurrentTime = FSlateApplicationBase::Get().GetCurrentTime();
+
+	auto MakeText = [&](const FString& Text, const FVector2f& Location, const FLinearColor& Color)
+	{
+		FSlateDrawElement::MakeText(
+			InOutDrawElements
+			, InOutLayerId
+			, InAllottedGeometry.ToPaintGeometry(FVector2f(1.f, 1.f), FSlateLayoutTransform(Location))
+			, Text
+			, FontInfo
+			, ESlateDrawEffect::None
+			, Color
+		);
+	};
+
+	for (const TPair<FConsoleSlateDebuggerUtility::TSWidgetId, FWidgetInfo>& Pair : Widgets)
+	{
+		const FWidgetInfo& WidgetInfo = Pair.Value;
+		if (WidgetInfo.Window != PaintWindow)
+		{
+			continue;
+		}
+		
+		if (WidgetInfo.LastUpdatedFrame == GFrameNumber)
+		{
+			++NumberOfWidgetsUpdatedThisFrame;
+			if (bLogWidgetNameOnce)
+			{
+				UE_LOG(LogSlateDebugger, Log, TEXT("%s"), *WidgetInfo.WidgetName);
+			}
+		}
+
+		// LerpValue of 0 represents a widget painted this frame, 1 a Widget painted FadeDuration ago (or more)
+		const float LerpValue = FMath::Clamp((float)(SlateApplicationCurrentTime - WidgetInfo.LastUpdatedTime) / FadeDuration, 0.0f, 1.0f);
+		const FLinearColor FinalColor = FMath::Lerp(MostRecentColor, LeastRecentColor, LerpValue);
+		const FGeometry Geometry = FGeometry::MakeRoot(WidgetInfo.PaintSize, FSlateLayoutTransform(1.f, WidgetInfo.PaintLocation));
+		const FPaintGeometry PaintGeometry = Geometry.ToPaintGeometry();
+		
+		if (bDrawBox)
+		{
+			FSlateDrawElement::MakeBox(
+				InOutDrawElements,
+				InOutLayerId,
+				PaintGeometry,
+				BoxBrush,
+				ESlateDrawEffect::None,
+				FinalColor
+			);
+		}
+
+		if (bDrawBorder)
+		{
+			FSlateDrawElement::MakeBox(
+				InOutDrawElements,
+				InOutLayerId,
+				PaintGeometry,
+				QuadBrush,
+				ESlateDrawEffect::None,
+				FinalColor.CopyWithNewOpacity(1.0f)
+			);
+		}
+		
+		if (bDisplayWidgetsNameList && LerpValue < 1.0f) // Only display the name of recent widgets
+		{
+			if (NumberOfWidgetsLoggedThisFrame < (uint32)MaxNumberOfWidgetInList)
+			{
+				const FString LoggedName = WidgetInfo.LastUpdatedFrame == GFrameNumber ? // Show differently a widget that was updated on a previous frame
+					WidgetInfo.WidgetName :
+					FString::Printf(TEXT("%s ( %3d frames ago )"), *WidgetInfo.WidgetName, GFrameNumber - WidgetInfo.LastUpdatedFrame);
+
+				MakeText(LoggedName, FVector2f(0.f, (12.f * NumberOfWidgetsLoggedThisFrame) + TextElementY), DrawWidgetNameColor);
+				++NumberOfWidgetsLoggedThisFrame;
+			}
+		}
+	}
+	bLogWidgetNameOnce = false;
+
+	{
+		FString NumberOfWidgetsDrawn = GetNumberOfWidgetsUpdatedLogString(NumberOfWidgetsUpdatedThisFrame);
+		MakeText(NumberOfWidgetsDrawn, WidgetLogLocation, DrawWidgetNameColor);
+	}
+
+	if (bDisplayWidgetsNameList && NumberOfWidgetsUpdatedThisFrame > (uint32)MaxNumberOfWidgetInList)
+	{
+		FString WidgetDisplayName = FString::Printf(TEXT("   %d more widgets"), NumberOfWidgetsUpdatedThisFrame - MaxNumberOfWidgetInList);
+		MakeText(WidgetDisplayName, FVector2f(0.f, (12.f * NumberOfWidgetsLoggedThisFrame) + TextElementY), FLinearColor::White);
+	}
+}
+
+void FConsoleSlateDebuggerPassBase::BuildInitialWidgetList()
+{
+	// Do not log the widget names on initial build
+	TGuardValue LogWidgetNameGuard(bLogWidgetName, false);
+
+	TArray<TSharedRef<SWindow>> SlateWindows = FSlateApplicationBase::Get().GetTopLevelWindows();
+	for (TArray<TSharedRef<SWindow>>::TIterator It(SlateWindows); It; ++It)
+	{
+		const TSharedRef<SWindow>& Window = *It;
+		
+		if (!(bDebugGameWindowOnly && Window->GetType() != EWindowType::GameWindow && Window->GetTag() != PIEWindowTag))
+		{
+			const FConsoleSlateDebuggerUtility::TSWindowId WindowId = FConsoleSlateDebuggerUtility::GetId(*Window);
+			AddInitialVisibleWidget(*StaticCastSharedRef<SWidget>(Window), WindowId);
+		}
+		
+		SlateWindows.Append(Window->GetChildWindows());
+	}
+}
+
+void FConsoleSlateDebuggerPassBase::AddInitialVisibleWidget(const SWidget& Widget, const FConsoleSlateDebuggerUtility::TSWindowId WindowId)
+{
+	if (!Widget.GetVisibility().IsVisible() || Widget.Debug_GetLastPaintFrame() == 0)
+	{
+		return;
+	}
+
+	const bool bIncrementUpdateCount = false; // This is called from BuildInitialWidgetList and we therefore do not want to record an update event
+	AddUpdatedWidget(Widget, WindowId, bIncrementUpdateCount);
+
+	const FChildren* Children = const_cast<SWidget&>(Widget).Debug_GetChildrenForReflector();
+	if (ensure(Children))
+	{
+		Children->ForEachWidget([&](const SWidget& ChildWidget)
+		{
+			AddInitialVisibleWidget(ChildWidget, WindowId);
+		});
+	}
+}
+
+FConsoleSlateDebuggerPassBase::FWidgetInfo& FConsoleSlateDebuggerPassBase::AddUpdatedWidget_Internal(const SWidget& Widget, const FConsoleSlateDebuggerUtility::TSWindowId WindowId, uint32 LastUpdatedFrame)
+{
+	// Use the Widget pointer for the id.
+	//That may introduce bug when a widget is destroyed and the same memory is reused for another widget. We do not care for this debug tool.
+	//We do not keep the widget alive or reuse it later, cache all the info that we need.
+	const FConsoleSlateDebuggerUtility::TSWidgetId WidgetId = FConsoleSlateDebuggerUtility::GetId(Widget);
+
+	FWidgetInfo* FoundItem = Widgets.Find(WidgetId);
+	if (FoundItem == nullptr)
+	{
+		FoundItem = &Widgets.Add(WidgetId);
+		FoundItem->UpdateCount = 0;
+		FoundItem->Window = WindowId;
+		FoundItem->WidgetName = bUseWidgetPathAsName ? FReflectionMetaData::GetWidgetPath(Widget) : FReflectionMetaData::GetWidgetDebugInfo(Widget);
+	}
+	else if (!ensure(FoundItem->Window == WindowId))
+	{
+		FoundItem->Window = WindowId;
+	}
+	
+	if (bLogWidgetName)
+	{
+		UE_LOG(LogSlateDebugger, Log, TEXT("%s"), *(FoundItem->WidgetName));
+	}
+	
+	FoundItem->Widget = Widget.AsWeak();
+	FoundItem->PaintLocation = Widget.GetPersistentState().AllottedGeometry.GetAbsolutePosition();
+	FoundItem->PaintSize = Widget.GetPersistentState().AllottedGeometry.GetAbsoluteSize();
+	FoundItem->LastUpdatedFrame = LastUpdatedFrame;
+	FoundItem->LastUpdatedTime = FSlateApplicationBase::Get().GetCurrentTime();
+	if (GFrameNumber != FoundItem->LastUpdatedFrame) // if we do not have the time at which it was drawn, we estimate
+	{
+		FoundItem->LastUpdatedTime -= (GFrameNumber - FoundItem->LastUpdatedFrame) * FApp::GetDeltaTime(); 
+	}
+
+	return *FoundItem;
+}
+
+#endif // WITH_SLATE_DEBUGGING
\ No newline at end of file
diff --git a/Engine/Source/Runtime/SlateCore/Private/Debugging/ConsoleSlateDebuggerPassBase.h b/Engine/Source/Runtime/SlateCore/Private/Debugging/ConsoleSlateDebuggerPassBase.h
new file mode 100644
--- /dev/null
+++ b/Engine/Source/Runtime/SlateCore/Private/Debugging/ConsoleSlateDebuggerPassBase.h
@@ -0,0 +1,119 @@
+﻿// Copyright Epic Games, Inc. All Rights Reserved.
+
+#pragma once
+
+
+#include "Debugging/SlateDebugging.h"
+
+#if WITH_SLATE_DEBUGGING
+
+#include "CoreMinimal.h"
+#include "Debugging/ConsoleSlateDebuggerUtility.h"
+#include "Delegates/Delegate.h"
+#include "HAL/IConsoleManager.h"
+#include "Rendering/DrawElements.h"
+
+class FConsoleSlateDebuggerPassBase
+{
+public:
+	FConsoleSlateDebuggerPassBase();
+	virtual ~FConsoleSlateDebuggerPassBase();
+
+	/* Starts debugging by setting the enabled CVar to true. See GetEnabledCVar */
+	void StartDebugging();
+	/* Stops debugging by setting the enabled CVar to false. See GetEnabledCVar */
+	void StopDebugging();
+	
+	bool IsEnabled() const { return bEnabled; }
+
+	virtual void LoadConfig();
+	virtual void SaveConfig();
+
+protected:
+	/* Return the name of the config section used to load and save the config */
+	virtual FString GetConfigSection() const = 0;
+	/* Return the text to be shown on screen indicating the number of widgets updated when bDisplayWidgetsNameList is true */
+	virtual FString GetNumberOfWidgetsUpdatedLogString(uint32 NumberOfWidgetsUpdatedThisFrame) const = 0;
+	/* Return the CVar used to enable/disable the debugger pass. Used by Start/StopDebugging */
+	virtual FAutoConsoleVariableRef& GetEnabledCVar() = 0;
+	/* Function called when the debugging starts. Used to attach delegates */
+	virtual void StartDebugging_Internal();
+	/* Function called when the debugging stops. Used to remove delegates */
+	virtual void StopDebugging_Internal();
+
+	/**
+	 * Called when a widget should be added to the internal widget list.
+	 * @param Widget The Widget to add to the list
+	 * @param WindowId The ID of the window owning the widget
+	 * @param bIncrementUpdateCount If true, the UpdateCount of the widget should be increased. This can ve false when building the initial widget list
+	 */
+	virtual void AddUpdatedWidget(const SWidget& Widget, const FConsoleSlateDebuggerUtility::TSWindowId WindowId, bool bIncrementUpdateCount = true) = 0;
+	
+	struct FWidgetInfo 
+	{
+		FConsoleSlateDebuggerUtility::TSWindowId Window;
+		TWeakPtr<const SWidget> Widget; // used to check if the widget has been destroyed
+		FVector2f PaintLocation;
+		FVector2f PaintSize;
+		FString WidgetName;
+		
+		uint32 LastUpdatedFrame;
+		double LastUpdatedTime;
+		int32 UpdateCount;
+	};
+
+	/**
+	 * The internal function that should be called from AddUpdatedWidget to update the internal widget list.
+	 * This internal version has a parameter LastUpdatedFrame used to set when the widget was last updated.
+	 * @param Widget The Widget to add to the list
+	 * @param WindowId The ID of the window owning the widget
+	 * @param LastUpdatedFrame The Frame at which the widget was last updated
+	 * @return Returns the internal FWidgetInfo allowing additional properties to be set, like the widget's UpdateCount
+	 */
+	FWidgetInfo& AddUpdatedWidget_Internal(const SWidget& Widget, const FConsoleSlateDebuggerUtility::TSWindowId WindowId, uint32 LastUpdatedFrame);
+
+	/* Should be called by the GetEnabledCVar when it is updated. Will end up calling Start/StopDebugging_Internal */
+	void HandleEnabled(IConsoleVariable* Variable);
+	/* Should be called by the CVar enabling a one time log of the updated widgets */
+	void HandleLogOnce();
+	/* Should be called by the CVar enabling the updated widget list to be shown on screen */
+	void HandleToggleWidgetNameList();
+	/* Should be called by CVars that should only save the current config */
+	void SaveConfigOnVariableChanged(IConsoleVariable* Variable);
+	/* Should be called by the CVar that changes the scope of the widgets retrieved. Ends up rebuilding the widget list */
+	void HandleDebugGameWindowOnlyChanged(IConsoleVariable* Variable);
+	/* Called on EndFrame to reset the counts and remove deleted widgets */
+	void HandleEndFrame();
+	/* Called by Slate to draw additional elements on screen. */
+	void HandlePaintDebugInfo(const FPaintArgs& InArgs, const FGeometry& InAllottedGeometry, FSlateWindowElementList& InOutDrawElements, int32& InOutLayerId);
+
+private:
+	/* Not all widgets are accessible from the PaintEvents, so we need to loop through all of them to add all current ones */
+	void BuildInitialWidgetList();
+
+	void AddInitialVisibleWidget(const SWidget& Widget, const FConsoleSlateDebuggerUtility::TSWindowId WindowId);
+	
+protected:
+	//~ Settings
+	bool bEnabled;
+	bool bDisplayWidgetsNameList;
+	bool bUseWidgetPathAsName;
+	bool bDrawBox;
+	bool bDrawBorder;
+	bool bLogWidgetName;
+	bool bLogWidgetNameOnce;
+	bool bDebugGameWindowOnly;
+	FLinearColor MostRecentColor;
+	FLinearColor LeastRecentColor;
+	FLinearColor DrawWidgetNameColor;
+	int32 MaxNumberOfWidgetInList;
+	float FadeDuration;
+	const FName PIEWindowTag;
+
+	FVector2f WidgetLogLocation{10.f, 10.f};
+
+	using TUpdatedWidgetMap = TMap<FConsoleSlateDebuggerUtility::TSWidgetId, FWidgetInfo>;
+	TUpdatedWidgetMap Widgets;
+};
+
+#endif //WITH_SLATE_DEBUGGING
diff --git a/Engine/Source/Runtime/SlateCore/Private/Debugging/ConsoleSlateDebuggerPrepass.cpp b/Engine/Source/Runtime/SlateCore/Private/Debugging/ConsoleSlateDebuggerPrepass.cpp
new file mode 100644
--- /dev/null
+++ b/Engine/Source/Runtime/SlateCore/Private/Debugging/ConsoleSlateDebuggerPrepass.cpp
@@ -0,0 +1,117 @@
+﻿// Copyright Epic Games, Inc. All Rights Reserved.
+
+#include "ConsoleSlateDebuggerPrepass.h"
+
+#if WITH_SLATE_DEBUGGING
+
+#include "CoreGlobals.h"
+#include "Application/SlateApplicationBase.h"
+#include "Debugging/SlateDebugging.h"
+#include "Layout/WidgetPath.h"
+#include "Misc/App.h"
+#include "Misc/ConfigCacheIni.h"
+
+#define LOCTEXT_NAMESPACE "ConsoleSlateDebuggerPrepass"
+
+FConsoleSlateDebuggerPrepass::FConsoleSlateDebuggerPrepass()
+	: FConsoleSlateDebuggerPassBase()
+	, EnabledRefCVar(
+		TEXT("SlateDebugger.Prepass.Enable")
+		, bEnabled
+		, TEXT("Start/Stop the prepassed widget debug tool. It shows when widgets are prepassed.")
+		, FConsoleVariableDelegate::CreateRaw(this, &FConsoleSlateDebuggerPrepass::HandleEnabled))
+	, ShowPrepassWidgetCommand(
+		TEXT("SlateDebugger.Prepass.Start")
+		, TEXT("Start the prepassed widget debug tool. Use to show widget that have been prepassed this frame.")
+		, FConsoleCommandDelegate::CreateRaw(this, &FConsoleSlateDebuggerPrepass::StartDebugging))
+	, HidePrepassWidgetCommand(
+		TEXT("SlateDebugger.Prepass.Stop")
+		, TEXT("Stop the prepassed widget debug tool.")
+		, FConsoleCommandDelegate::CreateRaw(this, &FConsoleSlateDebuggerPrepass::StopDebugging))
+	, LogPrepassedWidgetOnceCommand(
+		TEXT("SlateDebugger.Prepass.LogOnce")
+		, TEXT("Log the names of all widgets that were prepassed during the last update.")
+		, FConsoleCommandDelegate::CreateRaw(this, &FConsoleSlateDebuggerPrepass::HandleLogOnce))
+	, EnableWidgetsNameListRefCVar(
+		TEXT("SlateDebugger.Prepass.EnableWidgetNameList")
+		, bDisplayWidgetsNameList
+		, TEXT("Start/Stop displaying the name of the widgets that have been prepassed.")
+		, FConsoleVariableDelegate::CreateRaw(this, &FConsoleSlateDebuggerPrepass::SaveConfigOnVariableChanged))
+	, ToggleWidgetsNameListCommand(
+		TEXT("SlateDebugger.Prepass.ToggleWidgetNameList")
+		, TEXT("Option displaying the name of the widgets that have been prepassed.")
+		, FConsoleCommandDelegate::CreateRaw(this, &FConsoleSlateDebuggerPrepass::HandleToggleWidgetNameList))
+	, MaxNumberOfWidgetInListRefCVar(
+		TEXT("SlateDebugger.Prepass.MaxNumberOfWidgetDisplayedInList")
+		, MaxNumberOfWidgetInList
+		, TEXT("The max number of widgets that will be displayed when DisplayWidgetNameList is active.")
+		, FConsoleVariableDelegate::CreateRaw(this, &FConsoleSlateDebuggerPrepass::SaveConfigOnVariableChanged))
+	, OnlyGameWindowRefCVar(
+		TEXT("SlateDebugger.Prepass.OnlyGameWindow")
+		, bDebugGameWindowOnly
+		, TEXT("Option to only debug the game window")
+		, FConsoleVariableDelegate::CreateRaw(this, &FConsoleSlateDebuggerPrepass::HandleDebugGameWindowOnlyChanged))
+	, DrawBorderEnabledRefCVar(
+		TEXT("SlateDebugger.Prepass.DrawBorder")
+		, bDrawBorder
+		, TEXT("Draw a border around the widgets being prepassed")
+		, FConsoleVariableDelegate::CreateRaw(this, &FConsoleSlateDebuggerPrepass::SaveConfigOnVariableChanged))
+	, DrawFillEnabledRefCVar(
+		TEXT("SlateDebugger.Prepass.DrawFill")
+		, bDrawBox
+		, TEXT("Fill the widgets being prepassed")
+		, FConsoleVariableDelegate::CreateRaw(this, &FConsoleSlateDebuggerPrepass::SaveConfigOnVariableChanged))
+{
+	WidgetLogLocation = FVector2f(10.f, 20.f); // Different from FConsoleSlateDebuggerPaint so the texts do not overlap
+}
+
+void FConsoleSlateDebuggerPrepass::LoadConfig()
+{
+	FConsoleSlateDebuggerPassBase::LoadConfig();
+	
+	OnlyGameWindowRefCVar->Set(bDebugGameWindowOnly, ECVF_SetByProjectSetting);
+	EnableWidgetsNameListRefCVar->Set(bDisplayWidgetsNameList, ECVF_SetByProjectSetting);
+	DrawFillEnabledRefCVar->Set(bDrawBox, ECVF_SetByProjectSetting);
+	DrawBorderEnabledRefCVar->Set(bDrawBorder, ECVF_SetByProjectSetting);
+}
+
+void FConsoleSlateDebuggerPrepass::StartDebugging_Internal()
+{
+	FConsoleSlateDebuggerPassBase::StartDebugging_Internal();
+	
+	FSlateDebugging::EndWidgetPrepass.AddRaw(this, &FConsoleSlateDebuggerPrepass::HandleEndWidgetPrepass);
+}
+
+void FConsoleSlateDebuggerPrepass::StopDebugging_Internal()
+{
+	FSlateDebugging::EndWidgetPrepass.RemoveAll(this);
+	
+	FConsoleSlateDebuggerPassBase::StopDebugging_Internal();
+}
+
+void FConsoleSlateDebuggerPrepass::AddUpdatedWidget(const SWidget& Widget, const FConsoleSlateDebuggerUtility::TSWindowId WindowId, bool bIncrementUpdateCount)
+{
+	FWidgetInfo& WidgetInfo = AddUpdatedWidget_Internal(Widget, WindowId, Widget.Debug_GetLastPrepassFrame());
+
+	if (bIncrementUpdateCount)
+	{
+		++WidgetInfo.UpdateCount;
+	}
+}
+
+void FConsoleSlateDebuggerPrepass::HandleEndWidgetPrepass(const SWidget* Widget)
+{
+	// Exclude all windows but the game window
+	const TSharedPtr<SWindow> WindowToDrawIn = FSlateApplicationBase::Get().FindWidgetWindow(Widget->AsShared());
+	if (!WindowToDrawIn || (bDebugGameWindowOnly && (WindowToDrawIn->GetType() != EWindowType::GameWindow && WindowToDrawIn->GetTag() != PIEWindowTag)))
+	{
+		return;
+	}
+
+	const FConsoleSlateDebuggerUtility::TSWindowId WindowId = FConsoleSlateDebuggerUtility::GetId(*WindowToDrawIn);
+	AddUpdatedWidget(*Widget, WindowId);
+}
+
+#undef LOCTEXT_NAMESPACE
+
+#endif // WITH_SLATE_DEBUGGING
diff --git a/Engine/Source/Runtime/SlateCore/Private/Debugging/ConsoleSlateDebuggerPrepass.h b/Engine/Source/Runtime/SlateCore/Private/Debugging/ConsoleSlateDebuggerPrepass.h
new file mode 100644
--- /dev/null
+++ b/Engine/Source/Runtime/SlateCore/Private/Debugging/ConsoleSlateDebuggerPrepass.h
@@ -0,0 +1,59 @@
+﻿// Copyright Epic Games, Inc. All Rights Reserved.
+
+#pragma once
+
+
+#include "Debugging/SlateDebugging.h"
+
+#if WITH_SLATE_DEBUGGING
+
+#include "CoreMinimal.h"
+#include "ConsoleSlateDebuggerPassBase.h"
+#include "Debugging/ConsoleSlateDebuggerUtility.h"
+#include "Delegates/Delegate.h"
+#include "HAL/IConsoleManager.h"
+#include "Rendering/DrawElements.h"
+
+/**
+ * Allows debugging the behavior of SWidget::Prepass from the console.
+ * Basics:
+ *   Start - SlateDebugger.Prepass.Start
+ *   Stop  - SlateDebugger.Prepass.Stop
+ */
+class FConsoleSlateDebuggerPrepass : public FConsoleSlateDebuggerPassBase
+{
+public:
+	FConsoleSlateDebuggerPrepass();
+
+	virtual void LoadConfig() override;
+
+protected:
+	virtual FString GetConfigSection() const override { return TEXT("SlateDebugger.Prepass"); }
+	virtual FString GetNumberOfWidgetsUpdatedLogString(uint32 NumberOfWidgetsUpdatedThisFrame) const override
+	{
+		return FString::Printf(TEXT("%d widgets prepassed"), NumberOfWidgetsUpdatedThisFrame);
+	}
+	virtual FAutoConsoleVariableRef& GetEnabledCVar() override { return EnabledRefCVar; }
+	virtual void StartDebugging_Internal() override;
+	virtual void StopDebugging_Internal() override;
+	
+	virtual void AddUpdatedWidget(const SWidget& Widget, const FConsoleSlateDebuggerUtility::TSWindowId WindowId, bool bIncrementUpdateCount = true) override;
+
+private:
+	void HandleEndWidgetPrepass(const SWidget* Widget);
+
+private:
+	//~ Console objects
+	FAutoConsoleVariableRef EnabledRefCVar;
+	FAutoConsoleCommand ShowPrepassWidgetCommand;
+	FAutoConsoleCommand HidePrepassWidgetCommand;
+	FAutoConsoleCommand LogPrepassedWidgetOnceCommand;
+	FAutoConsoleVariableRef EnableWidgetsNameListRefCVar;
+	FAutoConsoleCommand ToggleWidgetsNameListCommand;
+	FAutoConsoleVariableRef MaxNumberOfWidgetInListRefCVar;
+	FAutoConsoleVariableRef OnlyGameWindowRefCVar;
+	FAutoConsoleVariableRef DrawBorderEnabledRefCVar;
+	FAutoConsoleVariableRef DrawFillEnabledRefCVar;
+};
+
+#endif //WITH_SLATE_DEBUGGING
\ No newline at end of file
diff --git a/Engine/Source/Runtime/SlateCore/Private/Debugging/SlateDebugging.cpp b/Engine/Source/Runtime/SlateCore/Private/Debugging/SlateDebugging.cpp
--- a/Engine/Source/Runtime/SlateCore/Private/Debugging/SlateDebugging.cpp
+++ b/Engine/Source/Runtime/SlateCore/Private/Debugging/SlateDebugging.cpp
@@ -399,6 +399,10 @@ FSlateDebugging::FBeginWidgetPaint FSlateDebugging::BeginWidgetPaint;
 
 FSlateDebugging::FEndWidgetPaint FSlateDebugging::EndWidgetPaint;
 
+FSlateDebugging::FBeginWidgetPrepass FSlateDebugging::BeginWidgetPrepass;
+
+FSlateDebugging::FEndWidgetPrepass FSlateDebugging::EndWidgetPrepass;
+
 FSlateDebugging::FPaintDebugElements FSlateDebugging::PaintDebugElements;
 
 PRAGMA_DISABLE_DEPRECATION_WARNINGS
diff --git a/Engine/Source/Runtime/SlateCore/Private/SlateCoreModule.cpp b/Engine/Source/Runtime/SlateCore/Private/SlateCoreModule.cpp
--- a/Engine/Source/Runtime/SlateCore/Private/SlateCoreModule.cpp
+++ b/Engine/Source/Runtime/SlateCore/Private/SlateCoreModule.cpp
@@ -5,6 +5,7 @@
 #include "Debugging/ConsoleSlateDebuggerInvalidate.h"
 #include "Debugging/ConsoleSlateDebuggerInvalidationRoot.h"
 #include "Debugging/ConsoleSlateDebuggerPaint.h"
+#include "Debugging/ConsoleSlateDebuggerPrepass.h"
 #include "Debugging/ConsoleSlateDebuggerUpdate.h"
 #include "Debugging/ConsoleSlateDebuggerBreak.h"
 #include "Modules/ModuleInterface.h"
@@ -32,6 +33,9 @@ public:
 		SlateDebuggerInvalidate = MakeUnique<FConsoleSlateDebuggerInvalidate>();
 		SlateDebuggerInvalidationRoot = MakeUnique<FConsoleSlateDebuggerInvalidationRoot>();
 		SlateDebuggerPaint = MakeUnique<FConsoleSlateDebuggerPaint>();
+		SlateDebuggerPaint->LoadConfig();
+		SlateDebuggerPrepass = MakeUnique<FConsoleSlateDebuggerPrepass>();
+		SlateDebuggerPrepass->LoadConfig();
 		SlateDebuggerUpdate = MakeUnique<FConsoleSlateDebuggerUpdate>();
 		SlateDebuggerBreak = MakeUnique<FConsoleSlateDebuggerBreak>();
 #endif
@@ -43,6 +47,7 @@ private:
 	TUniquePtr<FConsoleSlateDebuggerInvalidate> SlateDebuggerInvalidate;
 	TUniquePtr<FConsoleSlateDebuggerInvalidationRoot> SlateDebuggerInvalidationRoot;
 	TUniquePtr<FConsoleSlateDebuggerPaint> SlateDebuggerPaint;
+	TUniquePtr<FConsoleSlateDebuggerPrepass> SlateDebuggerPrepass;
 	TUniquePtr<FConsoleSlateDebuggerUpdate> SlateDebuggerUpdate;
 	TUniquePtr<FConsoleSlateDebuggerBreak> SlateDebuggerBreak;
 #endif
diff --git a/Engine/Source/Runtime/SlateCore/Private/Widgets/SWidget.cpp b/Engine/Source/Runtime/SlateCore/Private/Widgets/SWidget.cpp
--- a/Engine/Source/Runtime/SlateCore/Private/Widgets/SWidget.cpp
+++ b/Engine/Source/Runtime/SlateCore/Private/Widgets/SWidget.cpp
@@ -1731,6 +1731,10 @@ void SWidget::ArrangeChildren(const FGeometry& AllottedGeometry, FArrangedChildr
 
 void SWidget::Prepass_Internal(float InLayoutScaleMultiplier)
 {
+#if WITH_SLATE_DEBUGGING
+	FSlateDebugging::BeginWidgetPrepass.Broadcast(this);
+#endif
+
 	PrepassLayoutScaleMultiplierValue = InLayoutScaleMultiplier;
 	bPrepassLayoutScaleMultiplierSet = true;
 
@@ -1754,6 +1758,10 @@ void SWidget::Prepass_Internal(float InLayoutScaleMultiplier)
 		// Cache this widget's desired size.
 		CacheDesiredSize(GetPrepassLayoutScaleMultiplier());
 		bNeedsPrepass = false;
+#if WITH_SLATE_DEBUGGING
+		Debug_UpdateLastPrepassFrame();
+		FSlateDebugging::EndWidgetPrepass.Broadcast(this);
+#endif
 	}
 }
 
diff --git a/Engine/Source/Runtime/SlateCore/Public/Debugging/SlateDebugging.h b/Engine/Source/Runtime/SlateCore/Public/Debugging/SlateDebugging.h
--- a/Engine/Source/Runtime/SlateCore/Public/Debugging/SlateDebugging.h
+++ b/Engine/Source/Runtime/SlateCore/Public/Debugging/SlateDebugging.h
@@ -296,6 +296,14 @@ public:
 	DECLARE_MULTICAST_DELEGATE_ThreeParams(FEndWidgetPaint, const SWidget* /*Widget*/, const FSlateWindowElementList& /*OutDrawElements*/, int32 /*LayerId*/);
 	static SLATECORE_API FEndWidgetPaint EndWidgetPaint;
 
+	/** Called just before a widget starts its prepass. */
+	DECLARE_MULTICAST_DELEGATE_OneParam(FBeginWidgetPrepass, const SWidget* /*Widget*/);
+	static SLATECORE_API FBeginWidgetPrepass BeginWidgetPrepass;
+
+	/** Called after a widget finishes its prepass. */
+	DECLARE_MULTICAST_DELEGATE_OneParam(FEndWidgetPrepass, const SWidget* /*Widget*/);
+	static SLATECORE_API FEndWidgetPrepass EndWidgetPrepass;
+
 	/** Called once the window is drawn (including PaintDeferred). It can be used to draw additional debug info on top of the window. */
 	DECLARE_MULTICAST_DELEGATE_FourParams(FPaintDebugElements, const FPaintArgs& /*InArgs*/, const FGeometry& /*InAllottedGeometry*/, FSlateWindowElementList& /*InOutDrawElements*/, int32& /*InOutLayerId*/);
 	static SLATECORE_API FPaintDebugElements PaintDebugElements;
diff --git a/Engine/Source/Runtime/SlateCore/Public/Widgets/SWidget.h b/Engine/Source/Runtime/SlateCore/Public/Widgets/SWidget.h
--- a/Engine/Source/Runtime/SlateCore/Public/Widgets/SWidget.h
+++ b/Engine/Source/Runtime/SlateCore/Public/Widgets/SWidget.h
@@ -779,11 +779,19 @@ public:
 	{ 
 		return LastPaintFrame; 
 	}
+	uint32 Debug_GetLastPrepassFrame() const 
+	{ 
+		return LastPrepassFrame; 
+	}
 private:
 	void Debug_UpdateLastPaintFrame() 
 	{ 
 		LastPaintFrame = GFrameNumber; 
 	}
+	void Debug_UpdateLastPrepassFrame() 
+	{ 
+		LastPrepassFrame = GFrameNumber; 
+	}
 #endif // WITH_SLATE_DEBUGGING
 
 public:
@@ -1860,6 +1868,8 @@ private:
 #if WITH_SLATE_DEBUGGING
 	/** The last time this widget got painted. */
 	uint32 LastPaintFrame = 0;
+	/** The last time this widget got prepassed. */
+	uint32 LastPrepassFrame = 0;
 	/** Flag to help detect when we access an invalid Widget. */
 	uint8 Debug_DestroyedTag = 0xDC;
 #endif // WITH_SLATE_DEBUGGING
