SlateCore several changes for code size

1. De-inline some ctors/assignment operators (especially FReply)
2. Make UnicodeBlockRange POD to avoid 30+k of array initializer code
3. Remove unnecessary no-op dtor in SlateColor
4. Move some globals from headers to .cpp

diff --git a/Engine/Plugins/Animation/BlendStack/Source/Editor/BlendStackEditor.Build.cs b/Engine/Plugins/Animation/BlendStack/Source/Editor/BlendStackEditor.Build.cs
--- a/Engine/Plugins/Animation/BlendStack/Source/Editor/BlendStackEditor.Build.cs
+++ b/Engine/Plugins/Animation/BlendStack/Source/Editor/BlendStackEditor.Build.cs
@@ -26,6 +26,7 @@ public class BlendStackEditor : ModuleRules
 			new string[]
 			{
 				"BlueprintGraph",
+				"SlateCore",
 				"UnrealEd"
 			}
 		);
diff --git a/Engine/Plugins/Experimental/MeshModelingToolsetExp/Source/MeshModelingToolsExp/MeshModelingToolsExp.Build.cs b/Engine/Plugins/Experimental/MeshModelingToolsetExp/Source/MeshModelingToolsExp/MeshModelingToolsExp.Build.cs
--- a/Engine/Plugins/Experimental/MeshModelingToolsetExp/Source/MeshModelingToolsExp/MeshModelingToolsExp.Build.cs
+++ b/Engine/Plugins/Experimental/MeshModelingToolsetExp/Source/MeshModelingToolsExp/MeshModelingToolsExp.Build.cs
@@ -50,6 +50,7 @@ public class MeshModelingToolsExp : ModuleRules
 				"RenderCore",
 				"InputCore",
 				"PhysicsCore",
+				"SlateCore",
 				// ... add private dependencies that you statically link with here ...
 			}
 		);
@@ -70,4 +71,4 @@ public class MeshModelingToolsExp : ModuleRules
 			}
 		);
 	}
-}
\ No newline at end of file
+}
diff --git a/Engine/Plugins/Runtime/MeshModelingToolset/Source/MeshModelingTools/MeshModelingTools.Build.cs b/Engine/Plugins/Runtime/MeshModelingToolset/Source/MeshModelingTools/MeshModelingTools.Build.cs
--- a/Engine/Plugins/Runtime/MeshModelingToolset/Source/MeshModelingTools/MeshModelingTools.Build.cs
+++ b/Engine/Plugins/Runtime/MeshModelingToolset/Source/MeshModelingTools/MeshModelingTools.Build.cs
@@ -32,6 +32,7 @@ public class MeshModelingTools : ModuleRules
 			{
 				"CoreUObject",
 				"Engine",
+				"SlateCore",
 				// ... add private dependencies that you statically link with here ...
 			}
 		);
diff --git a/Engine/Source/Editor/FontEditor/Private/SCompositeFontEditor.cpp b/Engine/Source/Editor/FontEditor/Private/SCompositeFontEditor.cpp
--- a/Engine/Source/Editor/FontEditor/Private/SCompositeFontEditor.cpp
+++ b/Engine/Source/Editor/FontEditor/Private/SCompositeFontEditor.cpp
@@ -1524,7 +1524,7 @@ void SCharacterRangeEditor::Construct(const FArguments& InArgs)
 		{
 			RangeSelectionComboData.Emplace(MakeShared<FUnicodeBlockRange>(UnicodeBlockRange));
 
-			if(CurrentRangeSelection.IsSet() && CurrentRangeSelection->Range == UnicodeBlockRange.Range)
+			if(CurrentRangeSelection.IsSet() && CurrentRangeSelection->GetRange() == UnicodeBlockRange.GetRange())
 			{
 				CurrentRangeSelectionItem = RangeSelectionComboData.Last();
 			}
@@ -1740,7 +1740,7 @@ void SCharacterRangeEditor::CacheCurrentRangeSelection()
 	{
 		for(const FUnicodeBlockRange& UnicodeBlockRange : UnicodeBlockRanges)
 		{
-			if(UnicodeBlockRange.Range == *CharacterRangePtr)
+			if(UnicodeBlockRange.GetRange() == *CharacterRangePtr)
 			{
 				CurrentRangeSelection = UnicodeBlockRange;
 			}
@@ -1750,14 +1750,14 @@ void SCharacterRangeEditor::CacheCurrentRangeSelection()
 
 FText SCharacterRangeEditor::GetCurrentRangeSelectionDisplayName() const
 {
-	return CurrentRangeSelection.IsSet() ? CurrentRangeSelection->DisplayName : LOCTEXT("UnicodeBlock_CustomSelection", "Custom");
+	return CurrentRangeSelection.IsSet() ? CurrentRangeSelection->GetDisplayName() : LOCTEXT("UnicodeBlock_CustomSelection", "Custom");
 }
 
 void SCharacterRangeEditor::OnRangeSelectionComboOpening()
 {
 	RangeSelectionComboData.Sort([](const TSharedPtr<FUnicodeBlockRange>& One, const TSharedPtr<FUnicodeBlockRange>& Two)
 	{
-		return One->DisplayName.CompareTo(Two->DisplayName) < 0;
+		return One->GetDisplayName().CompareTo(Two->GetDisplayName()) < 0;
 	});
 
 	if(RangeSelectionCombo.IsValid())
@@ -1776,7 +1776,7 @@ void SCharacterRangeEditor::OnRangeSelectionChanged(TSharedPtr<FUnicodeBlockRang
 			const FScopedTransaction Transaction(LOCTEXT("UpdateCharacterRange", "Update Character Range"));
 			CompositeFontEditorPtr->GetFontObject()->Modify();
 
-			*CharacterRangePtr = InNewRangeSelection->Range;
+			*CharacterRangePtr = InNewRangeSelection->GetRange();
 			CurrentRangeSelection = *InNewRangeSelection;
 
 			CompositeFontEditorPtr->FlushCachedFont();
@@ -1787,8 +1787,8 @@ void SCharacterRangeEditor::OnRangeSelectionChanged(TSharedPtr<FUnicodeBlockRang
 TSharedRef<SWidget> SCharacterRangeEditor::MakeRangeSelectionWidget(TSharedPtr<FUnicodeBlockRange> InRangeSelection)
 {
 	return SNew(STextBlock)
-		.Text(InRangeSelection->DisplayName)
-		.ToolTipText(FText::Format(LOCTEXT("RangeSelectionTooltipFmt", "{0} ({1} - {2})"), InRangeSelection->DisplayName, FText::AsCultureInvariant(FString::Printf(TEXT("0x%04x"), InRangeSelection->Range.GetLowerBoundValue())), FText::AsCultureInvariant(FString::Printf(TEXT("0x%04x"), InRangeSelection->Range.GetUpperBoundValue()))));
+		.Text(InRangeSelection->GetDisplayName())
+		.ToolTipText(FText::Format(LOCTEXT("RangeSelectionTooltipFmt", "{0} ({1})"), InRangeSelection->GetDisplayName(), FText::AsCultureInvariant(FString::Printf(TEXT("0x%04x - 0x%04x"), InRangeSelection->RangeLower, InRangeSelection->RangeUpper))));
 }
 
 void SFontScalingFactorEditor::Construct(const FArguments& InArgs)
diff --git a/Engine/Source/Editor/UnrealEd/Private/Toolkits/AssetEditorToolkit.cpp b/Engine/Source/Editor/UnrealEd/Private/Toolkits/AssetEditorToolkit.cpp
--- a/Engine/Source/Editor/UnrealEd/Private/Toolkits/AssetEditorToolkit.cpp
+++ b/Engine/Source/Editor/UnrealEd/Private/Toolkits/AssetEditorToolkit.cpp
@@ -918,6 +918,11 @@ void FAssetEditorToolkit::SetAssetEditorModeManager(FAssetEditorModeManager* InM
 	PRAGMA_ENABLE_DEPRECATION_WARNINGS
 }
 
+FReply FAssetEditorToolkit::OnViewportDrop(const FGeometry& MyGeometry, const FDragDropEvent& DragDropEvent)
+{
+	return FReply::Unhandled();
+}
+
 void FAssetEditorToolkit::RemoveEditingAsset(UObject* Asset)
 {
 	// Just close the editor tab if it's the last element
diff --git a/Engine/Source/Editor/UnrealEd/Public/Toolkits/AssetEditorToolkit.h b/Engine/Source/Editor/UnrealEd/Public/Toolkits/AssetEditorToolkit.h
--- a/Engine/Source/Editor/UnrealEd/Public/Toolkits/AssetEditorToolkit.h
+++ b/Engine/Source/Editor/UnrealEd/Public/Toolkits/AssetEditorToolkit.h
@@ -315,7 +315,7 @@ public:
 	/** React to a Drag&Drop events from the viewport */
 	virtual void OnViewportDragEnter(const FGeometry& MyGeometry, const FDragDropEvent& DragDropEvent) {}
 	virtual void OnViewportDragLeave(const FDragDropEvent& DragDropEvent) {}
-	virtual FReply OnViewportDrop(const FGeometry& MyGeometry, const FDragDropEvent& DragDropEvent) { return FReply::Unhandled(); }
+	UNREALED_API virtual FReply OnViewportDrop(const FGeometry& MyGeometry, const FDragDropEvent& DragDropEvent);
 
 protected:
 	friend class UAssetEditorToolkitMenuContext;
diff --git a/Engine/Source/Runtime/Engine/Classes/EdGraph/EdGraphSchema.h b/Engine/Source/Runtime/Engine/Classes/EdGraph/EdGraphSchema.h
--- a/Engine/Source/Runtime/Engine/Classes/EdGraph/EdGraphSchema.h
+++ b/Engine/Source/Runtime/Engine/Classes/EdGraph/EdGraphSchema.h
@@ -213,7 +213,7 @@ public:
 	}
 
 	/** Performs a double click on the action */
-	virtual FReply OnDoubleClick(UBlueprint* InBlueprint) { return FReply::Unhandled(); }
+	ENGINE_API virtual FReply OnDoubleClick(UBlueprint* InBlueprint);
 
 	// Updates the category of the *action* and refreshes the search text; does not change the persistent backing item
 	// (e.g., it will not actually move a user added variable or function to a new category)
@@ -1158,12 +1158,12 @@ class UEdGraphSchema : public UObject
 	 * Begins a drag and drop action to drag a graph action into another graph
 	 */
 	UE_DEPRECATED(5.0, "Use version that takes FPointerEvent instead.")
-	virtual FReply BeginGraphDragAction(TSharedPtr<FEdGraphSchemaAction> InAction) const { return FReply::Unhandled(); }
+	ENGINE_API virtual FReply BeginGraphDragAction(TSharedPtr<FEdGraphSchemaAction> InAction) const;
 
 	/*
 	* Begins a drag and drop action to drag a graph action into another graph
 	*/
-	virtual FReply BeginGraphDragAction(TSharedPtr<FEdGraphSchemaAction> InAction, const FPointerEvent& MouseEvent) const { return FReply::Unhandled(); }
+	ENGINE_API virtual FReply BeginGraphDragAction(TSharedPtr<FEdGraphSchemaAction> InAction, const FPointerEvent& MouseEvent) const;
 
 	/**
 	 * Gets display information for a graph
@@ -1186,7 +1186,7 @@ class UEdGraphSchema : public UObject
 	 * @param	InGraph				Graph to set the category for
 	 * @param	InCategory			Pipe "|" separated category for the graph.
 	 */
-	virtual FReply TrySetGraphCategory(const UEdGraph* InGraph, const FText& InCategory) const { return FReply::Unhandled(); }
+	ENGINE_API virtual FReply TrySetGraphCategory(const UEdGraph* InGraph, const FText& InCategory) const;
 
 	/** Called when asset(s) are dropped onto a graph background. */
 	UE_DEPRECATED(5.6, "Use the version of the function accepting FVector2f; this API no longer interfaces directly with double-precision scalars and vectors.")
diff --git a/Engine/Source/Runtime/Engine/Private/EdGraph/EdGraphSchema.cpp b/Engine/Source/Runtime/Engine/Private/EdGraph/EdGraphSchema.cpp
--- a/Engine/Source/Runtime/Engine/Private/EdGraph/EdGraphSchema.cpp
+++ b/Engine/Source/Runtime/Engine/Private/EdGraph/EdGraphSchema.cpp
@@ -18,6 +18,11 @@
 
 #define LOCTEXT_NAMESPACE "EdGraph"
 
+FReply FEdGraphSchemaAction::OnDoubleClick(UBlueprint* InBlueprint)
+{
+	return FReply::Unhandled();
+}
+
 void FEdGraphSchemaAction::CosmeticUpdateCategory(FText NewCategory)
 {
 	Category = MoveTemp(NewCategory);
@@ -1173,12 +1178,27 @@ PRAGMA_DISABLE_DEPRECATION_WARNINGS
 PRAGMA_ENABLE_DEPRECATION_WARNINGS
 }
 
+FReply UEdGraphSchema::BeginGraphDragAction(TSharedPtr<FEdGraphSchemaAction> InAction) const
+{
+	return FReply::Unhandled();
+}
+
+FReply UEdGraphSchema::BeginGraphDragAction(TSharedPtr<FEdGraphSchemaAction> InAction, const FPointerEvent& MouseEvent) const
+{
+	return FReply::Unhandled();
+}
+
 void UEdGraphSchema::GetGraphDisplayInformation(const UEdGraph& Graph, /*out*/ FGraphDisplayInfo& DisplayInfo) const
 {
 	DisplayInfo.PlainName = FText::FromString( Graph.GetName() );
 	DisplayInfo.DisplayName = DisplayInfo.PlainName;
 }
 
+FReply UEdGraphSchema::TrySetGraphCategory(const UEdGraph* InGraph, const FText& InCategory) const
+{
+	return FReply::Unhandled();
+}
+
 #if WITH_EDITOR
 TSharedPtr<INameValidatorInterface> UEdGraphSchema::GetNameValidator(const UBlueprint* InBlueprintObj, const FName& InOriginalName, const UStruct* InValidationScope, const FName& InActionTypeId) const
 {
diff --git a/Engine/Source/Runtime/PreLoadScreen/Private/PreLoadSettingsContainer.cpp b/Engine/Source/Runtime/PreLoadScreen/Private/PreLoadSettingsContainer.cpp
--- a/Engine/Source/Runtime/PreLoadScreen/Private/PreLoadSettingsContainer.cpp
+++ b/Engine/Source/Runtime/PreLoadScreen/Private/PreLoadSettingsContainer.cpp
@@ -352,7 +352,7 @@ void FPreLoadSettingsContainerBase::BuildCustomFont(const FString& FontIdentifie
 			SubFont.Cultures.Append(Language);
 			for (EUnicodeBlockRange& BlockRange : BlockRangesForLanguage)
 			{
-				SubFont.CharacterRanges.Add(FUnicodeBlockRange::GetUnicodeBlockRange(BlockRange).Range);
+				SubFont.CharacterRanges.Add(FUnicodeBlockRange::GetUnicodeBlockRange(BlockRange).GetRange());
 			}
 
 			//Finally append actual font
diff --git a/Engine/Source/Runtime/SlateCore/Private/Fonts/LegacySlateFontInfoCache.cpp b/Engine/Source/Runtime/SlateCore/Private/Fonts/LegacySlateFontInfoCache.cpp
--- a/Engine/Source/Runtime/SlateCore/Private/Fonts/LegacySlateFontInfoCache.cpp
+++ b/Engine/Source/Runtime/SlateCore/Private/Fonts/LegacySlateFontInfoCache.cpp
@@ -102,7 +102,7 @@ TSharedRef<const FCompositeFont> FLegacySlateFontInfoCache::GetDefaultFont()
 			(TYPEFACE).AppendFont(TEXT(NAME), FPaths::EngineContentDir() / TEXT("Editor") / TEXT("Slate") / TEXT("Fonts") / TEXT(FILENAME), HINTING, EFontLoadingPolicy::LazyLoad)
 
 		#define APPEND_RANGE(SUBFONT, RANGE) \
-			(SUBFONT).CharacterRanges.Add(FUnicodeBlockRange::GetUnicodeBlockRange(EUnicodeBlockRange::RANGE).Range)
+			(SUBFONT).CharacterRanges.Add(FUnicodeBlockRange::GetUnicodeBlockRange(EUnicodeBlockRange::RANGE).GetRange())
 
 		TSharedRef<FCompositeFont> MutableDefaultFont = MakeShared<FStandaloneCompositeFont>();
 
diff --git a/Engine/Source/Runtime/SlateCore/Private/Fonts/UnicodeBlockRange.cpp b/Engine/Source/Runtime/SlateCore/Private/Fonts/UnicodeBlockRange.cpp
--- a/Engine/Source/Runtime/SlateCore/Private/Fonts/UnicodeBlockRange.cpp
+++ b/Engine/Source/Runtime/SlateCore/Private/Fonts/UnicodeBlockRange.cpp
@@ -5,7 +5,7 @@
 TArrayView<const FUnicodeBlockRange> FUnicodeBlockRange::GetUnicodeBlockRanges()
 {
 	static const FUnicodeBlockRange UnicodeBlockRanges[] = {
-		#define REGISTER_UNICODE_BLOCK_RANGE(LOWERBOUND, UPPERBOUND, SYMBOLNAME, LOCTEXT_NAMESPACE, LOCTEXT_KEY, LOCTEXT_LITERAL) { EUnicodeBlockRange::SYMBOLNAME, TEXT(LOCTEXT_NAMESPACE), TEXT(LOCTEXT_KEY), TEXT(LOCTEXT_LITERAL), LOWERBOUND, UPPERBOUND },
+		#define REGISTER_UNICODE_BLOCK_RANGE(LOWERBOUND, UPPERBOUND, SYMBOLNAME, LOCTEXT_KEY, LOCTEXT_LITERAL) { EUnicodeBlockRange::SYMBOLNAME, TEXT(LOCTEXT_KEY), TEXT(LOCTEXT_LITERAL), LOWERBOUND, UPPERBOUND },
 		#include "Fonts/UnicodeBlockRange.inl"
 		#undef REGISTER_UNICODE_BLOCK_RANGE
 	};
@@ -18,10 +18,3 @@ FUnicodeBlockRange FUnicodeBlockRange::GetUnicodeBlockRange(const EUnicodeBlockR
 	TArrayView<const FUnicodeBlockRange> UnicodeBlockRanges = GetUnicodeBlockRanges();
 	return UnicodeBlockRanges[(int32)InBlockIndex];
 }
-
-FUnicodeBlockRange::FUnicodeBlockRange(EUnicodeBlockRange Index_, const TCHAR* Namespace, const TCHAR* Key, const TCHAR* Literal, int32 Lower, int32 Upper)
-:	Index(Index_)
-,	DisplayName(FText::AsLocalizable_Advanced(Namespace, Key, Literal))
-,	Range(FInt32Range::BoundsType::Inclusive(Lower), FInt32Range::BoundsType::Inclusive(Upper))
-{
-}
diff --git a/Engine/Source/Runtime/SlateCore/Private/Input/Reply.cpp b/Engine/Source/Runtime/SlateCore/Private/Input/Reply.cpp
--- a/Engine/Source/Runtime/SlateCore/Private/Input/Reply.cpp
+++ b/Engine/Source/Runtime/SlateCore/Private/Input/Reply.cpp
@@ -2,6 +2,35 @@
 
 #include "Input/Reply.h"
 
+FReply::FReply( bool bIsHandled )
+	: TReplyBase<FReply>(bIsHandled)
+	, RequestedMousePos()
+	, EventHandler(nullptr)
+	, MouseCaptor(nullptr)
+	, FocusRecipient(nullptr)
+	, MouseLockWidget(nullptr)
+	, DetectDragForWidget(nullptr)
+	, NavigationDestination(nullptr)
+	, DragDropContent(nullptr)
+	, FocusChangeReason(EFocusCause::SetDirectly)
+	, NavigationType(EUINavigation::Invalid)
+	, NavigationGenesis(ENavigationGenesis::User)
+	, NavigationSource(ENavigationSource::FocusedWidget)
+	, bReleaseMouseCapture(false)
+	, bSetUserFocus(false)
+	, bReleaseUserFocus(false)
+	, bAllUsers(false)
+	, bShouldReleaseMouseLock(false)
+	, bUseHighPrecisionMouse(false)
+	, bPreventThrottling(false)
+	, bEndDragDrop(false)
+{
+}
+
+FReply::FReply(const FReply&) = default;
+
+FReply& FReply::operator =(const FReply&) = default;
+
 FReply& FReply::SetMousePos(const FIntPoint& NewMousePos)
 {
 	this->RequestedMousePos = NewMousePos;
@@ -51,4 +80,4 @@ FString FReply::ToString()
 	}
 
 	return HandledStr;
-}
\ No newline at end of file
+}
diff --git a/Engine/Source/Runtime/SlateCore/Private/Sound/SlateSound.cpp b/Engine/Source/Runtime/SlateCore/Private/Sound/SlateSound.cpp
--- a/Engine/Source/Runtime/SlateCore/Private/Sound/SlateSound.cpp
+++ b/Engine/Source/Runtime/SlateCore/Private/Sound/SlateSound.cpp
@@ -10,6 +10,16 @@
 /* FSlateSound interface
  *****************************************************************************/
 
+FSlateSound::FSlateSound()
+	: ResourceObject(nullptr)
+{
+}
+
+
+FSlateSound::FSlateSound(const FSlateSound&) = default;
+
+FSlateSound& FSlateSound::operator=(const FSlateSound&) = default;
+
 UObject* FSlateSound::GetResourceObject( ) const
 {
 	// We might still be holding a legacy resource name, in which case we need to test that first and load it if required
diff --git a/Engine/Source/Runtime/SlateCore/Private/Styling/CoreStyle.cpp b/Engine/Source/Runtime/SlateCore/Private/Styling/CoreStyle.cpp
--- a/Engine/Source/Runtime/SlateCore/Private/Styling/CoreStyle.cpp
+++ b/Engine/Source/Runtime/SlateCore/Private/Styling/CoreStyle.cpp
@@ -21,6 +21,46 @@
 
 TSharedPtr< ISlateStyle > FCoreStyle::Instance = nullptr;
 
+namespace CoreStyleConstants
+{
+	// Note, these sizes are in Slate Units.
+	// Slate Units do NOT have to map to pixels.
+	const UE::Slate::FDeprecateVector2DResult Icon5x16(5.0f, 16.0f);
+	const UE::Slate::FDeprecateVector2DResult Icon6x8(6.0f, 8.0f);
+	const UE::Slate::FDeprecateVector2DResult Icon8x4(8.0f, 4.0f);
+	const UE::Slate::FDeprecateVector2DResult Icon16x4(16.0f, 4.0f);
+	const UE::Slate::FDeprecateVector2DResult Icon8x8(8.0f, 8.0f);
+	const UE::Slate::FDeprecateVector2DResult Icon4x4(4.0f, 4.0f);
+	const UE::Slate::FDeprecateVector2DResult Icon10x10(10.0f, 10.0f);
+	const UE::Slate::FDeprecateVector2DResult Icon12x12(12.0f, 12.0f);
+	const UE::Slate::FDeprecateVector2DResult Icon12x16(12.0f, 16.0f);
+	const UE::Slate::FDeprecateVector2DResult Icon14x14(14.0f, 14.0f);
+	const UE::Slate::FDeprecateVector2DResult Icon16x16(16.0f, 16.0f);
+	const UE::Slate::FDeprecateVector2DResult Icon18x18(18.0f, 18.0f);
+	const UE::Slate::FDeprecateVector2DResult Icon20x20(20.0f, 20.0f);
+	const UE::Slate::FDeprecateVector2DResult Icon22x22(22.0f, 22.0f);
+	const UE::Slate::FDeprecateVector2DResult Icon24x24(24.0f, 24.0f);
+	const UE::Slate::FDeprecateVector2DResult Icon25x25(25.0f, 25.0f);
+	const UE::Slate::FDeprecateVector2DResult Icon26x26(26.0f, 26.0f);
+	const UE::Slate::FDeprecateVector2DResult Icon32x32(32.0f, 32.0f);
+	const UE::Slate::FDeprecateVector2DResult Icon40x40(40.0f, 40.0f);
+	const UE::Slate::FDeprecateVector2DResult Icon64x64(64.0f, 64.0f);
+	const UE::Slate::FDeprecateVector2DResult Icon36x24(36.0f, 24.0f);
+	const UE::Slate::FDeprecateVector2DResult Icon128x128(128.0f, 128.0f);
+
+	// Common Margins
+	const FMargin DefaultMargins(8.f, 4.f);
+	// Buttons already have a built in (4., 2.) padding - adding to that a little
+	const FMargin ButtonMargins(12.f, 1.5f, 12.f, 1.5f);
+
+	const FMargin PressedButtonMargins(12.f, 2.5f, 12.f, 0.5f);
+	const FMargin ToggleButtonMargins(16.0f, 2.0f);
+	const FMargin ComboButtonMargin(8.f, 1.f, 8.f, 1.f);
+	const FMargin PressedComboButtonMargin(8.f, 2.f, 8.f, 0.f);
+
+	const float InputFocusRadius = 4.f;
+	const float InputFocusThickness = 1.0f;
+}
 
 /* FSlateCoreStyle helper class
  *****************************************************************************/
diff --git a/Engine/Source/Runtime/SlateCore/Public/Fonts/UnicodeBlockRange.h b/Engine/Source/Runtime/SlateCore/Public/Fonts/UnicodeBlockRange.h
--- a/Engine/Source/Runtime/SlateCore/Public/Fonts/UnicodeBlockRange.h
+++ b/Engine/Source/Runtime/SlateCore/Public/Fonts/UnicodeBlockRange.h
@@ -11,7 +11,7 @@
 /** Enumeration of pre-defined Unicode block ranges that can be used to access entries from FUnicodeBlockRange */
 enum class EUnicodeBlockRange : uint16
 {
-#define REGISTER_UNICODE_BLOCK_RANGE(LOWERBOUND, UPPERBOUND, SYMBOLNAME, LOCTEXT_NAMESPACE, LOCTEXT_KEY, LOCTEXT_LITERAL) SYMBOLNAME,
+#define REGISTER_UNICODE_BLOCK_RANGE(LOWERBOUND, UPPERBOUND, SYMBOLNAME, LOCTEXT_KEY, LOCTEXT_LITERAL) SYMBOLNAME,
 #include "UnicodeBlockRange.inl"
 #undef REGISTER_UNICODE_BLOCK_RANGE
 };
@@ -22,11 +22,17 @@ struct FUnicodeBlockRange
 	/** Index enum of this block */
 	EUnicodeBlockRange Index;
 
-	/** Display name of this block */
-	FText DisplayName;
+	/** Display name key for this block. (Use GetDisplayName() method) */
+	const TCHAR* DisplayNameKey;
 
-	/** Range of this block */
-	FInt32Range Range;
+	/** Display name literal for this block. (Use GetDisplayName() method) */
+	const TCHAR* DisplayNameLiteral;
+
+	/** Range lower bound of this block. (Use GetRange() method) */
+	int32 RangeLower;
+
+	/** Range upper bound of this block. (Use GetRange() method) */
+	int32 RangeUpper;
 
 	/** Returns an array containing all of the pre-defined block ranges */
 	static SLATECORE_API TArrayView<const FUnicodeBlockRange> GetUnicodeBlockRanges();
@@ -34,6 +40,15 @@ struct FUnicodeBlockRange
 	/** Returns the block corresponding to the given enum */
 	static SLATECORE_API FUnicodeBlockRange GetUnicodeBlockRange(const EUnicodeBlockRange InBlockIndex);
 
-	/** Ctor. Forced to not be inlined to improve compile times */
-	FORCENOINLINE FUnicodeBlockRange(EUnicodeBlockRange Index, const TCHAR* Namespace, const TCHAR* Key, const TCHAR* Literal, int32 Lower, int32 Upper);
+	/** Resolve display name. */
+	FText GetDisplayName() const
+	{
+		return FText::AsLocalizable_Advanced(TEXT("UnicodeBlock"), DisplayNameKey, DisplayNameLiteral);
+	}
+
+	/** Resolve range. */
+	FInt32Range GetRange() const
+	{
+		return FInt32Range(FInt32Range::BoundsType::Inclusive(RangeLower), FInt32Range::BoundsType::Inclusive(RangeUpper));
+	}
 };
diff --git a/Engine/Source/Runtime/SlateCore/Public/Fonts/UnicodeBlockRange.inl b/Engine/Source/Runtime/SlateCore/Public/Fonts/UnicodeBlockRange.inl
--- a/Engine/Source/Runtime/SlateCore/Public/Fonts/UnicodeBlockRange.inl
+++ b/Engine/Source/Runtime/SlateCore/Public/Fonts/UnicodeBlockRange.inl
@@ -1,279 +1,279 @@
 // Copyright Epic Games, Inc. All Rights Reserved.
 
 // XMacro header
-// #define REGISTER_UNICODE_BLOCK_RANGE(LOWERBOUND, UPPERBOUND, SYMBOLNAME, LOCTEXT_NAMESPACE, LOCTEXT_KEY, LOCTEXT_LITERAL) before including this header
+// #define REGISTER_UNICODE_BLOCK_RANGE(LOWERBOUND, UPPERBOUND, SYMBOLNAME, LOCTEXT_KEY, LOCTEXT_LITERAL) before including this header
 
-REGISTER_UNICODE_BLOCK_RANGE(0x0000, 0x001f, ControlCharacter,							"UnicodeBlock", "ControlCharacter", "Control Character")
-REGISTER_UNICODE_BLOCK_RANGE(0x0020, 0x007f, BasicLatin,								"UnicodeBlock", "BasicLatin", "Basic Latin")
-REGISTER_UNICODE_BLOCK_RANGE(0x0080, 0x00ff, Latin1Supplement,							"UnicodeBlock", "Latin1Supplement", "Latin - 1 Supplement")
-REGISTER_UNICODE_BLOCK_RANGE(0x0100, 0x017f, LatinExtendedA,							"UnicodeBlock", "LatinExtendedA", "Latin Extended - A")
-REGISTER_UNICODE_BLOCK_RANGE(0x0180, 0x024f, LatinExtendedB,							"UnicodeBlock", "LatinExtendedB", "Latin Extended - B")
-REGISTER_UNICODE_BLOCK_RANGE(0x0250, 0x02af, IPAExtensions,								"UnicodeBlock", "IPAExtensions", "IPA Extensions")
-REGISTER_UNICODE_BLOCK_RANGE(0x02b0, 0x02ff, SpacingModifierLetters,					"UnicodeBlock", "SpacingModifierLetters", "Spacing Modifier Letters")
-REGISTER_UNICODE_BLOCK_RANGE(0x0300, 0x036f, CombiningDiacriticalMarks,					"UnicodeBlock", "CombiningDiacriticalMarks", "Combining Diacritical Marks")
-REGISTER_UNICODE_BLOCK_RANGE(0x0370, 0x03ff, GreekAndCoptic,							"UnicodeBlock", "GreekAndCoptic", "Greek and Coptic")
-REGISTER_UNICODE_BLOCK_RANGE(0x0400, 0x04ff, Cyrillic,									"UnicodeBlock", "Cyrillic", "Cyrillic")
-REGISTER_UNICODE_BLOCK_RANGE(0x0500, 0x052f, CyrillicSupplementary,						"UnicodeBlock", "CyrillicSupplementary", "Cyrillic Supplementary")
-REGISTER_UNICODE_BLOCK_RANGE(0x0530, 0x058f, Armenian,									"UnicodeBlock", "Armenian", "Armenian")
-REGISTER_UNICODE_BLOCK_RANGE(0x0590, 0x05ff, Hebrew,									"UnicodeBlock", "Hebrew", "Hebrew")
-REGISTER_UNICODE_BLOCK_RANGE(0x0600, 0x06ff, Arabic,									"UnicodeBlock", "Arabic", "Arabic")
-REGISTER_UNICODE_BLOCK_RANGE(0x0700, 0x074f, Syriac,									"UnicodeBlock", "Syriac", "Syriac")
-REGISTER_UNICODE_BLOCK_RANGE(0x0750, 0x077f, ArabicSupplement,							"UnicodeBlock", "ArabicSupplement", "Arabic Supplement")
-REGISTER_UNICODE_BLOCK_RANGE(0x0780, 0x07bf, Thaana,									"UnicodeBlock", "Thaana", "Thaana")
-REGISTER_UNICODE_BLOCK_RANGE(0x07c0, 0x07ff, NKo,										"UnicodeBlock", "NKo", "NKo")
-REGISTER_UNICODE_BLOCK_RANGE(0x0800, 0x083f, Samaritan,									"UnicodeBlock", "Samaritan", "Samaritan")
-REGISTER_UNICODE_BLOCK_RANGE(0x0840, 0x085f, Mandaic,									"UnicodeBlock", "Mandaic", "Mandaic")
-REGISTER_UNICODE_BLOCK_RANGE(0x08a0, 0x08ff, ArabicExtendedA,							"UnicodeBlock", "ArabicExtendedA", "Arabic Extended - A")
-REGISTER_UNICODE_BLOCK_RANGE(0x0900, 0x097f, Devanagari,								"UnicodeBlock", "Devanagari", "Devanagari")
-REGISTER_UNICODE_BLOCK_RANGE(0x0980, 0x09ff, Bengali,									"UnicodeBlock", "Bengali", "Bengali")
-REGISTER_UNICODE_BLOCK_RANGE(0x0a00, 0x0a7f, Gurmukhi,									"UnicodeBlock", "Gurmukhi", "Gurmukhi")
-REGISTER_UNICODE_BLOCK_RANGE(0x0a80, 0x0aff, Gujarati,									"UnicodeBlock", "Gujarati", "Gujarati")
-REGISTER_UNICODE_BLOCK_RANGE(0x0b00, 0x0b7f, Oriya,										"UnicodeBlock", "Oriya", "Oriya")
-REGISTER_UNICODE_BLOCK_RANGE(0x0b80, 0x0bff, Tamil,										"UnicodeBlock", "Tamil", "Tamil")
-REGISTER_UNICODE_BLOCK_RANGE(0x0c00, 0x0c7f, Telugu,									"UnicodeBlock", "Telugu", "Telugu")
-REGISTER_UNICODE_BLOCK_RANGE(0x0c80, 0x0cff, Kannada,									"UnicodeBlock", "Kannada", "Kannada")
-REGISTER_UNICODE_BLOCK_RANGE(0x0d00, 0x0d7f, Malayalam,									"UnicodeBlock", "Malayalam", "Malayalam")
-REGISTER_UNICODE_BLOCK_RANGE(0x0d80, 0x0dff, Sinhala,									"UnicodeBlock", "Sinhala", "Sinhala")
-REGISTER_UNICODE_BLOCK_RANGE(0x0e00, 0x0e7f, Thai,										"UnicodeBlock", "Thai", "Thai")
-REGISTER_UNICODE_BLOCK_RANGE(0x0e80, 0x0eff, Lao,										"UnicodeBlock", "Lao", "Lao")
-REGISTER_UNICODE_BLOCK_RANGE(0x0f00, 0x0fff, Tibetan,									"UnicodeBlock", "Tibetan", "Tibetan")
-REGISTER_UNICODE_BLOCK_RANGE(0x1000, 0x109f, Myanmar,									"UnicodeBlock", "Myanmar", "Myanmar")
-REGISTER_UNICODE_BLOCK_RANGE(0x10a0, 0x10ff, Georgian,									"UnicodeBlock", "Georgian", "Georgian")
-REGISTER_UNICODE_BLOCK_RANGE(0x1100, 0x11ff, HangulJamo,								"UnicodeBlock", "HangulJamo", "Hangul Jamo")
-REGISTER_UNICODE_BLOCK_RANGE(0x1200, 0x137f, Ethiopic,									"UnicodeBlock", "Ethiopic", "Ethiopic")
-REGISTER_UNICODE_BLOCK_RANGE(0x1380, 0x139f, EthiopicSupplement,						"UnicodeBlock", "EthiopicSupplement", "Ethiopic Supplement")
-REGISTER_UNICODE_BLOCK_RANGE(0x13a0, 0x13ff, Cherokee,									"UnicodeBlock", "Cherokee", "Cherokee")
-REGISTER_UNICODE_BLOCK_RANGE(0x1400, 0x167f, UnifiedCanadianAboriginalSyllabics,		"UnicodeBlock", "UnifiedCanadianAboriginalSyllabics", "Unified Canadian Aboriginal Syllabics")
-REGISTER_UNICODE_BLOCK_RANGE(0x1680, 0x169f, Ogham,										"UnicodeBlock", "Ogham", "Ogham")
-REGISTER_UNICODE_BLOCK_RANGE(0x16a0, 0x16ff, Runic,										"UnicodeBlock", "Runic", "Runic")
-REGISTER_UNICODE_BLOCK_RANGE(0x1700, 0x171f, Tagalog,									"UnicodeBlock", "Tagalog", "Tagalog")
-REGISTER_UNICODE_BLOCK_RANGE(0x1720, 0x173f, Hanunoo,									"UnicodeBlock", "Hanunoo", "Hanunoo")
-REGISTER_UNICODE_BLOCK_RANGE(0x1740, 0x175f, Buhid,										"UnicodeBlock", "Buhid", "Buhid")
-REGISTER_UNICODE_BLOCK_RANGE(0x1760, 0x177f, Tagbanwa,									"UnicodeBlock", "Tagbanwa", "Tagbanwa")
-REGISTER_UNICODE_BLOCK_RANGE(0x1780, 0x17ff, Khmer,										"UnicodeBlock", "Khmer", "Khmer")
-REGISTER_UNICODE_BLOCK_RANGE(0x1800, 0x18af, Mongolian,									"UnicodeBlock", "Mongolian", "Mongolian")
-REGISTER_UNICODE_BLOCK_RANGE(0x18b0, 0x18ff, UnifiedCanadianAboriginalSyllabicsExtended,"UnicodeBlock", "UnifiedCanadianAboriginalSyllabicsExtended", "Unified Canadian Aboriginal Syllabics Extended")
-REGISTER_UNICODE_BLOCK_RANGE(0x1900, 0x194f, Limbu,										"UnicodeBlock", "Limbu", "Limbu")
-REGISTER_UNICODE_BLOCK_RANGE(0x1950, 0x197f, TaiLe,										"UnicodeBlock", "TaiLe", "Tai Le")
-REGISTER_UNICODE_BLOCK_RANGE(0x1980, 0x19Df, NewTaiLue,									"UnicodeBlock", "NewTaiLue", "New Tai Lue")
-REGISTER_UNICODE_BLOCK_RANGE(0x19e0, 0x19ff, KhmerSymbols,								"UnicodeBlock", "KhmerSymbols", "Khmer Symbols")
-REGISTER_UNICODE_BLOCK_RANGE(0x1a00, 0x1a1f, Buginese,									"UnicodeBlock", "Buginese", "Buginese")
-REGISTER_UNICODE_BLOCK_RANGE(0x1a20, 0x1aaf, TaiTham,									"UnicodeBlock", "TaiTham", "Tai Tham")
-REGISTER_UNICODE_BLOCK_RANGE(0x1ab0, 0x1aff, CombiningDiacriticalMarksExtended,			"UnicodeBlock", "CombiningDiacriticalMarksExtended", "Combining Diacritical Marks Extended")
-REGISTER_UNICODE_BLOCK_RANGE(0x1b00, 0x1b7f, Balinese,									"UnicodeBlock", "Balinese", "Balinese")
-REGISTER_UNICODE_BLOCK_RANGE(0x1b80, 0x1bbf, Sundanese,									"UnicodeBlock", "Sundanese", "Sundanese")
-REGISTER_UNICODE_BLOCK_RANGE(0x1bc0, 0x1bff, Batak,										"UnicodeBlock", "Batak", "Batak")
-REGISTER_UNICODE_BLOCK_RANGE(0x1c00, 0x1c4f, Lepcha,									"UnicodeBlock", "Lepcha", "Lepcha")
-REGISTER_UNICODE_BLOCK_RANGE(0x1c50, 0x1c7f, OlChiki,									"UnicodeBlock", "OlChiki", "Ol Chiki")
-REGISTER_UNICODE_BLOCK_RANGE(0x1c80, 0x1c87, CyrillicExtendedC,							"UnicodeBlock", "CyrillicExtendedC", "Cyrillic Extended C")
-REGISTER_UNICODE_BLOCK_RANGE(0x1cc0, 0x1ccf, SundaneseSupplement,						"UnicodeBlock", "SundaneseSupplement", "Sundanese Supplement")
-REGISTER_UNICODE_BLOCK_RANGE(0x1cd0, 0x1cff, VedicExtensions,							"UnicodeBlock", "VedicExtensions", "Vedic Extensions")
-REGISTER_UNICODE_BLOCK_RANGE(0x1d00, 0x1d7f, PhoneticExtensions,						"UnicodeBlock", "PhoneticExtensions", "Phonetic Extensions")
-REGISTER_UNICODE_BLOCK_RANGE(0x1d80, 0x1dbf, PhoneticExtensionsSupplement,				"UnicodeBlock", "PhoneticExtensionsSupplement", "Phonetic Extensions Supplement")
-REGISTER_UNICODE_BLOCK_RANGE(0x1dc0, 0x1dff, CombiningDiacriticalMarksSupplement,		"UnicodeBlock", "CombiningDiacriticalMarksSupplement", "Combining Diacritical Marks Supplement")
-REGISTER_UNICODE_BLOCK_RANGE(0x1e00, 0x1eff, LatinExtendedAdditional,					"UnicodeBlock", "LatinExtendedAdditional", "Latin Extended Additional")
-REGISTER_UNICODE_BLOCK_RANGE(0x1f00, 0x1fff, GreekExtended,								"UnicodeBlock", "GreekExtended", "Greek Extended")
-REGISTER_UNICODE_BLOCK_RANGE(0x2000, 0x206f, GeneralPunctuation,						"UnicodeBlock", "GeneralPunctuation", "General Punctuation")
-REGISTER_UNICODE_BLOCK_RANGE(0x2070, 0x209f, SuperscriptsAndSubscripts,					"UnicodeBlock", "SuperscriptsAndSubscripts", "Superscripts and Subscripts")
-REGISTER_UNICODE_BLOCK_RANGE(0x20a0, 0x20cf, CurrencySymbols,							"UnicodeBlock", "CurrencySymbols", "Currency Symbols")
-REGISTER_UNICODE_BLOCK_RANGE(0x20d0, 0x20ff, CombiningDiacriticalMarksForSymbols,		"UnicodeBlock", "CombiningDiacriticalMarksForSymbols", "Combining Diacritical Marks for Symbols")
-REGISTER_UNICODE_BLOCK_RANGE(0x2100, 0x214f, LetterlikeSymbols,							"UnicodeBlock", "LetterlikeSymbols", "Letterlike Symbols")
-REGISTER_UNICODE_BLOCK_RANGE(0x2150, 0x218f, NumberForms,								"UnicodeBlock", "NumberForms", "Number Forms")
-REGISTER_UNICODE_BLOCK_RANGE(0x2190, 0x21ff, Arrows,									"UnicodeBlock", "Arrows", "Arrows")
-REGISTER_UNICODE_BLOCK_RANGE(0x2200, 0x22ff, MathematicalOperators,						"UnicodeBlock", "MathematicalOperators", "Mathematical Operators")
-REGISTER_UNICODE_BLOCK_RANGE(0x2300, 0x23ff, MiscellaneousTechnical,					"UnicodeBlock", "MiscellaneousTechnical", "Miscellaneous Technical")
-REGISTER_UNICODE_BLOCK_RANGE(0x2400, 0x243f, ControlPictures,							"UnicodeBlock", "ControlPictures", "Control Pictures")
-REGISTER_UNICODE_BLOCK_RANGE(0x2440, 0x245f, OpticalCharacterRecognition,				"UnicodeBlock", "OpticalCharacterRecognition", "Optical Character Recognition")
-REGISTER_UNICODE_BLOCK_RANGE(0x2460, 0x24ff, EnclosedAlphanumerics,						"UnicodeBlock", "EnclosedAlphanumerics", "Enclosed Alphanumerics")
-REGISTER_UNICODE_BLOCK_RANGE(0x2500, 0x257f, BoxDrawing,								"UnicodeBlock", "BoxDrawing", "Box Drawing")
-REGISTER_UNICODE_BLOCK_RANGE(0x2580, 0x259f, BlockElements,								"UnicodeBlock", "BlockElements", "Block Elements")
-REGISTER_UNICODE_BLOCK_RANGE(0x25a0, 0x25ff, GeometricShapes,							"UnicodeBlock", "GeometricShapes", "Geometric Shapes")
-REGISTER_UNICODE_BLOCK_RANGE(0x2600, 0x26ff, MiscellaneousSymbols,						"UnicodeBlock", "MiscellaneousSymbols", "Miscellaneous Symbols")
-REGISTER_UNICODE_BLOCK_RANGE(0x2700, 0x27bf, Dingbats,									"UnicodeBlock", "Dingbats", "Dingbats")
-REGISTER_UNICODE_BLOCK_RANGE(0x27c0, 0x27ef, MiscellaneousMathematicalSymbolsA,			"UnicodeBlock", "MiscellaneousMathematicalSymbolsA", "Miscellaneous Mathematical Symbols - A")
-REGISTER_UNICODE_BLOCK_RANGE(0x27f0, 0x27ff, SupplementalArrowsA,						"UnicodeBlock", "SupplementalArrowsA", "Supplemental Arrows - A")
-REGISTER_UNICODE_BLOCK_RANGE(0x2800, 0x28ff, BraillePatterns,							"UnicodeBlock", "BraillePatterns", "Braille Patterns")
-REGISTER_UNICODE_BLOCK_RANGE(0x2900, 0x297f, SupplementalArrowsB,						"UnicodeBlock", "SupplementalArrowsB", "Supplemental Arrows - B")
-REGISTER_UNICODE_BLOCK_RANGE(0x2980, 0x29ff, MiscellaneousMathematicalSymbolsB,			"UnicodeBlock", "MiscellaneousMathematicalSymbolsB", "Miscellaneous Mathematical Symbols - B")
-REGISTER_UNICODE_BLOCK_RANGE(0x2a00, 0x2aff, SupplementalMathematicalOperators,			"UnicodeBlock", "SupplementalMathematicalOperators", "Supplemental Mathematical Operators")
-REGISTER_UNICODE_BLOCK_RANGE(0x2b00, 0x2bff, MiscellaneousSymbolsAndArrows,				"UnicodeBlock", "MiscellaneousSymbolsAndArrows", "Miscellaneous Symbols and Arrows")
-REGISTER_UNICODE_BLOCK_RANGE(0x2c00, 0x2c5f, Glagolitic,								"UnicodeBlock", "Glagolitic", "Glagolitic")
-REGISTER_UNICODE_BLOCK_RANGE(0x2c60, 0x2c7f, LatinExtendedC,							"UnicodeBlock", "LatinExtendedC", "Latin Extended - C")
-REGISTER_UNICODE_BLOCK_RANGE(0x2c80, 0x2cff, Coptic,									"UnicodeBlock", "Coptic", "Coptic")
-REGISTER_UNICODE_BLOCK_RANGE(0x2d00, 0x2d2f, GeorgianSupplement,						"UnicodeBlock", "GeorgianSupplement", "Georgian Supplement")
-REGISTER_UNICODE_BLOCK_RANGE(0x2d30, 0x2d7f, Tifinagh,									"UnicodeBlock", "Tifinagh", "Tifinagh")
-REGISTER_UNICODE_BLOCK_RANGE(0x2d80, 0x2ddf, EthiopicExtended,							"UnicodeBlock", "EthiopicExtended", "Ethiopic Extended")
-REGISTER_UNICODE_BLOCK_RANGE(0x2de0, 0x2dff, CyrillicExtendedA,							"UnicodeBlock", "CyrillicExtendedA", "Cyrillic Extended - A")
-REGISTER_UNICODE_BLOCK_RANGE(0x2e00, 0x2e7f, SupplementalPunctuation,					"UnicodeBlock", "SupplementalPunctuation", "Supplemental Punctuation")
-REGISTER_UNICODE_BLOCK_RANGE(0x2e80, 0x2eff, CJKRadicalsSupplement,						"UnicodeBlock", "CJKRadicalsSupplement", "CJK Radicals Supplement")
-REGISTER_UNICODE_BLOCK_RANGE(0x2f00, 0x2fdf, KangxiRadicals,							"UnicodeBlock", "KangxiRadicals", "Kangxi Radicals")
-REGISTER_UNICODE_BLOCK_RANGE(0x2ff0, 0x2fff, IdeographicDescriptionCharacters,			"UnicodeBlock", "IdeographicDescriptionCharacters", "Ideographic Description Characters")
-REGISTER_UNICODE_BLOCK_RANGE(0x3000, 0x303f, CJKSymbolsAndPunctuation,					"UnicodeBlock", "CJKSymbolsAndPunctuation", "CJK Symbols and Punctuation")
-REGISTER_UNICODE_BLOCK_RANGE(0x3040, 0x309f, Hiragana,									"UnicodeBlock", "Hiragana", "Hiragana")
-REGISTER_UNICODE_BLOCK_RANGE(0x30a0, 0x30ff, Katakana,									"UnicodeBlock", "Katakana", "Katakana")
-REGISTER_UNICODE_BLOCK_RANGE(0x3100, 0x312f, Bopomofo,									"UnicodeBlock", "Bopomofo", "Bopomofo")
-REGISTER_UNICODE_BLOCK_RANGE(0x3130, 0x318f, HangulCompatibilityJamo,					"UnicodeBlock", "HangulCompatibilityJamo", "Hangul Compatibility Jamo")
-REGISTER_UNICODE_BLOCK_RANGE(0x3190, 0x319f, Kanbun,									"UnicodeBlock", "Kanbun", "Kanbun")
-REGISTER_UNICODE_BLOCK_RANGE(0x31a0, 0x31bf, BopomofoExtended,							"UnicodeBlock", "BopomofoExtended", "Bopomofo Extended")
-REGISTER_UNICODE_BLOCK_RANGE(0x31c0, 0x31ef, CJKStrokes,								"UnicodeBlock", "CJKStrokes", "CJK Strokes")
-REGISTER_UNICODE_BLOCK_RANGE(0x31f0, 0x31ff, KatakanaPhoneticExtensions,				"UnicodeBlock", "KatakanaPhoneticExtensions", "Katakana Phonetic Extensions")
-REGISTER_UNICODE_BLOCK_RANGE(0x3200, 0x32ff, EnclosedCJKLettersAndMonths,				"UnicodeBlock", "EnclosedCJKLettersAndMonths", "Enclosed CJK Letters and Months")
-REGISTER_UNICODE_BLOCK_RANGE(0x3300, 0x33ff, CJKCompatibility,							"UnicodeBlock", "CJKCompatibility", "CJK Compatibility")
-REGISTER_UNICODE_BLOCK_RANGE(0x3400, 0x4dbf, CJKUnifiedIdeographsExtensionA,			"UnicodeBlock", "CJKUnifiedIdeographsExtensionA", "CJK Unified Ideographs Extension A")
-REGISTER_UNICODE_BLOCK_RANGE(0x4dc0, 0x4dff, YijingHexagramSymbols,						"UnicodeBlock", "YijingHexagramSymbols", "Yijing Hexagram Symbols")
-REGISTER_UNICODE_BLOCK_RANGE(0x4e00, 0x9fff, CJKUnifiedIdeographs,						"UnicodeBlock", "CJKUnifiedIdeographs", "CJK Unified Ideographs")
-REGISTER_UNICODE_BLOCK_RANGE(0xa000, 0xa48f, YiSyllables,								"UnicodeBlock", "YiSyllables", "Yi Syllables")
-REGISTER_UNICODE_BLOCK_RANGE(0xa490, 0xa4cf, YiRadicals,								"UnicodeBlock", "YiRadicals", "Yi Radicals")
-REGISTER_UNICODE_BLOCK_RANGE(0xa4d0, 0xa4ff, Lisu,										"UnicodeBlock", "Lisu", "Lisu")
-REGISTER_UNICODE_BLOCK_RANGE(0xa500, 0xa63f, Vai,										"UnicodeBlock", "Vai", "Vai")
-REGISTER_UNICODE_BLOCK_RANGE(0xa640, 0xa69f, CyrillicExtendedB,							"UnicodeBlock", "CyrillicExtendedB", "Cyrillic Extended - B")
-REGISTER_UNICODE_BLOCK_RANGE(0xa6a0, 0xa6ff, Bamum,										"UnicodeBlock", "Bamum", "Bamum")
-REGISTER_UNICODE_BLOCK_RANGE(0xa700, 0xa71f, ModifierToneLetters,						"UnicodeBlock", "ModifierToneLetters", "Modifier Tone Letters")
-REGISTER_UNICODE_BLOCK_RANGE(0xa720, 0xa7ff, LatinExtendedD,							"UnicodeBlock", "LatinExtendedD", "Latin Extended - D")
-REGISTER_UNICODE_BLOCK_RANGE(0xa800, 0xa82f, SylotiNagri,								"UnicodeBlock", "SylotiNagri", "Syloti Nagri")
-REGISTER_UNICODE_BLOCK_RANGE(0xa830, 0xa83f, CommonIndicNumberForms,					"UnicodeBlock", "CommonIndicNumberForms", "Common Indic Number Forms")
-REGISTER_UNICODE_BLOCK_RANGE(0xa840, 0xa87f, PhagsPa,									"UnicodeBlock", "PhagsPa", "Phags - pa")
-REGISTER_UNICODE_BLOCK_RANGE(0xa880, 0xa8df, Saurashtra,								"UnicodeBlock", "Saurashtra", "Saurashtra")
-REGISTER_UNICODE_BLOCK_RANGE(0xa8e0, 0xa8ff, DevanagariExtended,						"UnicodeBlock", "DevanagariExtended", "Devanagari Extended")
-REGISTER_UNICODE_BLOCK_RANGE(0xa900, 0xa92f, KayahLi,									"UnicodeBlock", "KayahLi", "Kayah Li")
-REGISTER_UNICODE_BLOCK_RANGE(0xa930, 0xa95f, Rejang,									"UnicodeBlock", "Rejang", "Rejang")
-REGISTER_UNICODE_BLOCK_RANGE(0xa960, 0xa97f, HangulJamoExtendedA,						"UnicodeBlock", "HangulJamoExtendedA", "Hangul Jamo Extended - A")
-REGISTER_UNICODE_BLOCK_RANGE(0xa980, 0xa9df, Javanese,									"UnicodeBlock", "Javanese", "Javanese")
-REGISTER_UNICODE_BLOCK_RANGE(0xa9e0, 0xa9ff, MyanmarExtendedB,							"UnicodeBlock", "MyanmarExtendedB", "Myanmar Extended - B")
-REGISTER_UNICODE_BLOCK_RANGE(0xaa00, 0xaa5f, Cham,										"UnicodeBlock", "Cham", "Cham")
-REGISTER_UNICODE_BLOCK_RANGE(0xaa60, 0xaa7f, MyanmarExtendedA,							"UnicodeBlock", "MyanmarExtendedA", "Myanmar Extended - A")
-REGISTER_UNICODE_BLOCK_RANGE(0xaa80, 0xaadf, TaiViet,									"UnicodeBlock", "TaiViet", "Tai Viet")
-REGISTER_UNICODE_BLOCK_RANGE(0xaae0, 0xaaff, MeeteiMayekExtensions,						"UnicodeBlock", "MeeteiMayekExtensions", "Meetei Mayek Extensions")
-REGISTER_UNICODE_BLOCK_RANGE(0xab00, 0xab2f, EthiopicExtendedA,							"UnicodeBlock", "EthiopicExtendedA", "Ethiopic Extended - A")
-REGISTER_UNICODE_BLOCK_RANGE(0xab30, 0xab6f, LatinExtendedE,							"UnicodeBlock", "LatinExtendedE", "Latin Extended - E")
-REGISTER_UNICODE_BLOCK_RANGE(0xab70, 0xabbf, CherokeeSupplement,						"UnicodeBlock", "CherokeeSupplement", "Cherokee Supplement")
-REGISTER_UNICODE_BLOCK_RANGE(0xabc0, 0xabff, MeeteiMayek,								"UnicodeBlock", "MeeteiMayek", "Meetei Mayek")
-REGISTER_UNICODE_BLOCK_RANGE(0xac00, 0xd7af, HangulSyllables,							"UnicodeBlock", "HangulSyllables", "Hangul Syllables")
-REGISTER_UNICODE_BLOCK_RANGE(0xd7b0, 0xd7ff, HangulJamoExtendedB,						"UnicodeBlock", "HangulJamoExtendedB", "Hangul Jamo Extended - B")
-REGISTER_UNICODE_BLOCK_RANGE(0xd800, 0xdb7f, HighSurrogates,							"UnicodeBlock", "HighSurrogates", "High Surrogates")
-REGISTER_UNICODE_BLOCK_RANGE(0xdb80, 0xdbff, HighPrivateUseSurrogates,					"UnicodeBlock", "HighPrivateUseSurrogates", "High Private Use Surrogates")
-REGISTER_UNICODE_BLOCK_RANGE(0xdc00, 0xdfff, LowSurrogates,								"UnicodeBlock", "LowSurrogates", "Low Surrogates")
-REGISTER_UNICODE_BLOCK_RANGE(0xe000, 0xf8ff, PrivateUseArea,							"UnicodeBlock", "PrivateUseArea", "Private Use Area")
-REGISTER_UNICODE_BLOCK_RANGE(0xf900, 0xfaff, CJKCompatibilityIdeographs,				"UnicodeBlock", "CJKCompatibilityIdeographs", "CJK Compatibility Ideographs")
-REGISTER_UNICODE_BLOCK_RANGE(0xfb00, 0xfb4f, AlphabeticPresentationForms,				"UnicodeBlock", "AlphabeticPresentationForms", "Alphabetic Presentation Forms")
-REGISTER_UNICODE_BLOCK_RANGE(0xfb50, 0xfdff, ArabicPresentationFormsA,					"UnicodeBlock", "ArabicPresentationFormsA", "Arabic Presentation Forms - A")
-REGISTER_UNICODE_BLOCK_RANGE(0xfe00, 0xfe0f, VariationSelectors,						"UnicodeBlock", "VariationSelectors", "Variation Selectors")
-REGISTER_UNICODE_BLOCK_RANGE(0xfe10, 0xfe1f, VerticalForms,								"UnicodeBlock", "VerticalForms", "Vertical Forms")
-REGISTER_UNICODE_BLOCK_RANGE(0xfe20, 0xfe2f, CombiningHalfMarks,						"UnicodeBlock", "CombiningHalfMarks", "Combining Half Marks")
-REGISTER_UNICODE_BLOCK_RANGE(0xfe30, 0xfe4f, CJKCompatibilityForms,						"UnicodeBlock", "CJKCompatibilityForms", "CJK Compatibility Forms")
-REGISTER_UNICODE_BLOCK_RANGE(0xfe50, 0xfe6f, SmallFormVariants,							"UnicodeBlock", "SmallFormVariants", "Small Form Variants")
-REGISTER_UNICODE_BLOCK_RANGE(0xfe70, 0xfeff, ArabicPresentationFormsB,					"UnicodeBlock", "ArabicPresentationFormsB", "Arabic Presentation Forms - B")
-REGISTER_UNICODE_BLOCK_RANGE(0xff00, 0xffef, HalfwidthAndFullwidthForms,				"UnicodeBlock", "HalfwidthAndFullwidthForms", "Halfwidth and Fullwidth Forms")
-REGISTER_UNICODE_BLOCK_RANGE(0xfff0, 0xffff, Specials,									"UnicodeBlock", "Specials", "Specials")
+REGISTER_UNICODE_BLOCK_RANGE(0x0000, 0x001f, ControlCharacter,							"ControlCharacter", "Control Character")
+REGISTER_UNICODE_BLOCK_RANGE(0x0020, 0x007f, BasicLatin,								"BasicLatin", "Basic Latin")
+REGISTER_UNICODE_BLOCK_RANGE(0x0080, 0x00ff, Latin1Supplement,							"Latin1Supplement", "Latin - 1 Supplement")
+REGISTER_UNICODE_BLOCK_RANGE(0x0100, 0x017f, LatinExtendedA,							"LatinExtendedA", "Latin Extended - A")
+REGISTER_UNICODE_BLOCK_RANGE(0x0180, 0x024f, LatinExtendedB,							"LatinExtendedB", "Latin Extended - B")
+REGISTER_UNICODE_BLOCK_RANGE(0x0250, 0x02af, IPAExtensions,								"IPAExtensions", "IPA Extensions")
+REGISTER_UNICODE_BLOCK_RANGE(0x02b0, 0x02ff, SpacingModifierLetters,					"SpacingModifierLetters", "Spacing Modifier Letters")
+REGISTER_UNICODE_BLOCK_RANGE(0x0300, 0x036f, CombiningDiacriticalMarks,					"CombiningDiacriticalMarks", "Combining Diacritical Marks")
+REGISTER_UNICODE_BLOCK_RANGE(0x0370, 0x03ff, GreekAndCoptic,							"GreekAndCoptic", "Greek and Coptic")
+REGISTER_UNICODE_BLOCK_RANGE(0x0400, 0x04ff, Cyrillic,									"Cyrillic", "Cyrillic")
+REGISTER_UNICODE_BLOCK_RANGE(0x0500, 0x052f, CyrillicSupplementary,						"CyrillicSupplementary", "Cyrillic Supplementary")
+REGISTER_UNICODE_BLOCK_RANGE(0x0530, 0x058f, Armenian,									"Armenian", "Armenian")
+REGISTER_UNICODE_BLOCK_RANGE(0x0590, 0x05ff, Hebrew,									"Hebrew", "Hebrew")
+REGISTER_UNICODE_BLOCK_RANGE(0x0600, 0x06ff, Arabic,									"Arabic", "Arabic")
+REGISTER_UNICODE_BLOCK_RANGE(0x0700, 0x074f, Syriac,									"Syriac", "Syriac")
+REGISTER_UNICODE_BLOCK_RANGE(0x0750, 0x077f, ArabicSupplement,							"ArabicSupplement", "Arabic Supplement")
+REGISTER_UNICODE_BLOCK_RANGE(0x0780, 0x07bf, Thaana,									"Thaana", "Thaana")
+REGISTER_UNICODE_BLOCK_RANGE(0x07c0, 0x07ff, NKo,										"NKo", "NKo")
+REGISTER_UNICODE_BLOCK_RANGE(0x0800, 0x083f, Samaritan,									"Samaritan", "Samaritan")
+REGISTER_UNICODE_BLOCK_RANGE(0x0840, 0x085f, Mandaic,									"Mandaic", "Mandaic")
+REGISTER_UNICODE_BLOCK_RANGE(0x08a0, 0x08ff, ArabicExtendedA,							"ArabicExtendedA", "Arabic Extended - A")
+REGISTER_UNICODE_BLOCK_RANGE(0x0900, 0x097f, Devanagari,								"Devanagari", "Devanagari")
+REGISTER_UNICODE_BLOCK_RANGE(0x0980, 0x09ff, Bengali,									"Bengali", "Bengali")
+REGISTER_UNICODE_BLOCK_RANGE(0x0a00, 0x0a7f, Gurmukhi,									"Gurmukhi", "Gurmukhi")
+REGISTER_UNICODE_BLOCK_RANGE(0x0a80, 0x0aff, Gujarati,									"Gujarati", "Gujarati")
+REGISTER_UNICODE_BLOCK_RANGE(0x0b00, 0x0b7f, Oriya,										"Oriya", "Oriya")
+REGISTER_UNICODE_BLOCK_RANGE(0x0b80, 0x0bff, Tamil,										"Tamil", "Tamil")
+REGISTER_UNICODE_BLOCK_RANGE(0x0c00, 0x0c7f, Telugu,									"Telugu", "Telugu")
+REGISTER_UNICODE_BLOCK_RANGE(0x0c80, 0x0cff, Kannada,									"Kannada", "Kannada")
+REGISTER_UNICODE_BLOCK_RANGE(0x0d00, 0x0d7f, Malayalam,									"Malayalam", "Malayalam")
+REGISTER_UNICODE_BLOCK_RANGE(0x0d80, 0x0dff, Sinhala,									"Sinhala", "Sinhala")
+REGISTER_UNICODE_BLOCK_RANGE(0x0e00, 0x0e7f, Thai,										"Thai", "Thai")
+REGISTER_UNICODE_BLOCK_RANGE(0x0e80, 0x0eff, Lao,										"Lao", "Lao")
+REGISTER_UNICODE_BLOCK_RANGE(0x0f00, 0x0fff, Tibetan,									"Tibetan", "Tibetan")
+REGISTER_UNICODE_BLOCK_RANGE(0x1000, 0x109f, Myanmar,									"Myanmar", "Myanmar")
+REGISTER_UNICODE_BLOCK_RANGE(0x10a0, 0x10ff, Georgian,									"Georgian", "Georgian")
+REGISTER_UNICODE_BLOCK_RANGE(0x1100, 0x11ff, HangulJamo,								"HangulJamo", "Hangul Jamo")
+REGISTER_UNICODE_BLOCK_RANGE(0x1200, 0x137f, Ethiopic,									"Ethiopic", "Ethiopic")
+REGISTER_UNICODE_BLOCK_RANGE(0x1380, 0x139f, EthiopicSupplement,						"EthiopicSupplement", "Ethiopic Supplement")
+REGISTER_UNICODE_BLOCK_RANGE(0x13a0, 0x13ff, Cherokee,									"Cherokee", "Cherokee")
+REGISTER_UNICODE_BLOCK_RANGE(0x1400, 0x167f, UnifiedCanadianAboriginalSyllabics,		"UnifiedCanadianAboriginalSyllabics", "Unified Canadian Aboriginal Syllabics")
+REGISTER_UNICODE_BLOCK_RANGE(0x1680, 0x169f, Ogham,										"Ogham", "Ogham")
+REGISTER_UNICODE_BLOCK_RANGE(0x16a0, 0x16ff, Runic,										"Runic", "Runic")
+REGISTER_UNICODE_BLOCK_RANGE(0x1700, 0x171f, Tagalog,									"Tagalog", "Tagalog")
+REGISTER_UNICODE_BLOCK_RANGE(0x1720, 0x173f, Hanunoo,									"Hanunoo", "Hanunoo")
+REGISTER_UNICODE_BLOCK_RANGE(0x1740, 0x175f, Buhid,										"Buhid", "Buhid")
+REGISTER_UNICODE_BLOCK_RANGE(0x1760, 0x177f, Tagbanwa,									"Tagbanwa", "Tagbanwa")
+REGISTER_UNICODE_BLOCK_RANGE(0x1780, 0x17ff, Khmer,										"Khmer", "Khmer")
+REGISTER_UNICODE_BLOCK_RANGE(0x1800, 0x18af, Mongolian,									"Mongolian", "Mongolian")
+REGISTER_UNICODE_BLOCK_RANGE(0x18b0, 0x18ff, UnifiedCanadianAboriginalSyllabicsExtended,"UnifiedCanadianAboriginalSyllabicsExtended", "Unified Canadian Aboriginal Syllabics Extended")
+REGISTER_UNICODE_BLOCK_RANGE(0x1900, 0x194f, Limbu,										"Limbu", "Limbu")
+REGISTER_UNICODE_BLOCK_RANGE(0x1950, 0x197f, TaiLe,										"TaiLe", "Tai Le")
+REGISTER_UNICODE_BLOCK_RANGE(0x1980, 0x19Df, NewTaiLue,									"NewTaiLue", "New Tai Lue")
+REGISTER_UNICODE_BLOCK_RANGE(0x19e0, 0x19ff, KhmerSymbols,								"KhmerSymbols", "Khmer Symbols")
+REGISTER_UNICODE_BLOCK_RANGE(0x1a00, 0x1a1f, Buginese,									"Buginese", "Buginese")
+REGISTER_UNICODE_BLOCK_RANGE(0x1a20, 0x1aaf, TaiTham,									"TaiTham", "Tai Tham")
+REGISTER_UNICODE_BLOCK_RANGE(0x1ab0, 0x1aff, CombiningDiacriticalMarksExtended,			"CombiningDiacriticalMarksExtended", "Combining Diacritical Marks Extended")
+REGISTER_UNICODE_BLOCK_RANGE(0x1b00, 0x1b7f, Balinese,									"Balinese", "Balinese")
+REGISTER_UNICODE_BLOCK_RANGE(0x1b80, 0x1bbf, Sundanese,									"Sundanese", "Sundanese")
+REGISTER_UNICODE_BLOCK_RANGE(0x1bc0, 0x1bff, Batak,										"Batak", "Batak")
+REGISTER_UNICODE_BLOCK_RANGE(0x1c00, 0x1c4f, Lepcha,									"Lepcha", "Lepcha")
+REGISTER_UNICODE_BLOCK_RANGE(0x1c50, 0x1c7f, OlChiki,									"OlChiki", "Ol Chiki")
+REGISTER_UNICODE_BLOCK_RANGE(0x1c80, 0x1c87, CyrillicExtendedC,							"CyrillicExtendedC", "Cyrillic Extended C")
+REGISTER_UNICODE_BLOCK_RANGE(0x1cc0, 0x1ccf, SundaneseSupplement,						"SundaneseSupplement", "Sundanese Supplement")
+REGISTER_UNICODE_BLOCK_RANGE(0x1cd0, 0x1cff, VedicExtensions,							"VedicExtensions", "Vedic Extensions")
+REGISTER_UNICODE_BLOCK_RANGE(0x1d00, 0x1d7f, PhoneticExtensions,						"PhoneticExtensions", "Phonetic Extensions")
+REGISTER_UNICODE_BLOCK_RANGE(0x1d80, 0x1dbf, PhoneticExtensionsSupplement,				"PhoneticExtensionsSupplement", "Phonetic Extensions Supplement")
+REGISTER_UNICODE_BLOCK_RANGE(0x1dc0, 0x1dff, CombiningDiacriticalMarksSupplement,		"CombiningDiacriticalMarksSupplement", "Combining Diacritical Marks Supplement")
+REGISTER_UNICODE_BLOCK_RANGE(0x1e00, 0x1eff, LatinExtendedAdditional,					"LatinExtendedAdditional", "Latin Extended Additional")
+REGISTER_UNICODE_BLOCK_RANGE(0x1f00, 0x1fff, GreekExtended,								"GreekExtended", "Greek Extended")
+REGISTER_UNICODE_BLOCK_RANGE(0x2000, 0x206f, GeneralPunctuation,						"GeneralPunctuation", "General Punctuation")
+REGISTER_UNICODE_BLOCK_RANGE(0x2070, 0x209f, SuperscriptsAndSubscripts,					"SuperscriptsAndSubscripts", "Superscripts and Subscripts")
+REGISTER_UNICODE_BLOCK_RANGE(0x20a0, 0x20cf, CurrencySymbols,							"CurrencySymbols", "Currency Symbols")
+REGISTER_UNICODE_BLOCK_RANGE(0x20d0, 0x20ff, CombiningDiacriticalMarksForSymbols,		"CombiningDiacriticalMarksForSymbols", "Combining Diacritical Marks for Symbols")
+REGISTER_UNICODE_BLOCK_RANGE(0x2100, 0x214f, LetterlikeSymbols,							"LetterlikeSymbols", "Letterlike Symbols")
+REGISTER_UNICODE_BLOCK_RANGE(0x2150, 0x218f, NumberForms,								"NumberForms", "Number Forms")
+REGISTER_UNICODE_BLOCK_RANGE(0x2190, 0x21ff, Arrows,									"Arrows", "Arrows")
+REGISTER_UNICODE_BLOCK_RANGE(0x2200, 0x22ff, MathematicalOperators,						"MathematicalOperators", "Mathematical Operators")
+REGISTER_UNICODE_BLOCK_RANGE(0x2300, 0x23ff, MiscellaneousTechnical,					"MiscellaneousTechnical", "Miscellaneous Technical")
+REGISTER_UNICODE_BLOCK_RANGE(0x2400, 0x243f, ControlPictures,							"ControlPictures", "Control Pictures")
+REGISTER_UNICODE_BLOCK_RANGE(0x2440, 0x245f, OpticalCharacterRecognition,				"OpticalCharacterRecognition", "Optical Character Recognition")
+REGISTER_UNICODE_BLOCK_RANGE(0x2460, 0x24ff, EnclosedAlphanumerics,						"EnclosedAlphanumerics", "Enclosed Alphanumerics")
+REGISTER_UNICODE_BLOCK_RANGE(0x2500, 0x257f, BoxDrawing,								"BoxDrawing", "Box Drawing")
+REGISTER_UNICODE_BLOCK_RANGE(0x2580, 0x259f, BlockElements,								"BlockElements", "Block Elements")
+REGISTER_UNICODE_BLOCK_RANGE(0x25a0, 0x25ff, GeometricShapes,							"GeometricShapes", "Geometric Shapes")
+REGISTER_UNICODE_BLOCK_RANGE(0x2600, 0x26ff, MiscellaneousSymbols,						"MiscellaneousSymbols", "Miscellaneous Symbols")
+REGISTER_UNICODE_BLOCK_RANGE(0x2700, 0x27bf, Dingbats,									"Dingbats", "Dingbats")
+REGISTER_UNICODE_BLOCK_RANGE(0x27c0, 0x27ef, MiscellaneousMathematicalSymbolsA,			"MiscellaneousMathematicalSymbolsA", "Miscellaneous Mathematical Symbols - A")
+REGISTER_UNICODE_BLOCK_RANGE(0x27f0, 0x27ff, SupplementalArrowsA,						"SupplementalArrowsA", "Supplemental Arrows - A")
+REGISTER_UNICODE_BLOCK_RANGE(0x2800, 0x28ff, BraillePatterns,							"BraillePatterns", "Braille Patterns")
+REGISTER_UNICODE_BLOCK_RANGE(0x2900, 0x297f, SupplementalArrowsB,						"SupplementalArrowsB", "Supplemental Arrows - B")
+REGISTER_UNICODE_BLOCK_RANGE(0x2980, 0x29ff, MiscellaneousMathematicalSymbolsB,			"MiscellaneousMathematicalSymbolsB", "Miscellaneous Mathematical Symbols - B")
+REGISTER_UNICODE_BLOCK_RANGE(0x2a00, 0x2aff, SupplementalMathematicalOperators,			"SupplementalMathematicalOperators", "Supplemental Mathematical Operators")
+REGISTER_UNICODE_BLOCK_RANGE(0x2b00, 0x2bff, MiscellaneousSymbolsAndArrows,				"MiscellaneousSymbolsAndArrows", "Miscellaneous Symbols and Arrows")
+REGISTER_UNICODE_BLOCK_RANGE(0x2c00, 0x2c5f, Glagolitic,								"Glagolitic", "Glagolitic")
+REGISTER_UNICODE_BLOCK_RANGE(0x2c60, 0x2c7f, LatinExtendedC,							"LatinExtendedC", "Latin Extended - C")
+REGISTER_UNICODE_BLOCK_RANGE(0x2c80, 0x2cff, Coptic,									"Coptic", "Coptic")
+REGISTER_UNICODE_BLOCK_RANGE(0x2d00, 0x2d2f, GeorgianSupplement,						"GeorgianSupplement", "Georgian Supplement")
+REGISTER_UNICODE_BLOCK_RANGE(0x2d30, 0x2d7f, Tifinagh,									"Tifinagh", "Tifinagh")
+REGISTER_UNICODE_BLOCK_RANGE(0x2d80, 0x2ddf, EthiopicExtended,							"EthiopicExtended", "Ethiopic Extended")
+REGISTER_UNICODE_BLOCK_RANGE(0x2de0, 0x2dff, CyrillicExtendedA,							"CyrillicExtendedA", "Cyrillic Extended - A")
+REGISTER_UNICODE_BLOCK_RANGE(0x2e00, 0x2e7f, SupplementalPunctuation,					"SupplementalPunctuation", "Supplemental Punctuation")
+REGISTER_UNICODE_BLOCK_RANGE(0x2e80, 0x2eff, CJKRadicalsSupplement,						"CJKRadicalsSupplement", "CJK Radicals Supplement")
+REGISTER_UNICODE_BLOCK_RANGE(0x2f00, 0x2fdf, KangxiRadicals,							"KangxiRadicals", "Kangxi Radicals")
+REGISTER_UNICODE_BLOCK_RANGE(0x2ff0, 0x2fff, IdeographicDescriptionCharacters,			"IdeographicDescriptionCharacters", "Ideographic Description Characters")
+REGISTER_UNICODE_BLOCK_RANGE(0x3000, 0x303f, CJKSymbolsAndPunctuation,					"CJKSymbolsAndPunctuation", "CJK Symbols and Punctuation")
+REGISTER_UNICODE_BLOCK_RANGE(0x3040, 0x309f, Hiragana,									"Hiragana", "Hiragana")
+REGISTER_UNICODE_BLOCK_RANGE(0x30a0, 0x30ff, Katakana,									"Katakana", "Katakana")
+REGISTER_UNICODE_BLOCK_RANGE(0x3100, 0x312f, Bopomofo,									"Bopomofo", "Bopomofo")
+REGISTER_UNICODE_BLOCK_RANGE(0x3130, 0x318f, HangulCompatibilityJamo,					"HangulCompatibilityJamo", "Hangul Compatibility Jamo")
+REGISTER_UNICODE_BLOCK_RANGE(0x3190, 0x319f, Kanbun,									"Kanbun", "Kanbun")
+REGISTER_UNICODE_BLOCK_RANGE(0x31a0, 0x31bf, BopomofoExtended,							"BopomofoExtended", "Bopomofo Extended")
+REGISTER_UNICODE_BLOCK_RANGE(0x31c0, 0x31ef, CJKStrokes,								"CJKStrokes", "CJK Strokes")
+REGISTER_UNICODE_BLOCK_RANGE(0x31f0, 0x31ff, KatakanaPhoneticExtensions,				"KatakanaPhoneticExtensions", "Katakana Phonetic Extensions")
+REGISTER_UNICODE_BLOCK_RANGE(0x3200, 0x32ff, EnclosedCJKLettersAndMonths,				"EnclosedCJKLettersAndMonths", "Enclosed CJK Letters and Months")
+REGISTER_UNICODE_BLOCK_RANGE(0x3300, 0x33ff, CJKCompatibility,							"CJKCompatibility", "CJK Compatibility")
+REGISTER_UNICODE_BLOCK_RANGE(0x3400, 0x4dbf, CJKUnifiedIdeographsExtensionA,			"CJKUnifiedIdeographsExtensionA", "CJK Unified Ideographs Extension A")
+REGISTER_UNICODE_BLOCK_RANGE(0x4dc0, 0x4dff, YijingHexagramSymbols,						"YijingHexagramSymbols", "Yijing Hexagram Symbols")
+REGISTER_UNICODE_BLOCK_RANGE(0x4e00, 0x9fff, CJKUnifiedIdeographs,						"CJKUnifiedIdeographs", "CJK Unified Ideographs")
+REGISTER_UNICODE_BLOCK_RANGE(0xa000, 0xa48f, YiSyllables,								"YiSyllables", "Yi Syllables")
+REGISTER_UNICODE_BLOCK_RANGE(0xa490, 0xa4cf, YiRadicals,								"YiRadicals", "Yi Radicals")
+REGISTER_UNICODE_BLOCK_RANGE(0xa4d0, 0xa4ff, Lisu,										"Lisu", "Lisu")
+REGISTER_UNICODE_BLOCK_RANGE(0xa500, 0xa63f, Vai,										"Vai", "Vai")
+REGISTER_UNICODE_BLOCK_RANGE(0xa640, 0xa69f, CyrillicExtendedB,							"CyrillicExtendedB", "Cyrillic Extended - B")
+REGISTER_UNICODE_BLOCK_RANGE(0xa6a0, 0xa6ff, Bamum,										"Bamum", "Bamum")
+REGISTER_UNICODE_BLOCK_RANGE(0xa700, 0xa71f, ModifierToneLetters,						"ModifierToneLetters", "Modifier Tone Letters")
+REGISTER_UNICODE_BLOCK_RANGE(0xa720, 0xa7ff, LatinExtendedD,							"LatinExtendedD", "Latin Extended - D")
+REGISTER_UNICODE_BLOCK_RANGE(0xa800, 0xa82f, SylotiNagri,								"SylotiNagri", "Syloti Nagri")
+REGISTER_UNICODE_BLOCK_RANGE(0xa830, 0xa83f, CommonIndicNumberForms,					"CommonIndicNumberForms", "Common Indic Number Forms")
+REGISTER_UNICODE_BLOCK_RANGE(0xa840, 0xa87f, PhagsPa,									"PhagsPa", "Phags - pa")
+REGISTER_UNICODE_BLOCK_RANGE(0xa880, 0xa8df, Saurashtra,								"Saurashtra", "Saurashtra")
+REGISTER_UNICODE_BLOCK_RANGE(0xa8e0, 0xa8ff, DevanagariExtended,						"DevanagariExtended", "Devanagari Extended")
+REGISTER_UNICODE_BLOCK_RANGE(0xa900, 0xa92f, KayahLi,									"KayahLi", "Kayah Li")
+REGISTER_UNICODE_BLOCK_RANGE(0xa930, 0xa95f, Rejang,									"Rejang", "Rejang")
+REGISTER_UNICODE_BLOCK_RANGE(0xa960, 0xa97f, HangulJamoExtendedA,						"HangulJamoExtendedA", "Hangul Jamo Extended - A")
+REGISTER_UNICODE_BLOCK_RANGE(0xa980, 0xa9df, Javanese,									"Javanese", "Javanese")
+REGISTER_UNICODE_BLOCK_RANGE(0xa9e0, 0xa9ff, MyanmarExtendedB,							"MyanmarExtendedB", "Myanmar Extended - B")
+REGISTER_UNICODE_BLOCK_RANGE(0xaa00, 0xaa5f, Cham,										"Cham", "Cham")
+REGISTER_UNICODE_BLOCK_RANGE(0xaa60, 0xaa7f, MyanmarExtendedA,							"MyanmarExtendedA", "Myanmar Extended - A")
+REGISTER_UNICODE_BLOCK_RANGE(0xaa80, 0xaadf, TaiViet,									"TaiViet", "Tai Viet")
+REGISTER_UNICODE_BLOCK_RANGE(0xaae0, 0xaaff, MeeteiMayekExtensions,						"MeeteiMayekExtensions", "Meetei Mayek Extensions")
+REGISTER_UNICODE_BLOCK_RANGE(0xab00, 0xab2f, EthiopicExtendedA,							"EthiopicExtendedA", "Ethiopic Extended - A")
+REGISTER_UNICODE_BLOCK_RANGE(0xab30, 0xab6f, LatinExtendedE,							"LatinExtendedE", "Latin Extended - E")
+REGISTER_UNICODE_BLOCK_RANGE(0xab70, 0xabbf, CherokeeSupplement,						"CherokeeSupplement", "Cherokee Supplement")
+REGISTER_UNICODE_BLOCK_RANGE(0xabc0, 0xabff, MeeteiMayek,								"MeeteiMayek", "Meetei Mayek")
+REGISTER_UNICODE_BLOCK_RANGE(0xac00, 0xd7af, HangulSyllables,							"HangulSyllables", "Hangul Syllables")
+REGISTER_UNICODE_BLOCK_RANGE(0xd7b0, 0xd7ff, HangulJamoExtendedB,						"HangulJamoExtendedB", "Hangul Jamo Extended - B")
+REGISTER_UNICODE_BLOCK_RANGE(0xd800, 0xdb7f, HighSurrogates,							"HighSurrogates", "High Surrogates")
+REGISTER_UNICODE_BLOCK_RANGE(0xdb80, 0xdbff, HighPrivateUseSurrogates,					"HighPrivateUseSurrogates", "High Private Use Surrogates")
+REGISTER_UNICODE_BLOCK_RANGE(0xdc00, 0xdfff, LowSurrogates,								"LowSurrogates", "Low Surrogates")
+REGISTER_UNICODE_BLOCK_RANGE(0xe000, 0xf8ff, PrivateUseArea,							"PrivateUseArea", "Private Use Area")
+REGISTER_UNICODE_BLOCK_RANGE(0xf900, 0xfaff, CJKCompatibilityIdeographs,				"CJKCompatibilityIdeographs", "CJK Compatibility Ideographs")
+REGISTER_UNICODE_BLOCK_RANGE(0xfb00, 0xfb4f, AlphabeticPresentationForms,				"AlphabeticPresentationForms", "Alphabetic Presentation Forms")
+REGISTER_UNICODE_BLOCK_RANGE(0xfb50, 0xfdff, ArabicPresentationFormsA,					"ArabicPresentationFormsA", "Arabic Presentation Forms - A")
+REGISTER_UNICODE_BLOCK_RANGE(0xfe00, 0xfe0f, VariationSelectors,						"VariationSelectors", "Variation Selectors")
+REGISTER_UNICODE_BLOCK_RANGE(0xfe10, 0xfe1f, VerticalForms,								"VerticalForms", "Vertical Forms")
+REGISTER_UNICODE_BLOCK_RANGE(0xfe20, 0xfe2f, CombiningHalfMarks,						"CombiningHalfMarks", "Combining Half Marks")
+REGISTER_UNICODE_BLOCK_RANGE(0xfe30, 0xfe4f, CJKCompatibilityForms,						"CJKCompatibilityForms", "CJK Compatibility Forms")
+REGISTER_UNICODE_BLOCK_RANGE(0xfe50, 0xfe6f, SmallFormVariants,							"SmallFormVariants", "Small Form Variants")
+REGISTER_UNICODE_BLOCK_RANGE(0xfe70, 0xfeff, ArabicPresentationFormsB,					"ArabicPresentationFormsB", "Arabic Presentation Forms - B")
+REGISTER_UNICODE_BLOCK_RANGE(0xff00, 0xffef, HalfwidthAndFullwidthForms,				"HalfwidthAndFullwidthForms", "Halfwidth and Fullwidth Forms")
+REGISTER_UNICODE_BLOCK_RANGE(0xfff0, 0xffff, Specials,									"Specials", "Specials")
 
 // The following are outside the Basic Multilingual Plane, so may not be supported by UE on all platforms
-REGISTER_UNICODE_BLOCK_RANGE(0x10000, 0x1007f, LinearBSyllabary,						"UnicodeBlock", "LinearBSyllabary", "Linear B Syllabary")
-REGISTER_UNICODE_BLOCK_RANGE(0x10080, 0x100ff, LinearBIdeograms,						"UnicodeBlock", "LinearBIdeograms", "Linear B Ideograms")
-REGISTER_UNICODE_BLOCK_RANGE(0x10100, 0x1013f, AegeanNumbers,							"UnicodeBlock", "AegeanNumbers", "Aegean Numbers")
-REGISTER_UNICODE_BLOCK_RANGE(0x10140, 0x1018f, AncientGreekNumbers,						"UnicodeBlock", "AncientGreekNumbers", "Ancient Greek Numbers")
-REGISTER_UNICODE_BLOCK_RANGE(0x10190, 0x101cf, AncientSymbols,							"UnicodeBlock", "AncientSymbols", "Ancient Symbols")
-REGISTER_UNICODE_BLOCK_RANGE(0x101d0, 0x101ff, PhaistosDisc,							"UnicodeBlock", "PhaistosDisc", "Phaistos Disc")
-REGISTER_UNICODE_BLOCK_RANGE(0x10280, 0x1029f, Lycian,									"UnicodeBlock", "Lycian", "Lycian")
-REGISTER_UNICODE_BLOCK_RANGE(0x102a0, 0x102df, Carian,									"UnicodeBlock", "Carian", "Carian")
-REGISTER_UNICODE_BLOCK_RANGE(0x102e0, 0x102ff, CopticEpactNumbers,						"UnicodeBlock", "CopticEpactNumbers", "Coptic Epact Numbers")
-REGISTER_UNICODE_BLOCK_RANGE(0x10300, 0x1032f, OldItalic,								"UnicodeBlock", "OldItalic", "Old Italic")
-REGISTER_UNICODE_BLOCK_RANGE(0x10330, 0x1034f, Gothic,									"UnicodeBlock", "Gothic", "Gothic")
-REGISTER_UNICODE_BLOCK_RANGE(0x10350, 0x1037f, OldPermic,								"UnicodeBlock", "OldPermic", "Old Permic")
-REGISTER_UNICODE_BLOCK_RANGE(0x10380, 0x1039f, Ugaritic,								"UnicodeBlock", "Ugaritic", "Ugaritic")
-REGISTER_UNICODE_BLOCK_RANGE(0x103a0, 0x103df, OldPersian,								"UnicodeBlock", "Old Persian", "Old Persian")
-REGISTER_UNICODE_BLOCK_RANGE(0x10400, 0x1044f, Deseret,									"UnicodeBlock", "Deseret", "Deseret")
-REGISTER_UNICODE_BLOCK_RANGE(0x10450, 0x1047f, Shavian,									"UnicodeBlock", "Shavian", "Shavian")
-REGISTER_UNICODE_BLOCK_RANGE(0x10480, 0x104af, Osmanya,									"UnicodeBlock", "Osmanya", "Osmanya")
-REGISTER_UNICODE_BLOCK_RANGE(0x104b0, 0x104ff, Osage,									"UnicodeBlock", "Osage", "Osage")
-REGISTER_UNICODE_BLOCK_RANGE(0x10500, 0x1052f, Elbasan,									"UnicodeBlock", "Elbasan", "Elbasan")
-REGISTER_UNICODE_BLOCK_RANGE(0x10530, 0x1056f, CaucasianAlbanian,						"UnicodeBlock", "CaucasianAlbanian", "Caucasian Albanian")
-REGISTER_UNICODE_BLOCK_RANGE(0x10600, 0x1077f, LinearA,									"UnicodeBlock", "LinearA", "Linear A")
-REGISTER_UNICODE_BLOCK_RANGE(0x10800, 0x1083f, CypriotSyllabary,						"UnicodeBlock", "CypriotSyllabary", "Cypriot Syllabary")
-REGISTER_UNICODE_BLOCK_RANGE(0x10840, 0x1085f, ImperialAramaic,							"UnicodeBlock", "ImperialAramaic", "Imperial Aramaic")
-REGISTER_UNICODE_BLOCK_RANGE(0x10860, 0x1087f, Palmyrene,								"UnicodeBlock", "Palmyrene", "Palmyrene")
-REGISTER_UNICODE_BLOCK_RANGE(0x10880, 0x108af, Nabataean,								"UnicodeBlock", "Nabataean", "Nabataean")
-REGISTER_UNICODE_BLOCK_RANGE(0x108e0, 0x108ff, Hatran,									"UnicodeBlock", "Hatran", "Hatran")
-REGISTER_UNICODE_BLOCK_RANGE(0x10900, 0x1091f, Phoenician,								"UnicodeBlock", "Phoenician", "Phoenician")
-REGISTER_UNICODE_BLOCK_RANGE(0x10920, 0x1093f, Lydian,									"UnicodeBlock", "Lydian", "Lydian")
-REGISTER_UNICODE_BLOCK_RANGE(0x10980, 0x1099f, MeroiticHieroglyphs,						"UnicodeBlock", "MeroiticHieroglyphs", "Meroitic Hieroglyphs")
-REGISTER_UNICODE_BLOCK_RANGE(0x109a0, 0x109ff, MeroiticCursive,							"UnicodeBlock", "MeroiticCursive", "Meroitic Cursive")
-REGISTER_UNICODE_BLOCK_RANGE(0x10a00, 0x10a5f, Kharoshthi,								"UnicodeBlock", "Kharoshthi", "Kharoshthi")
-REGISTER_UNICODE_BLOCK_RANGE(0x10a60, 0x10a7f, OldSouthArabian,							"UnicodeBlock", "OldSouthArabian", "Old South Arabian")
-REGISTER_UNICODE_BLOCK_RANGE(0x10a80, 0x10a9f, OldNorthArabian,							"UnicodeBlock", "OldNorthArabian", "Old North Arabian")
-REGISTER_UNICODE_BLOCK_RANGE(0x10ac0, 0x10aff, Manichaean,								"UnicodeBlock", "Manichaean", "Manichaean")
-REGISTER_UNICODE_BLOCK_RANGE(0x10b00, 0x10b3f, Avestan,									"UnicodeBlock", "Avestan", "Avestan")
-REGISTER_UNICODE_BLOCK_RANGE(0x10b40, 0x10b5f, InscriptionalParthian,					"UnicodeBlock", "InscriptionalParthian", "Inscriptional Parthian")
-REGISTER_UNICODE_BLOCK_RANGE(0x10b60, 0x10b7f, InscriptionalPahlavi,					"UnicodeBlock", "InscriptionalPahlavi", "Inscriptional Pahlavi")
-REGISTER_UNICODE_BLOCK_RANGE(0x10b80, 0x10baf, PsalterPahlavi,							"UnicodeBlock", "PsalterPahlavi", "Psalter Pahlavi")
-REGISTER_UNICODE_BLOCK_RANGE(0x10c00, 0x10c4f, OldTurkic,								"UnicodeBlock", "OldTurkic", "Old Turkic")
-REGISTER_UNICODE_BLOCK_RANGE(0x10c80, 0x10cff, OldHungarian,							"UnicodeBlock", "OldHungarian", "Old Hungarian")
-REGISTER_UNICODE_BLOCK_RANGE(0x10e60, 0x10e7f, RumiNumeralSymbols,						"UnicodeBlock", "RumiNumeralSymbols", "Rumi Numeral Symbols")
-REGISTER_UNICODE_BLOCK_RANGE(0x11000, 0x1107f, Brahmi,									"UnicodeBlock", "Brahmi", "Brahmi")
-REGISTER_UNICODE_BLOCK_RANGE(0x11080, 0x110cf, Kaithi,									"UnicodeBlock", "Kaithi", "Kaithi")
-REGISTER_UNICODE_BLOCK_RANGE(0x110d0, 0x110ff, SoraSompeng,								"UnicodeBlock", "SoraSompeng", "Sora Sompeng")
-REGISTER_UNICODE_BLOCK_RANGE(0x11100, 0x1114f, Chakma,									"UnicodeBlock", "Chakma", "Chakma")
-REGISTER_UNICODE_BLOCK_RANGE(0x11150, 0x1117f, Mahajani,								"UnicodeBlock", "Mahajani", "Mahajani")
-REGISTER_UNICODE_BLOCK_RANGE(0x11180, 0x111df, Sharada,									"UnicodeBlock", "Sharada", "Sharada")
-REGISTER_UNICODE_BLOCK_RANGE(0x111e0, 0x111ff, SinhalaArchaicNumbers,					"UnicodeBlock", "SinhalaArchaicNumbers", "Sinhala Archaic Numbers")
-REGISTER_UNICODE_BLOCK_RANGE(0x11200, 0x1124f, Khojki,									"UnicodeBlock", "Khojki", "Khojki")
-REGISTER_UNICODE_BLOCK_RANGE(0x11280, 0x112af, Multani,									"UnicodeBlock", "Multani", "Multani")
-REGISTER_UNICODE_BLOCK_RANGE(0x112b0, 0x112ff, Khudawadi,								"UnicodeBlock", "Khudawadi", "Khudawadi")
-REGISTER_UNICODE_BLOCK_RANGE(0x11300, 0x1137f, Grantha,									"UnicodeBlock", "Grantha", "Grantha")
-REGISTER_UNICODE_BLOCK_RANGE(0x11400, 0x1147f, Newa,									"UnicodeBlock", "Newa", "Newa")
-REGISTER_UNICODE_BLOCK_RANGE(0x11480, 0x114df, Tirhuta,									"UnicodeBlock", "Tirhuta", "Tirhuta")
-REGISTER_UNICODE_BLOCK_RANGE(0x11580, 0x115ff, Siddham,									"UnicodeBlock", "Siddham", "Siddham")
-REGISTER_UNICODE_BLOCK_RANGE(0x11600, 0x1165f, Modi,									"UnicodeBlock", "Modi", "Modi")
-REGISTER_UNICODE_BLOCK_RANGE(0x11660, 0x1167f, MongolianSupplement,						"UnicodeBlock", "MongolianSupplement", "Mongolian Supplement")
-REGISTER_UNICODE_BLOCK_RANGE(0x11680, 0x116cf, Takri,									"UnicodeBlock", "Takri", "Takri")
-REGISTER_UNICODE_BLOCK_RANGE(0x11700, 0x1173f, Ahom,									"UnicodeBlock", "Ahom", "Ahom")
-REGISTER_UNICODE_BLOCK_RANGE(0x118a0, 0x118ff, WarangCiti,								"UnicodeBlock", "WarangCiti", "Warang Citi")
-REGISTER_UNICODE_BLOCK_RANGE(0x11ac0, 0x11aff, PauCinHau,								"UnicodeBlock", "PauCinHau", "Pau Cin Hau")
-REGISTER_UNICODE_BLOCK_RANGE(0x11c00, 0x11c6f, Bhaiksuki,								"UnicodeBlock", "Bhaiksuki", "Bhaiksuki")
-REGISTER_UNICODE_BLOCK_RANGE(0x11c70, 0x11cbf, Marchen,									"UnicodeBlock", "Marchen", "Marchen")
-REGISTER_UNICODE_BLOCK_RANGE(0x12000, 0x123ff, Cuneiform,								"UnicodeBlock", "Cuneiform", "Cuneiform")
-REGISTER_UNICODE_BLOCK_RANGE(0x12400, 0x1247f, CuneiformNumbersAndPunctuation,			"UnicodeBlock", "CuneiformNumbersAndPunctuation", "Cuneiform Numbers and Punctuation")
-REGISTER_UNICODE_BLOCK_RANGE(0x12480, 0x1254f, EarlyDynasticCuneiform,					"UnicodeBlock", "EarlyDynasticCuneiform", "Early Dynastic Cuneiform")
-REGISTER_UNICODE_BLOCK_RANGE(0x13000, 0x1342f, EgyptianHieroglyphs,						"UnicodeBlock", "EgyptianHieroglyphs", "Egyptian Hieroglyphs")
-REGISTER_UNICODE_BLOCK_RANGE(0x14400, 0x1467f, AnatolianHieroglyphs,					"UnicodeBlock", "AnatolianHieroglyphs", "Anatolian Hieroglyphs")
-REGISTER_UNICODE_BLOCK_RANGE(0x16800, 0x16a3f, BamumSupplement,							"UnicodeBlock", "BamumSupplement", "Bamum Supplement")
-REGISTER_UNICODE_BLOCK_RANGE(0x16a40, 0x16a6f, Mro,										"UnicodeBlock", "Mro", "Mro")
-REGISTER_UNICODE_BLOCK_RANGE(0x16ad0, 0x16aff, BassaVah,								"UnicodeBlock", "BassaVah", "Bassa Vah")
-REGISTER_UNICODE_BLOCK_RANGE(0x16b00, 0x16b8f, PahawhHmong,								"UnicodeBlock", "PahawhHmong", "Pahawh Hmong")
-REGISTER_UNICODE_BLOCK_RANGE(0x16f00, 0x16f9f, Miao,									"UnicodeBlock", "Miao", "Miao")
-REGISTER_UNICODE_BLOCK_RANGE(0x16fe0, 0x16fff, IdeographicSymbolsAndPunctuation,		"UnicodeBlock", "IdeographicSymbolsAndPunctuation", "Ideographic Symbols and Punctuation")
-REGISTER_UNICODE_BLOCK_RANGE(0x17000, 0x187ff, Tangut,									"UnicodeBlock", "Tangut", "Tangut")
-REGISTER_UNICODE_BLOCK_RANGE(0x18800, 0x18aff, TangutComponents,						"UnicodeBlock", "TangutComponents", "Tangut Components")
-REGISTER_UNICODE_BLOCK_RANGE(0x1b000, 0x1b0ff, KanaSupplement,							"UnicodeBlock", "KanaSupplement", "Kana Supplement")
-REGISTER_UNICODE_BLOCK_RANGE(0x1bc00, 0x1bc9f, Duployan,								"UnicodeBlock", "Duployan", "Duployan")
-REGISTER_UNICODE_BLOCK_RANGE(0x1bca0, 0x1bcaf, ShorthandFormatControls,					"UnicodeBlock", "ShorthandFormatControls", "Shorthand Format Controls")
-REGISTER_UNICODE_BLOCK_RANGE(0x1d000, 0x1d0ff, ByzantineMusicalSymbols,					"UnicodeBlock", "ByzantineMusicalSymbols", "Byzantine Musical Symbols")
-REGISTER_UNICODE_BLOCK_RANGE(0x1d100, 0x1d1ff, MusicalSymbols,							"UnicodeBlock", "MusicalSymbols", "Musical Symbols")
-REGISTER_UNICODE_BLOCK_RANGE(0x1d200, 0x1d24f, AncientGreekMusicalNotation,				"UnicodeBlock", "AncientGreekMusicalNotation", "Ancient Greek Musical Notation")
-REGISTER_UNICODE_BLOCK_RANGE(0x1d300, 0x1d35f, TaiXuanJingSymbols,						"UnicodeBlock", "TaiXuanJingSymbols", "Tai Xuan Jing Symbols")
-REGISTER_UNICODE_BLOCK_RANGE(0x1d360, 0x1d37f, CountingRodNumerals,						"UnicodeBlock", "CountingRodNumerals", "Counting Rod Numerals")
-REGISTER_UNICODE_BLOCK_RANGE(0x1d400, 0x1d7ff, MathematicalAlphanumericSymbols,			"UnicodeBlock", "MathematicalAlphanumericSymbols", "Mathematical Alphanumeric Symbols")
-REGISTER_UNICODE_BLOCK_RANGE(0x1d800, 0x1daaf, SuttonSignWriting,						"UnicodeBlock", "SuttonSignWriting", "Sutton SignWriting")
-REGISTER_UNICODE_BLOCK_RANGE(0x1e000, 0x1e02f, GlagoliticSupplement,					"UnicodeBlock", "GlagoliticSupplement", "Glagolitic Supplement")
-REGISTER_UNICODE_BLOCK_RANGE(0x1e800, 0x1e8df, MendeKikakui,							"UnicodeBlock", "MendeKikakui", "Mende Kikakui")
-REGISTER_UNICODE_BLOCK_RANGE(0x1e900, 0x1e95f, Adlam,									"UnicodeBlock", "Adlam", "Adlam")
-REGISTER_UNICODE_BLOCK_RANGE(0x1ee00, 0x1eeff, ArabicMathematicalAlphabeticSymbols,		"UnicodeBlock", "ArabicMathematicalAlphabeticSymbols", "Arabic Mathematical Alphabetic Symbols")
-REGISTER_UNICODE_BLOCK_RANGE(0x1f000, 0x1f02f, MahjongTiles,							"UnicodeBlock", "MahjongTiles", "Mahjong Tiles")
-REGISTER_UNICODE_BLOCK_RANGE(0x1f030, 0x1f09f, DominoTiles,								"UnicodeBlock", "DominoTiles", "Domino Tiles")
-REGISTER_UNICODE_BLOCK_RANGE(0x1f0a0, 0x1f0ff, PlayingCards,							"UnicodeBlock", "PlayingCards", "Playing Cards")
-REGISTER_UNICODE_BLOCK_RANGE(0x1f100, 0x1f1ff, EnclosedAlphanumericSupplement,			"UnicodeBlock", "EnclosedAlphanumericSupplement", "Enclosed Alphanumeric Supplement")
-REGISTER_UNICODE_BLOCK_RANGE(0x1f200, 0x1f2ff, EnclosedIdeographicSupplement,			"UnicodeBlock", "EnclosedIdeographicSupplement", "Enclosed Ideographic Supplement")
-REGISTER_UNICODE_BLOCK_RANGE(0x1f300, 0x1f5ff, MiscellaneousSymbolsAndPictographs,		"UnicodeBlock", "MiscellaneousSymbolsAndPictographs", "Miscellaneous Symbols and Pictographs")
-REGISTER_UNICODE_BLOCK_RANGE(0x1f600, 0x1f64f, EmoticonsEmoji,							"UnicodeBlock", "EmoticonsEmoji", "Emoticons (Emoji)")
-REGISTER_UNICODE_BLOCK_RANGE(0x1f650, 0x1f67f, OrnamentalDingbats,						"UnicodeBlock", "OrnamentalDingbats", "Ornamental Dingbats")
-REGISTER_UNICODE_BLOCK_RANGE(0x1f680, 0x1f6ff, TransportAndMapSymbols,					"UnicodeBlock", "TransportAndMapSymbols", "Transport and Map Symbols")
-REGISTER_UNICODE_BLOCK_RANGE(0x1f700, 0x1f77f, AlchemicalSymbols,						"UnicodeBlock", "AlchemicalSymbols", "Alchemical Symbols")
-REGISTER_UNICODE_BLOCK_RANGE(0x1f780, 0x1f7ff, GeometricShapesExtended,					"UnicodeBlock", "GeometricShapesExtended", "Geometric Shapes Extended")
-REGISTER_UNICODE_BLOCK_RANGE(0x1f800, 0x1f8ff, SupplementalArrowsC,						"UnicodeBlock", "SupplementalArrowsC", "Supplemental Arrows - C")
-REGISTER_UNICODE_BLOCK_RANGE(0x1f900, 0x1f9ff, SupplementalSymbolsAndPictographs,		"UnicodeBlock", "SupplementalSymbolsAndPictographs", "Supplemental Symbols and Pictographs")
-REGISTER_UNICODE_BLOCK_RANGE(0x20000, 0x2a6d6, CJKUnifiedIdeographsExtensionB,			"UnicodeBlock", "CJKUnifiedIdeographsExtensionB", "CJK Unified Ideographs Extension B")
-REGISTER_UNICODE_BLOCK_RANGE(0x2a700, 0x2b734, CJKUnifiedIdeographsExtensionC,			"UnicodeBlock", "CJKUnifiedIdeographsExtensionC", "CJK Unified Ideographs Extension C")
-REGISTER_UNICODE_BLOCK_RANGE(0x2b740, 0x2b81d, CJKUnifiedIdeographsExtensionD,			"UnicodeBlock", "CJKUnifiedIdeographsExtensionD", "CJK Unified Ideographs Extension D")
-REGISTER_UNICODE_BLOCK_RANGE(0x2b820, 0x2cea1, CJKUnifiedIdeographsExtensionE,			"UnicodeBlock", "CJKUnifiedIdeographsExtensionE", "CJK Unified Ideographs Extension E")
-REGISTER_UNICODE_BLOCK_RANGE(0x2f800, 0x2fa1f, CJKCompatibilityIdeographsSupplement,	"UnicodeBlock", "CJKCompatibilityIdeographsSupplement", "CJK Compatibility Ideographs Supplement")
-REGISTER_UNICODE_BLOCK_RANGE(0xe0000, 0xe007f, Tags,									"UnicodeBlock", "Tags", "Tags")
-REGISTER_UNICODE_BLOCK_RANGE(0xe0100, 0xe01ef, VariationSelectorsSupplement,			"UnicodeBlock", "VariationSelectorsSupplement", "Variation Selectors Supplement")
+REGISTER_UNICODE_BLOCK_RANGE(0x10000, 0x1007f, LinearBSyllabary,						"LinearBSyllabary", "Linear B Syllabary")
+REGISTER_UNICODE_BLOCK_RANGE(0x10080, 0x100ff, LinearBIdeograms,						"LinearBIdeograms", "Linear B Ideograms")
+REGISTER_UNICODE_BLOCK_RANGE(0x10100, 0x1013f, AegeanNumbers,							"AegeanNumbers", "Aegean Numbers")
+REGISTER_UNICODE_BLOCK_RANGE(0x10140, 0x1018f, AncientGreekNumbers,						"AncientGreekNumbers", "Ancient Greek Numbers")
+REGISTER_UNICODE_BLOCK_RANGE(0x10190, 0x101cf, AncientSymbols,							"AncientSymbols", "Ancient Symbols")
+REGISTER_UNICODE_BLOCK_RANGE(0x101d0, 0x101ff, PhaistosDisc,							"PhaistosDisc", "Phaistos Disc")
+REGISTER_UNICODE_BLOCK_RANGE(0x10280, 0x1029f, Lycian,									"Lycian", "Lycian")
+REGISTER_UNICODE_BLOCK_RANGE(0x102a0, 0x102df, Carian,									"Carian", "Carian")
+REGISTER_UNICODE_BLOCK_RANGE(0x102e0, 0x102ff, CopticEpactNumbers,						"CopticEpactNumbers", "Coptic Epact Numbers")
+REGISTER_UNICODE_BLOCK_RANGE(0x10300, 0x1032f, OldItalic,								"OldItalic", "Old Italic")
+REGISTER_UNICODE_BLOCK_RANGE(0x10330, 0x1034f, Gothic,									"Gothic", "Gothic")
+REGISTER_UNICODE_BLOCK_RANGE(0x10350, 0x1037f, OldPermic,								"OldPermic", "Old Permic")
+REGISTER_UNICODE_BLOCK_RANGE(0x10380, 0x1039f, Ugaritic,								"Ugaritic", "Ugaritic")
+REGISTER_UNICODE_BLOCK_RANGE(0x103a0, 0x103df, OldPersian,								"Old Persian", "Old Persian")
+REGISTER_UNICODE_BLOCK_RANGE(0x10400, 0x1044f, Deseret,									"Deseret", "Deseret")
+REGISTER_UNICODE_BLOCK_RANGE(0x10450, 0x1047f, Shavian,									"Shavian", "Shavian")
+REGISTER_UNICODE_BLOCK_RANGE(0x10480, 0x104af, Osmanya,									"Osmanya", "Osmanya")
+REGISTER_UNICODE_BLOCK_RANGE(0x104b0, 0x104ff, Osage,									"Osage", "Osage")
+REGISTER_UNICODE_BLOCK_RANGE(0x10500, 0x1052f, Elbasan,									"Elbasan", "Elbasan")
+REGISTER_UNICODE_BLOCK_RANGE(0x10530, 0x1056f, CaucasianAlbanian,						"CaucasianAlbanian", "Caucasian Albanian")
+REGISTER_UNICODE_BLOCK_RANGE(0x10600, 0x1077f, LinearA,									"LinearA", "Linear A")
+REGISTER_UNICODE_BLOCK_RANGE(0x10800, 0x1083f, CypriotSyllabary,						"CypriotSyllabary", "Cypriot Syllabary")
+REGISTER_UNICODE_BLOCK_RANGE(0x10840, 0x1085f, ImperialAramaic,							"ImperialAramaic", "Imperial Aramaic")
+REGISTER_UNICODE_BLOCK_RANGE(0x10860, 0x1087f, Palmyrene,								"Palmyrene", "Palmyrene")
+REGISTER_UNICODE_BLOCK_RANGE(0x10880, 0x108af, Nabataean,								"Nabataean", "Nabataean")
+REGISTER_UNICODE_BLOCK_RANGE(0x108e0, 0x108ff, Hatran,									"Hatran", "Hatran")
+REGISTER_UNICODE_BLOCK_RANGE(0x10900, 0x1091f, Phoenician,								"Phoenician", "Phoenician")
+REGISTER_UNICODE_BLOCK_RANGE(0x10920, 0x1093f, Lydian,									"Lydian", "Lydian")
+REGISTER_UNICODE_BLOCK_RANGE(0x10980, 0x1099f, MeroiticHieroglyphs,						"MeroiticHieroglyphs", "Meroitic Hieroglyphs")
+REGISTER_UNICODE_BLOCK_RANGE(0x109a0, 0x109ff, MeroiticCursive,							"MeroiticCursive", "Meroitic Cursive")
+REGISTER_UNICODE_BLOCK_RANGE(0x10a00, 0x10a5f, Kharoshthi,								"Kharoshthi", "Kharoshthi")
+REGISTER_UNICODE_BLOCK_RANGE(0x10a60, 0x10a7f, OldSouthArabian,							"OldSouthArabian", "Old South Arabian")
+REGISTER_UNICODE_BLOCK_RANGE(0x10a80, 0x10a9f, OldNorthArabian,							"OldNorthArabian", "Old North Arabian")
+REGISTER_UNICODE_BLOCK_RANGE(0x10ac0, 0x10aff, Manichaean,								"Manichaean", "Manichaean")
+REGISTER_UNICODE_BLOCK_RANGE(0x10b00, 0x10b3f, Avestan,									"Avestan", "Avestan")
+REGISTER_UNICODE_BLOCK_RANGE(0x10b40, 0x10b5f, InscriptionalParthian,					"InscriptionalParthian", "Inscriptional Parthian")
+REGISTER_UNICODE_BLOCK_RANGE(0x10b60, 0x10b7f, InscriptionalPahlavi,					"InscriptionalPahlavi", "Inscriptional Pahlavi")
+REGISTER_UNICODE_BLOCK_RANGE(0x10b80, 0x10baf, PsalterPahlavi,							"PsalterPahlavi", "Psalter Pahlavi")
+REGISTER_UNICODE_BLOCK_RANGE(0x10c00, 0x10c4f, OldTurkic,								"OldTurkic", "Old Turkic")
+REGISTER_UNICODE_BLOCK_RANGE(0x10c80, 0x10cff, OldHungarian,							"OldHungarian", "Old Hungarian")
+REGISTER_UNICODE_BLOCK_RANGE(0x10e60, 0x10e7f, RumiNumeralSymbols,						"RumiNumeralSymbols", "Rumi Numeral Symbols")
+REGISTER_UNICODE_BLOCK_RANGE(0x11000, 0x1107f, Brahmi,									"Brahmi", "Brahmi")
+REGISTER_UNICODE_BLOCK_RANGE(0x11080, 0x110cf, Kaithi,									"Kaithi", "Kaithi")
+REGISTER_UNICODE_BLOCK_RANGE(0x110d0, 0x110ff, SoraSompeng,								"SoraSompeng", "Sora Sompeng")
+REGISTER_UNICODE_BLOCK_RANGE(0x11100, 0x1114f, Chakma,									"Chakma", "Chakma")
+REGISTER_UNICODE_BLOCK_RANGE(0x11150, 0x1117f, Mahajani,								"Mahajani", "Mahajani")
+REGISTER_UNICODE_BLOCK_RANGE(0x11180, 0x111df, Sharada,									"Sharada", "Sharada")
+REGISTER_UNICODE_BLOCK_RANGE(0x111e0, 0x111ff, SinhalaArchaicNumbers,					"SinhalaArchaicNumbers", "Sinhala Archaic Numbers")
+REGISTER_UNICODE_BLOCK_RANGE(0x11200, 0x1124f, Khojki,									"Khojki", "Khojki")
+REGISTER_UNICODE_BLOCK_RANGE(0x11280, 0x112af, Multani,									"Multani", "Multani")
+REGISTER_UNICODE_BLOCK_RANGE(0x112b0, 0x112ff, Khudawadi,								"Khudawadi", "Khudawadi")
+REGISTER_UNICODE_BLOCK_RANGE(0x11300, 0x1137f, Grantha,									"Grantha", "Grantha")
+REGISTER_UNICODE_BLOCK_RANGE(0x11400, 0x1147f, Newa,									"Newa", "Newa")
+REGISTER_UNICODE_BLOCK_RANGE(0x11480, 0x114df, Tirhuta,									"Tirhuta", "Tirhuta")
+REGISTER_UNICODE_BLOCK_RANGE(0x11580, 0x115ff, Siddham,									"Siddham", "Siddham")
+REGISTER_UNICODE_BLOCK_RANGE(0x11600, 0x1165f, Modi,									"Modi", "Modi")
+REGISTER_UNICODE_BLOCK_RANGE(0x11660, 0x1167f, MongolianSupplement,						"MongolianSupplement", "Mongolian Supplement")
+REGISTER_UNICODE_BLOCK_RANGE(0x11680, 0x116cf, Takri,									"Takri", "Takri")
+REGISTER_UNICODE_BLOCK_RANGE(0x11700, 0x1173f, Ahom,									"Ahom", "Ahom")
+REGISTER_UNICODE_BLOCK_RANGE(0x118a0, 0x118ff, WarangCiti,								"WarangCiti", "Warang Citi")
+REGISTER_UNICODE_BLOCK_RANGE(0x11ac0, 0x11aff, PauCinHau,								"PauCinHau", "Pau Cin Hau")
+REGISTER_UNICODE_BLOCK_RANGE(0x11c00, 0x11c6f, Bhaiksuki,								"Bhaiksuki", "Bhaiksuki")
+REGISTER_UNICODE_BLOCK_RANGE(0x11c70, 0x11cbf, Marchen,									"Marchen", "Marchen")
+REGISTER_UNICODE_BLOCK_RANGE(0x12000, 0x123ff, Cuneiform,								"Cuneiform", "Cuneiform")
+REGISTER_UNICODE_BLOCK_RANGE(0x12400, 0x1247f, CuneiformNumbersAndPunctuation,			"CuneiformNumbersAndPunctuation", "Cuneiform Numbers and Punctuation")
+REGISTER_UNICODE_BLOCK_RANGE(0x12480, 0x1254f, EarlyDynasticCuneiform,					"EarlyDynasticCuneiform", "Early Dynastic Cuneiform")
+REGISTER_UNICODE_BLOCK_RANGE(0x13000, 0x1342f, EgyptianHieroglyphs,						"EgyptianHieroglyphs", "Egyptian Hieroglyphs")
+REGISTER_UNICODE_BLOCK_RANGE(0x14400, 0x1467f, AnatolianHieroglyphs,					"AnatolianHieroglyphs", "Anatolian Hieroglyphs")
+REGISTER_UNICODE_BLOCK_RANGE(0x16800, 0x16a3f, BamumSupplement,							"BamumSupplement", "Bamum Supplement")
+REGISTER_UNICODE_BLOCK_RANGE(0x16a40, 0x16a6f, Mro,										"Mro", "Mro")
+REGISTER_UNICODE_BLOCK_RANGE(0x16ad0, 0x16aff, BassaVah,								"BassaVah", "Bassa Vah")
+REGISTER_UNICODE_BLOCK_RANGE(0x16b00, 0x16b8f, PahawhHmong,								"PahawhHmong", "Pahawh Hmong")
+REGISTER_UNICODE_BLOCK_RANGE(0x16f00, 0x16f9f, Miao,									"Miao", "Miao")
+REGISTER_UNICODE_BLOCK_RANGE(0x16fe0, 0x16fff, IdeographicSymbolsAndPunctuation,		"IdeographicSymbolsAndPunctuation", "Ideographic Symbols and Punctuation")
+REGISTER_UNICODE_BLOCK_RANGE(0x17000, 0x187ff, Tangut,									"Tangut", "Tangut")
+REGISTER_UNICODE_BLOCK_RANGE(0x18800, 0x18aff, TangutComponents,						"TangutComponents", "Tangut Components")
+REGISTER_UNICODE_BLOCK_RANGE(0x1b000, 0x1b0ff, KanaSupplement,							"KanaSupplement", "Kana Supplement")
+REGISTER_UNICODE_BLOCK_RANGE(0x1bc00, 0x1bc9f, Duployan,								"Duployan", "Duployan")
+REGISTER_UNICODE_BLOCK_RANGE(0x1bca0, 0x1bcaf, ShorthandFormatControls,					"ShorthandFormatControls", "Shorthand Format Controls")
+REGISTER_UNICODE_BLOCK_RANGE(0x1d000, 0x1d0ff, ByzantineMusicalSymbols,					"ByzantineMusicalSymbols", "Byzantine Musical Symbols")
+REGISTER_UNICODE_BLOCK_RANGE(0x1d100, 0x1d1ff, MusicalSymbols,							"MusicalSymbols", "Musical Symbols")
+REGISTER_UNICODE_BLOCK_RANGE(0x1d200, 0x1d24f, AncientGreekMusicalNotation,				"AncientGreekMusicalNotation", "Ancient Greek Musical Notation")
+REGISTER_UNICODE_BLOCK_RANGE(0x1d300, 0x1d35f, TaiXuanJingSymbols,						"TaiXuanJingSymbols", "Tai Xuan Jing Symbols")
+REGISTER_UNICODE_BLOCK_RANGE(0x1d360, 0x1d37f, CountingRodNumerals,						"CountingRodNumerals", "Counting Rod Numerals")
+REGISTER_UNICODE_BLOCK_RANGE(0x1d400, 0x1d7ff, MathematicalAlphanumericSymbols,			"MathematicalAlphanumericSymbols", "Mathematical Alphanumeric Symbols")
+REGISTER_UNICODE_BLOCK_RANGE(0x1d800, 0x1daaf, SuttonSignWriting,						"SuttonSignWriting", "Sutton SignWriting")
+REGISTER_UNICODE_BLOCK_RANGE(0x1e000, 0x1e02f, GlagoliticSupplement,					"GlagoliticSupplement", "Glagolitic Supplement")
+REGISTER_UNICODE_BLOCK_RANGE(0x1e800, 0x1e8df, MendeKikakui,							"MendeKikakui", "Mende Kikakui")
+REGISTER_UNICODE_BLOCK_RANGE(0x1e900, 0x1e95f, Adlam,									"Adlam", "Adlam")
+REGISTER_UNICODE_BLOCK_RANGE(0x1ee00, 0x1eeff, ArabicMathematicalAlphabeticSymbols,		"ArabicMathematicalAlphabeticSymbols", "Arabic Mathematical Alphabetic Symbols")
+REGISTER_UNICODE_BLOCK_RANGE(0x1f000, 0x1f02f, MahjongTiles,							"MahjongTiles", "Mahjong Tiles")
+REGISTER_UNICODE_BLOCK_RANGE(0x1f030, 0x1f09f, DominoTiles,								"DominoTiles", "Domino Tiles")
+REGISTER_UNICODE_BLOCK_RANGE(0x1f0a0, 0x1f0ff, PlayingCards,							"PlayingCards", "Playing Cards")
+REGISTER_UNICODE_BLOCK_RANGE(0x1f100, 0x1f1ff, EnclosedAlphanumericSupplement,			"EnclosedAlphanumericSupplement", "Enclosed Alphanumeric Supplement")
+REGISTER_UNICODE_BLOCK_RANGE(0x1f200, 0x1f2ff, EnclosedIdeographicSupplement,			"EnclosedIdeographicSupplement", "Enclosed Ideographic Supplement")
+REGISTER_UNICODE_BLOCK_RANGE(0x1f300, 0x1f5ff, MiscellaneousSymbolsAndPictographs,		"MiscellaneousSymbolsAndPictographs", "Miscellaneous Symbols and Pictographs")
+REGISTER_UNICODE_BLOCK_RANGE(0x1f600, 0x1f64f, EmoticonsEmoji,							"EmoticonsEmoji", "Emoticons (Emoji)")
+REGISTER_UNICODE_BLOCK_RANGE(0x1f650, 0x1f67f, OrnamentalDingbats,						"OrnamentalDingbats", "Ornamental Dingbats")
+REGISTER_UNICODE_BLOCK_RANGE(0x1f680, 0x1f6ff, TransportAndMapSymbols,					"TransportAndMapSymbols", "Transport and Map Symbols")
+REGISTER_UNICODE_BLOCK_RANGE(0x1f700, 0x1f77f, AlchemicalSymbols,						"AlchemicalSymbols", "Alchemical Symbols")
+REGISTER_UNICODE_BLOCK_RANGE(0x1f780, 0x1f7ff, GeometricShapesExtended,					"GeometricShapesExtended", "Geometric Shapes Extended")
+REGISTER_UNICODE_BLOCK_RANGE(0x1f800, 0x1f8ff, SupplementalArrowsC,						"SupplementalArrowsC", "Supplemental Arrows - C")
+REGISTER_UNICODE_BLOCK_RANGE(0x1f900, 0x1f9ff, SupplementalSymbolsAndPictographs,		"SupplementalSymbolsAndPictographs", "Supplemental Symbols and Pictographs")
+REGISTER_UNICODE_BLOCK_RANGE(0x20000, 0x2a6d6, CJKUnifiedIdeographsExtensionB,			"CJKUnifiedIdeographsExtensionB", "CJK Unified Ideographs Extension B")
+REGISTER_UNICODE_BLOCK_RANGE(0x2a700, 0x2b734, CJKUnifiedIdeographsExtensionC,			"CJKUnifiedIdeographsExtensionC", "CJK Unified Ideographs Extension C")
+REGISTER_UNICODE_BLOCK_RANGE(0x2b740, 0x2b81d, CJKUnifiedIdeographsExtensionD,			"CJKUnifiedIdeographsExtensionD", "CJK Unified Ideographs Extension D")
+REGISTER_UNICODE_BLOCK_RANGE(0x2b820, 0x2cea1, CJKUnifiedIdeographsExtensionE,			"CJKUnifiedIdeographsExtensionE", "CJK Unified Ideographs Extension E")
+REGISTER_UNICODE_BLOCK_RANGE(0x2f800, 0x2fa1f, CJKCompatibilityIdeographsSupplement,	"CJKCompatibilityIdeographsSupplement", "CJK Compatibility Ideographs Supplement")
+REGISTER_UNICODE_BLOCK_RANGE(0xe0000, 0xe007f, Tags,									"Tags", "Tags")
+REGISTER_UNICODE_BLOCK_RANGE(0xe0100, 0xe01ef, VariationSelectorsSupplement,			"VariationSelectorsSupplement", "Variation Selectors Supplement")
diff --git a/Engine/Source/Runtime/SlateCore/Public/Input/Reply.h b/Engine/Source/Runtime/SlateCore/Public/Input/Reply.h
--- a/Engine/Source/Runtime/SlateCore/Public/Input/Reply.h
+++ b/Engine/Source/Runtime/SlateCore/Public/Input/Reply.h
@@ -243,35 +243,19 @@ public:
 		return FReply(false);
 	}
 
+public:
+	/** Default copy constructor. */
+	SLATECORE_API FReply(const FReply&);
+
+	/** Default copy-assignment operator. */
+	SLATECORE_API FReply& operator=(const FReply&);
+
 private:
 
 	/**
 	 * Hidden default constructor.
 	 */
-	FReply( bool bIsHandled )
-		: TReplyBase<FReply>(bIsHandled)
-		, RequestedMousePos()
-		, EventHandler(nullptr)
-		, MouseCaptor(nullptr)
-		, FocusRecipient(nullptr)
-		, MouseLockWidget(nullptr)
-		, DetectDragForWidget(nullptr)
-		, NavigationDestination(nullptr)
-		, DragDropContent(nullptr)
-		, FocusChangeReason(EFocusCause::SetDirectly)
-		, NavigationType(EUINavigation::Invalid)
-		, NavigationGenesis(ENavigationGenesis::User)
-		, NavigationSource(ENavigationSource::FocusedWidget)
-		, bReleaseMouseCapture(false)
-		, bSetUserFocus(false)
-		, bReleaseUserFocus(false)
-		, bAllUsers(false)
-		, bShouldReleaseMouseLock(false)
-		, bUseHighPrecisionMouse(false)
-		, bPreventThrottling(false)
-		, bEndDragDrop(false)
-	{ }
-		
+	SLATECORE_API FReply( bool bIsHandled );
 
 private:
 	friend class FSlateApplication;
diff --git a/Engine/Source/Runtime/SlateCore/Public/Sound/SlateSound.h b/Engine/Source/Runtime/SlateCore/Public/Sound/SlateSound.h
--- a/Engine/Source/Runtime/SlateCore/Public/Sound/SlateSound.h
+++ b/Engine/Source/Runtime/SlateCore/Public/Sound/SlateSound.h
@@ -19,9 +19,13 @@ struct FSlateSound
 public:
 
 	/** Default constructor. */
-	FSlateSound( )
-		: ResourceObject(nullptr)
-	{ }
+	SLATECORE_API FSlateSound( );
+
+	/** Default copy constructor. */
+	SLATECORE_API FSlateSound(const FSlateSound&);
+
+	/** Default copy-assignment operator. */
+	SLATECORE_API FSlateSound& operator=(const FSlateSound&);
 
 public:
 
diff --git a/Engine/Source/Runtime/SlateCore/Public/Styling/CoreStyle.h b/Engine/Source/Runtime/SlateCore/Public/Styling/CoreStyle.h
--- a/Engine/Source/Runtime/SlateCore/Public/Styling/CoreStyle.h
+++ b/Engine/Source/Runtime/SlateCore/Public/Styling/CoreStyle.h
@@ -89,39 +89,39 @@ namespace CoreStyleConstants
 {
 	// Note, these sizes are in Slate Units.
 	// Slate Units do NOT have to map to pixels.
-	inline const UE::Slate::FDeprecateVector2DResult Icon5x16(5.0f, 16.0f);
-	inline const UE::Slate::FDeprecateVector2DResult Icon6x8(6.0f, 8.0f);
-	inline const UE::Slate::FDeprecateVector2DResult Icon8x4(8.0f, 4.0f);
-	inline const UE::Slate::FDeprecateVector2DResult Icon16x4(16.0f, 4.0f);
-	inline const UE::Slate::FDeprecateVector2DResult Icon8x8(8.0f, 8.0f);
-	inline const UE::Slate::FDeprecateVector2DResult Icon4x4(4.0f, 4.0f);
-	inline const UE::Slate::FDeprecateVector2DResult Icon10x10(10.0f, 10.0f);
-	inline const UE::Slate::FDeprecateVector2DResult Icon12x12(12.0f, 12.0f);
-	inline const UE::Slate::FDeprecateVector2DResult Icon12x16(12.0f, 16.0f);
-	inline const UE::Slate::FDeprecateVector2DResult Icon14x14(14.0f, 14.0f);
-	inline const UE::Slate::FDeprecateVector2DResult Icon16x16(16.0f, 16.0f);
-	inline const UE::Slate::FDeprecateVector2DResult Icon18x18(18.0f, 18.0f);
-	inline const UE::Slate::FDeprecateVector2DResult Icon20x20(20.0f, 20.0f);
-	inline const UE::Slate::FDeprecateVector2DResult Icon22x22(22.0f, 22.0f);
-	inline const UE::Slate::FDeprecateVector2DResult Icon24x24(24.0f, 24.0f);
-	inline const UE::Slate::FDeprecateVector2DResult Icon25x25(25.0f, 25.0f);
-	inline const UE::Slate::FDeprecateVector2DResult Icon26x26(26.0f, 26.0f);
-	inline const UE::Slate::FDeprecateVector2DResult Icon32x32(32.0f, 32.0f);
-	inline const UE::Slate::FDeprecateVector2DResult Icon40x40(40.0f, 40.0f);
-	inline const UE::Slate::FDeprecateVector2DResult Icon64x64(64.0f, 64.0f);
-	inline const UE::Slate::FDeprecateVector2DResult Icon36x24(36.0f, 24.0f);
-	inline const UE::Slate::FDeprecateVector2DResult Icon128x128(128.0f, 128.0f);
+	extern SLATECORE_API const UE::Slate::FDeprecateVector2DResult Icon5x16;
+	extern SLATECORE_API const UE::Slate::FDeprecateVector2DResult Icon6x8;
+	extern SLATECORE_API const UE::Slate::FDeprecateVector2DResult Icon8x4;
+	extern SLATECORE_API const UE::Slate::FDeprecateVector2DResult Icon16x4;
+	extern SLATECORE_API const UE::Slate::FDeprecateVector2DResult Icon8x8;
+	extern SLATECORE_API const UE::Slate::FDeprecateVector2DResult Icon4x4;
+	extern SLATECORE_API const UE::Slate::FDeprecateVector2DResult Icon10x10;
+	extern SLATECORE_API const UE::Slate::FDeprecateVector2DResult Icon12x12;
+	extern SLATECORE_API const UE::Slate::FDeprecateVector2DResult Icon12x16;
+	extern SLATECORE_API const UE::Slate::FDeprecateVector2DResult Icon14x14;
+	extern SLATECORE_API const UE::Slate::FDeprecateVector2DResult Icon16x16;
+	extern SLATECORE_API const UE::Slate::FDeprecateVector2DResult Icon18x18;
+	extern SLATECORE_API const UE::Slate::FDeprecateVector2DResult Icon20x20;
+	extern SLATECORE_API const UE::Slate::FDeprecateVector2DResult Icon22x22;
+	extern SLATECORE_API const UE::Slate::FDeprecateVector2DResult Icon24x24;
+	extern SLATECORE_API const UE::Slate::FDeprecateVector2DResult Icon25x25;
+	extern SLATECORE_API const UE::Slate::FDeprecateVector2DResult Icon26x26;
+	extern SLATECORE_API const UE::Slate::FDeprecateVector2DResult Icon32x32;
+	extern SLATECORE_API const UE::Slate::FDeprecateVector2DResult Icon40x40;
+	extern SLATECORE_API const UE::Slate::FDeprecateVector2DResult Icon64x64;
+	extern SLATECORE_API const UE::Slate::FDeprecateVector2DResult Icon36x24;
+	extern SLATECORE_API const UE::Slate::FDeprecateVector2DResult Icon128x128;
 
 	// Common Margins
-	inline const FMargin DefaultMargins(8.f, 4.f);
+	extern SLATECORE_API const FMargin DefaultMargins;
 	// Buttons already have a built in (4., 2.) padding - adding to that a little
-	inline const FMargin ButtonMargins(12.f, 1.5f, 12.f, 1.5f);
+	extern SLATECORE_API const FMargin ButtonMargins;
 
-	inline const FMargin PressedButtonMargins(12.f, 2.5f, 12.f, 0.5f);
-	inline const FMargin ToggleButtonMargins(16.0f, 2.0f);
-	inline const FMargin ComboButtonMargin(8.f, 1.f, 8.f, 1.f);
-	inline const FMargin PressedComboButtonMargin(8.f, 2.f, 8.f, 0.f);
+	extern SLATECORE_API const FMargin PressedButtonMargins;
+	extern SLATECORE_API const FMargin ToggleButtonMargins;
+	extern SLATECORE_API const FMargin ComboButtonMargin;
+	extern SLATECORE_API const FMargin PressedComboButtonMargin;
 
-	inline const float InputFocusRadius = 4.f;
-	inline const float InputFocusThickness = 1.0f;
+	extern SLATECORE_API const float InputFocusRadius;
+	extern SLATECORE_API const float InputFocusThickness;
 }
diff --git a/Engine/Source/Runtime/SlateCore/Public/Styling/SlateColor.h b/Engine/Source/Runtime/SlateCore/Public/Styling/SlateColor.h
--- a/Engine/Source/Runtime/SlateCore/Public/Styling/SlateColor.h
+++ b/Engine/Source/Runtime/SlateCore/Public/Styling/SlateColor.h
@@ -92,9 +92,6 @@ public:
 		, ColorTableId(InColorTableId)
 	{ }
 
-	~FSlateColor()
-	{ }
-
 public:
 
 	/**
