Fix localization tooling getting confused by Unicode block range macro.

diff --git a/Engine/Source/Runtime/SlateCore/Public/Fonts/UnicodeBlockRange.h b/Engine/Source/Runtime/SlateCore/Public/Fonts/UnicodeBlockRange.h
--- a/Engine/Source/Runtime/SlateCore/Public/Fonts/UnicodeBlockRange.h
+++ b/Engine/Source/Runtime/SlateCore/Public/Fonts/UnicodeBlockRange.h
@@ -11,7 +11,9 @@
 /** Enumeration of pre-defined Unicode block ranges that can be used to access entries from FUnicodeBlockRange */
 enum class EUnicodeBlockRange : uint16
 {
-#define REGISTER_UNICODE_BLOCK_RANGE(LOWERBOUND, UPPERBOUND, SYMBOLNAME, LOCTEXT_KEY, LOCTEXT_LITERAL) SYMBOLNAME,
+// Macro parameter names here have to be TEXT_KEY and TEXT_LITERAL and not LOCTEXT_KEY/LOCTEXT_LITERAL
+// to not confuse the localization tooling.	
+#define REGISTER_UNICODE_BLOCK_RANGE(LOWERBOUND, UPPERBOUND, SYMBOLNAME, TEXT_KEY, TEXT_LITERAL) SYMBOLNAME,
 #include "UnicodeBlockRange.inl"
 #undef REGISTER_UNICODE_BLOCK_RANGE
 };
