Text drawing (non shaped) : prevent an ensure that could happens if the version of FSlateDrawElement::MakeText taking a start and an end index was generating an empty text to draw. The code is now shared with the other MakeText function and make len check properly.
Removed duplicated culling code, and fixed float comparison of the alpha that could send invisible text to the renderer for negative or small alpha values.

#tests Tested in the editor that text was still properly displayed, and culling working properly.

diff --git a/Engine/Source/Runtime/SlateCore/Private/Rendering/DrawElementTypes.cpp b/Engine/Source/Runtime/SlateCore/Private/Rendering/DrawElementTypes.cpp
--- a/Engine/Source/Runtime/SlateCore/Private/Rendering/DrawElementTypes.cpp
+++ b/Engine/Source/Runtime/SlateCore/Private/Rendering/DrawElementTypes.cpp
@@ -122,7 +122,7 @@ static bool ShouldCull(const FSlateWindowElementList& ElementList, const FPaintG
 
 static bool ShouldCull(const FSlateWindowElementList& ElementList, const FPaintGeometry& PaintGeometry, const FLinearColor& InTint)
 {
-	if (InTint.A == 0 || ShouldCull(ElementList, PaintGeometry))
+	if (InTint.A < UE_KINDA_SMALL_NUMBER || ShouldCull(ElementList, PaintGeometry))
 	{
 		return true;
 	}
@@ -132,7 +132,7 @@ static bool ShouldCull(const FSlateWindowElementList& ElementList, const FPaintG
 
 static bool ShouldCull(const FSlateWindowElementList& ElementList, const FPaintGeometry& PaintGeometry, const FLinearColor& InTint, const FString& InText)
 {
-	if (InTint.A == 0 || InText.Len() == 0 || ShouldCull(ElementList, PaintGeometry))
+	if (InTint.A < UE_KINDA_SMALL_NUMBER || InText.Len() == 0 || ShouldCull(ElementList, PaintGeometry))
 	{
 		return true;
 	}
@@ -142,7 +142,7 @@ static bool ShouldCull(const FSlateWindowElementList& ElementList, const FPaintG
 
 static bool ShouldCull(const FSlateWindowElementList& ElementList, const FPaintGeometry& PaintGeometry, const FSlateBrush* InBrush, const FLinearColor& InTint)
 {
-	bool bIsFullyTransparent = InTint.A == 0 && (InBrush->OutlineSettings.Color.GetSpecifiedColor().A == 0 || InBrush->OutlineSettings.bUseBrushTransparency);
+	bool bIsFullyTransparent = InTint.A < UE_KINDA_SMALL_NUMBER && (InBrush->OutlineSettings.Color.GetSpecifiedColor().A < UE_KINDA_SMALL_NUMBER || InBrush->OutlineSettings.bUseBrushTransparency);
 	if (bIsFullyTransparent || ShouldCull(ElementList, PaintGeometry, InBrush))
 	{
 		return true;
@@ -371,27 +371,16 @@ namespace UE::Slate::Private
 }
 #endif
 
-void FSlateDrawElement::MakeText( FSlateWindowElementList& ElementList, uint32 InLayer, const FPaintGeometry& PaintGeometry, const FString& InText, const int32 StartIndex, const int32 EndIndex, const FSlateFontInfo& InFontInfo, ESlateDrawEffect InDrawEffects, const FLinearColor& InTint )
+void FSlateDrawElement::MakeText( FSlateWindowElementList& ElementList, uint32 InLayer, const FPaintGeometry& PaintGeometry, const FString& InText, const int32 InStartIndex, const int32 InEndIndex, const FSlateFontInfo& InFontInfo, ESlateDrawEffect InDrawEffects, const FLinearColor& InTint )
 {
-	SCOPE_CYCLE_COUNTER(STAT_SlateDrawElementMakeTime)
-	PaintGeometry.CommitTransformsIfUsingLegacyConstructor();
-
-	if (ShouldCull(ElementList, PaintGeometry, InTint, InText))
+	const int32 SourceTextLen = InText.Len();
+	const int32 StartIndex = FMath::Min<int32>(InStartIndex, SourceTextLen);
+	const int32 EndIndex = FMath::Min<int32>(InEndIndex, SourceTextLen);
+	const int32 TextLength = (EndIndex > StartIndex) ? EndIndex - StartIndex : 0;
+	if (TextLength > 0)
 	{
-		return;
+		MakeText(ElementList, InLayer, PaintGeometry, InText.Mid(StartIndex, TextLength), InFontInfo, InDrawEffects, InTint);
 	}
-#if !UE_BUILD_SHIPPING
-	if (UE::Slate::Private::GLogPaintedText)
-	{
-		UE_LOG(LogSlate, Log, TEXT("MakeText: '%s'."), *InText);
-	}
-#endif
-	FSlateTextElement& Element = ElementList.AddUninitialized<EElementType::ET_Text>();
-
-	Element.SetTint(InTint);
-	Element.SetText(InText, InFontInfo, StartIndex, EndIndex);
-
-	Element.Init(ElementList, EElementType::ET_Text, InLayer, PaintGeometry, InDrawEffects);
 }
 
 void FSlateDrawElement::MakeText( FSlateWindowElementList& ElementList, uint32 InLayer, const FPaintGeometry& PaintGeometry, const FString& InText, const FSlateFontInfo& InFontInfo, ESlateDrawEffect InDrawEffects, const FLinearColor& InTint )
@@ -399,22 +388,11 @@ void FSlateDrawElement::MakeText( FSlateWindowElementList& ElementList, uint32 I
 	SCOPE_CYCLE_COUNTER(STAT_SlateDrawElementMakeTime)
 	PaintGeometry.CommitTransformsIfUsingLegacyConstructor();
 
-	// Don't try and render empty text
-	if (InText.Len() == 0)
-	{
-		return;
-	}
-
 	if (ShouldCull(ElementList, PaintGeometry, InTint, InText))
 	{
 		return;
 	}
 
-	// Don't do anything if there the font would be completely transparent 
-	if (InTint.A == 0 && !InFontInfo.OutlineSettings.IsVisible())
-	{
-		return;
-	}
 #if !UE_BUILD_SHIPPING
 	if (UE::Slate::Private::GLogPaintedText)
 	{
@@ -482,8 +460,8 @@ void FSlateDrawElement::MakeShapedText(FSlateWindowElementList& ElementList, uin
 	}
 
 	// Don't do anything if there the font would be completely transparent 
-	if ((BaseTint.A == 0 && InShapedGlyphSequence->GetFontOutlineSettings().OutlineSize == 0) || 
-		(BaseTint.A == 0 && OutlineTint.A == 0))
+	if ((BaseTint.A < UE_KINDA_SMALL_NUMBER && InShapedGlyphSequence->GetFontOutlineSettings().OutlineSize == 0) ||
+		(BaseTint.A < UE_KINDA_SMALL_NUMBER && OutlineTint.A < UE_KINDA_SMALL_NUMBER))
 	{
 		return;
 	}
@@ -496,7 +474,7 @@ void FSlateDrawElement::MakeShapedText(FSlateWindowElementList& ElementList, uin
 	FSlateShapedTextElement& Element = ElementList.AddUninitialized<EElementType::ET_ShapedText>();
 
 	Element.SetTint(BaseTint);
-	Element.SetShapedText(ElementList, InShapedGlyphSequence, OutlineTint);
+	Element.SetShapedText(InShapedGlyphSequence, OutlineTint);
 	Element.SetOverflowArgs(TextOverflowArgs);
 
 	Element.Init(ElementList, EElementType::ET_ShapedText, InLayer, PaintGeometry, InDrawEffects);
@@ -543,7 +521,7 @@ void FSlateDrawElement::MakeRotatedShapedText(
 	FSlateShapedTextElement& Element = ElementList.AddUninitialized<EElementType::ET_ShapedText>();
 	
 	Element.SetTint(BaseTint);
-	Element.SetShapedText(ElementList, InShapedGlyphSequence, OutlineTint);
+	Element.SetShapedText(InShapedGlyphSequence, OutlineTint);
 	Element.SetOverflowArgs(TextOverflowArgs);
 	
 	Element.Init(ElementList, EElementType::ET_ShapedText, InLayer, PaintGeometry, InDrawEffects);
diff --git a/Engine/Source/Runtime/SlateCore/Public/Rendering/DrawElementTypes.h b/Engine/Source/Runtime/SlateCore/Public/Rendering/DrawElementTypes.h
--- a/Engine/Source/Runtime/SlateCore/Public/Rendering/DrawElementTypes.h
+++ b/Engine/Source/Runtime/SlateCore/Public/Rendering/DrawElementTypes.h
@@ -506,7 +506,7 @@ struct FSlateShapedTextElement : public FSlateDrawElement, public FSlateTintable
 	const FShapedGlyphSequencePtr& GetShapedGlyphSequence() const { return ShapedGlyphSequence; }
 	FLinearColor GetOutlineTint() const { return OutlineTint; }
 
-	void SetShapedText(FSlateWindowElementList& ElementList, const FShapedGlyphSequencePtr& InShapedGlyphSequence, const FLinearColor& InOutlineTint)
+	void SetShapedText(const FShapedGlyphSequencePtr& InShapedGlyphSequence, const FLinearColor& InOutlineTint)
 	{
 		ShapedGlyphSequence = InShapedGlyphSequence;
 		OutlineTint = InOutlineTint;
