Fix of ConsoleSlateDebuggerPassBase Destructor

The destructor is currently calling a virtual function which crashes the engine when closing it while the debug view is on.

diff --git a/Engine/Source/Runtime/SlateCore/Private/Debugging/ConsoleSlateDebuggerPaint.cpp b/Engine/Source/Runtime/SlateCore/Private/Debugging/ConsoleSlateDebuggerPaint.cpp
--- a/Engine/Source/Runtime/SlateCore/Private/Debugging/ConsoleSlateDebuggerPaint.cpp
+++ b/Engine/Source/Runtime/SlateCore/Private/Debugging/ConsoleSlateDebuggerPaint.cpp
@@ -69,6 +69,12 @@ FConsoleSlateDebuggerPaint::FConsoleSlateDebuggerPaint()
 {
 }
 
+FConsoleSlateDebuggerPaint::~FConsoleSlateDebuggerPaint()
+{
+	// StopDebugging will end up calling the virtual function GetEnabledCVar()
+	StopDebugging();
+}
+
 void FConsoleSlateDebuggerPaint::LoadConfig()
 {
 	FConsoleSlateDebuggerPassBase::LoadConfig();
diff --git a/Engine/Source/Runtime/SlateCore/Private/Debugging/ConsoleSlateDebuggerPaint.h b/Engine/Source/Runtime/SlateCore/Private/Debugging/ConsoleSlateDebuggerPaint.h
--- a/Engine/Source/Runtime/SlateCore/Private/Debugging/ConsoleSlateDebuggerPaint.h
+++ b/Engine/Source/Runtime/SlateCore/Private/Debugging/ConsoleSlateDebuggerPaint.h
@@ -20,10 +20,11 @@
  *   Start - SlateDebugger.Paint.Start
  *   Stop  - SlateDebugger.Paint.Stop
  */
-class FConsoleSlateDebuggerPaint : public FConsoleSlateDebuggerPassBase
+class FConsoleSlateDebuggerPaint final : public FConsoleSlateDebuggerPassBase
 {
 public:
 	FConsoleSlateDebuggerPaint();
+	virtual ~FConsoleSlateDebuggerPaint() override ;
 
 	virtual void LoadConfig() override;
 	virtual void SaveConfig() override;
diff --git a/Engine/Source/Runtime/SlateCore/Private/Debugging/ConsoleSlateDebuggerPassBase.cpp b/Engine/Source/Runtime/SlateCore/Private/Debugging/ConsoleSlateDebuggerPassBase.cpp
--- a/Engine/Source/Runtime/SlateCore/Private/Debugging/ConsoleSlateDebuggerPassBase.cpp
+++ b/Engine/Source/Runtime/SlateCore/Private/Debugging/ConsoleSlateDebuggerPassBase.cpp
@@ -34,7 +34,6 @@ FConsoleSlateDebuggerPassBase::FConsoleSlateDebuggerPassBase()
 
 FConsoleSlateDebuggerPassBase::~FConsoleSlateDebuggerPassBase()
 {
-	StopDebugging();
 }
 
 void FConsoleSlateDebuggerPassBase::LoadConfig()
diff --git a/Engine/Source/Runtime/SlateCore/Private/Debugging/ConsoleSlateDebuggerPrepass.cpp b/Engine/Source/Runtime/SlateCore/Private/Debugging/ConsoleSlateDebuggerPrepass.cpp
--- a/Engine/Source/Runtime/SlateCore/Private/Debugging/ConsoleSlateDebuggerPrepass.cpp
+++ b/Engine/Source/Runtime/SlateCore/Private/Debugging/ConsoleSlateDebuggerPrepass.cpp
@@ -65,6 +65,12 @@ FConsoleSlateDebuggerPrepass::FConsoleSlateDebuggerPrepass()
 	WidgetLogLocation = FVector2f(10.f, 20.f); // Different from FConsoleSlateDebuggerPaint so the texts do not overlap
 }
 
+FConsoleSlateDebuggerPrepass::~FConsoleSlateDebuggerPrepass()
+{
+	// StopDebugging will end up calling the virtual function GetEnabledCVar()
+	StopDebugging();
+}
+
 void FConsoleSlateDebuggerPrepass::LoadConfig()
 {
 	FConsoleSlateDebuggerPassBase::LoadConfig();
diff --git a/Engine/Source/Runtime/SlateCore/Private/Debugging/ConsoleSlateDebuggerPrepass.h b/Engine/Source/Runtime/SlateCore/Private/Debugging/ConsoleSlateDebuggerPrepass.h
--- a/Engine/Source/Runtime/SlateCore/Private/Debugging/ConsoleSlateDebuggerPrepass.h
+++ b/Engine/Source/Runtime/SlateCore/Private/Debugging/ConsoleSlateDebuggerPrepass.h
@@ -20,10 +20,11 @@
  *   Start - SlateDebugger.Prepass.Start
  *   Stop  - SlateDebugger.Prepass.Stop
  */
-class FConsoleSlateDebuggerPrepass : public FConsoleSlateDebuggerPassBase
+class FConsoleSlateDebuggerPrepass final : public FConsoleSlateDebuggerPassBase
 {
 public:
 	FConsoleSlateDebuggerPrepass();
+	virtual ~FConsoleSlateDebuggerPrepass() override;
 
 	virtual void LoadConfig() override;
 
