Force SlateColorBrush constructors to be inlined to avoid MergeModules link time issues

#rnx

diff --git a/Engine/Source/Runtime/SlateCore/Public/Brushes/SlateColorBrush.h b/Engine/Source/Runtime/SlateCore/Public/Brushes/SlateColorBrush.h
--- a/Engine/Source/Runtime/SlateCore/Public/Brushes/SlateColorBrush.h
+++ b/Engine/Source/Runtime/SlateCore/Public/Brushes/SlateColorBrush.h
@@ -17,7 +17,7 @@ struct FSlateColorBrush
 	 *
 	 * @param InColor The linear color to use
 	 */
-	FORCENOINLINE FSlateColorBrush( const FLinearColor& InColor )
+	FORCEINLINE FSlateColorBrush( const FLinearColor& InColor )
 		: FSlateBrush(ESlateBrushDrawType::Image, NAME_None, FMargin(0.0f), ESlateBrushTileType::NoTile, ESlateBrushImageType::NoImage, FVector2f::ZeroVector, InColor)
 	{ }
 
@@ -26,7 +26,7 @@ struct FSlateColorBrush
 	 *
 	 * @param InColor The color value to use
 	 */
-	FORCENOINLINE FSlateColorBrush( const FColor& InColor )
+	FORCEINLINE FSlateColorBrush( const FColor& InColor )
 		: FSlateBrush(ESlateBrushDrawType::Image, NAME_None, FMargin(0.0f), ESlateBrushTileType::NoTile, ESlateBrushImageType::NoImage, FVector2f::ZeroVector, InColor.ReinterpretAsLinear())
 	{ }
 
@@ -35,7 +35,7 @@ struct FSlateColorBrush
 	 *
 	 * @param InColor The shared color to use
 	 */
-	FORCENOINLINE FSlateColorBrush( const TSharedRef< FLinearColor >& InColor )
+	FORCEINLINE FSlateColorBrush( const TSharedRef< FLinearColor >& InColor )
 		: FSlateBrush(ESlateBrushDrawType::Image, NAME_None, FMargin(0.0f), ESlateBrushTileType::NoTile, ESlateBrushImageType::NoImage, FVector2f::ZeroVector, InColor)
 	{ }
 
@@ -44,7 +44,7 @@ struct FSlateColorBrush
 	 *
 	 * @param InColor The shared color to use
 	 */
-	FORCENOINLINE FSlateColorBrush(const FSlateColor& InColor)
+	FORCEINLINE FSlateColorBrush(const FSlateColor& InColor)
 		: FSlateBrush(ESlateBrushDrawType::Image, NAME_None, FMargin(0.0f), ESlateBrushTileType::NoTile, ESlateBrushImageType::NoImage, FVector2f::ZeroVector, InColor)
 	{ }
-};
+};
\ No newline at end of file
