Slate: Guard mismatched buffer sizes in line element point and pointcolor buffers. ColorBuffer is reserved to the size of the point buffer, speculatively adding this guard to prevent reading uninitialized data.

#rnx

diff --git a/Engine/Source/Runtime/SlateCore/Private/Rendering/ElementBatcher.cpp b/Engine/Source/Runtime/SlateCore/Private/Rendering/ElementBatcher.cpp
--- a/Engine/Source/Runtime/SlateCore/Private/Rendering/ElementBatcher.cpp
+++ b/Engine/Source/Runtime/SlateCore/Private/Rendering/ElementBatcher.cpp
@@ -2289,7 +2289,7 @@ struct FLineBuilder
 			const FVector2f LastUp = Up;
 			const float LastLength = Length;
 
-			PointColor = PackedColors.Num() ? PackedColors[Point] : PackedTint;
+			PointColor = (Point < PackedColors.Num()) ? PackedColors[Point] : PackedTint;
 			Position = NextPosition;
 			NextPosition = Points[Point + 1];
 
@@ -2337,7 +2337,7 @@ struct FLineBuilder
 		}
 
 		// Build the last point's incoming segment and end cap
-		PointColor = PackedColors.Num() ? PackedColors[LastPointIndex] : PackedTint;
+		PointColor = (LastPointIndex < PackedColors.Num()) ? PackedColors[LastPointIndex] : PackedTint;
 		MakeEndCap(NextPosition, Direction, Length, Up, PointColor, Rounding);
 	}
 
