Add profiling point for slate attributes and some widgets.

#tests Used the profile point to profile with Unreal Insight.
#rnx

diff --git a/Engine/Source/Editor/SceneOutliner/Private/SSceneOutliner.cpp b/Engine/Source/Editor/SceneOutliner/Private/SSceneOutliner.cpp
--- a/Engine/Source/Editor/SceneOutliner/Private/SSceneOutliner.cpp
+++ b/Engine/Source/Editor/SceneOutliner/Private/SSceneOutliner.cpp
@@ -2480,6 +2480,8 @@ void SSceneOutliner::RequestSort()
 
 void SSceneOutliner::SortItems(TArray<FSceneOutlinerTreeItemPtr>& Items) const
 {
+	TRACE_CPUPROFILER_EVENT_SCOPE(SSceneOutliner::SortItems);
+
 	auto Column = Columns.FindRef(SortByColumn);
 	if (Column.IsValid())
 	{
diff --git a/Engine/Source/Runtime/Slate/Private/Widgets/Views/STableViewBase.cpp b/Engine/Source/Runtime/Slate/Private/Widgets/Views/STableViewBase.cpp
--- a/Engine/Source/Runtime/Slate/Private/Widgets/Views/STableViewBase.cpp
+++ b/Engine/Source/Runtime/Slate/Private/Widgets/Views/STableViewBase.cpp
@@ -363,6 +363,8 @@ EActiveTimerReturnType STableViewBase::EnsureTickToRefresh(double InCurrentTime,
 
 void STableViewBase::Tick( const FGeometry& AllottedGeometry, const double InCurrentTime, const float InDeltaTime )
 {
+	TRACE_CPUPROFILER_EVENT_SCOPE(STableViewBase::Tick);
+
 	if (ItemsPanel.IsValid())
 	{
 		FGeometry PanelGeometry = FindChildGeometry( AllottedGeometry, ItemsPanel.ToSharedRef() );
diff --git a/Engine/Source/Runtime/Slate/Public/Widgets/Views/STreeView.h b/Engine/Source/Runtime/Slate/Public/Widgets/Views/STreeView.h
--- a/Engine/Source/Runtime/Slate/Public/Widgets/Views/STreeView.h
+++ b/Engine/Source/Runtime/Slate/Public/Widgets/Views/STreeView.h
@@ -567,6 +567,8 @@ public:
 	 */
 	virtual void Tick( const FGeometry& AllottedGeometry, const double InCurrentTime, const float InDeltaTime ) override
 	{
+		TRACE_CPUPROFILER_EVENT_SCOPE(STreeView::Tick);
+
 		if ( bTreeItemsAreDirty )
 		{
 			// Check that ItemsPanel was made; we never make it if the user failed to specify all the parameters necessary to make the tree work.
@@ -657,6 +659,7 @@ public:
 		bool bAddingItems,
 		int32 ParentIndex)
 	{
+		TRACE_CPUPROFILER_EVENT_SCOPE(STreeView::PopulateLinearizedItems);
 
 		NeedsParentWire.Add(false);
 		const int32 NestingDepthIndex = NeedsParentWire.Num()-1;
@@ -668,7 +671,10 @@ public:
 
 			// Find this items children.
 			TArray<ItemType> ChildItems;
-			OnGetChildren.Execute(InItemsSource[ItemIndex], ChildItems );
+			{
+				TRACE_CPUPROFILER_EVENT_SCOPE(STreeView_GetChildren);
+				OnGetChildren.Execute(InItemsSource[ItemIndex], ChildItems);
+			}
 
 			const bool bHasChildren = ChildItems.Num() > 0;
 
diff --git a/Engine/Source/Runtime/SlateCore/Private/Types/SlateAttributeMetaData.cpp b/Engine/Source/Runtime/SlateCore/Private/Types/SlateAttributeMetaData.cpp
--- a/Engine/Source/Runtime/SlateCore/Private/Types/SlateAttributeMetaData.cpp
+++ b/Engine/Source/Runtime/SlateCore/Private/Types/SlateAttributeMetaData.cpp
@@ -4,11 +4,17 @@
 
 #include "Algo/BinarySearch.h"
 #include "Layout/Children.h"
+#include "Stats/Stats.h"
 #include "Types/ReflectionMetadata.h"
 #include "Widgets/SWidget.h"
 
 #include <limits>
 
+DECLARE_CYCLE_STAT(TEXT("UpdateVisibilityAttributes"), STAT_Slate_UpdateVisibilityAttributes, STATGROUP_Slate);
+DECLARE_CYCLE_STAT(TEXT("UpdateExceptVisibilityAttributes"), STAT_Slate_UpdateExceptVisibilityAttributes, STATGROUP_Slate);
+DECLARE_CYCLE_STAT(TEXT("UpdateChildrenOnlyVisibilityAttributes"), STAT_Slate_UpdateChildrenOnlyVisibilityAttributes, STATGROUP_Slate);
+DECLARE_CYCLE_STAT(TEXT("UpdateAllAttributes"), STAT_Slate_UpdateAllAttributes, STATGROUP_Slate);
+
 namespace UE
 {
 namespace Private
@@ -419,6 +425,8 @@ void FSlateAttributeMetaData::InvalidateWidget(ISlateAttributeContainer& OwningC
 
 void FSlateAttributeMetaData::UpdateAllAttributes(SWidget& OwningWidget, EInvalidationPermission InvalidationStyle)
 {
+	SCOPE_CYCLE_COUNTER(STAT_Slate_UpdateAllAttributes);
+
 	if (FSlateAttributeMetaData* AttributeMetaData = FSlateAttributeMetaData::FindMetaData(OwningWidget))
 	{
 		AttributeMetaData->UpdateAttributesImpl(OwningWidget, InvalidationStyle, 0, AttributeMetaData->Attributes.Num());
@@ -428,6 +436,8 @@ void FSlateAttributeMetaData::UpdateAllAttributes(SWidget& OwningWidget, EInvali
 
 void FSlateAttributeMetaData::UpdateOnlyVisibilityAttributes(SWidget& OwningWidget, EInvalidationPermission InvalidationStyle)
 {
+	SCOPE_CYCLE_COUNTER(STAT_Slate_UpdateVisibilityAttributes);
+
 	if (FSlateAttributeMetaData* AttributeMetaData = FSlateAttributeMetaData::FindMetaData(OwningWidget))
 	{
 		if (AttributeMetaData->AffectVisibilityCounter > 0)
@@ -442,6 +452,8 @@ void FSlateAttributeMetaData::UpdateOnlyVisibilityAttributes(SWidget& OwningWidg
 
 void FSlateAttributeMetaData::UpdateExceptVisibilityAttributes(SWidget& OwningWidget, EInvalidationPermission InvalidationStyle)
 {
+	SCOPE_CYCLE_COUNTER(STAT_Slate_UpdateExceptVisibilityAttributes);
+
 	if (FSlateAttributeMetaData* AttributeMetaData = FSlateAttributeMetaData::FindMetaData(OwningWidget))
 	{
 		if (AttributeMetaData->AffectVisibilityCounter < AttributeMetaData->Attributes.Num())
@@ -456,6 +468,8 @@ void FSlateAttributeMetaData::UpdateExceptVisibilityAttributes(SWidget& OwningWi
 
 void FSlateAttributeMetaData::UpdateChildrenOnlyVisibilityAttributes(SWidget& OwningWidget, EInvalidationPermission InvalidationStyle, bool bRecursive)
 {
+	SCOPE_CYCLE_COUNTER(STAT_Slate_UpdateChildrenOnlyVisibilityAttributes);
+
 	OwningWidget.GetChildren()->ForEachWidget([InvalidationStyle, bRecursive](SWidget& Child)
 		{
 			UpdateOnlyVisibilityAttributes(Child, InvalidationStyle);
