Slate: Add FSlateDrawElement::MakeOutline to draw an outline around a quad

Also fix the comment on MakeDebugQuad to specify "Solid" quad and clean up AddDebugQuadElement a bit.

#rnx

diff --git a/Engine/Source/Runtime/SlateCore/Private/Rendering/DrawElementTypes.cpp b/Engine/Source/Runtime/SlateCore/Private/Rendering/DrawElementTypes.cpp
--- a/Engine/Source/Runtime/SlateCore/Private/Rendering/DrawElementTypes.cpp
+++ b/Engine/Source/Runtime/SlateCore/Private/Rendering/DrawElementTypes.cpp
@@ -243,6 +243,23 @@ void FSlateDrawElement::MakeDebugQuad( FSlateWindowElementList& ElementList, uin
 	Element.Init(ElementList, EElementType::ET_DebugQuad, InLayer, PaintGeometry, ESlateDrawEffect::None);
 }
 
+void FSlateDrawElement::MakeOutline(FSlateWindowElementList& ElementList, uint32 InLayer, const FPaintGeometry& PaintGeometry, FLinearColor Tint, ESlateDrawEffect InDrawEffects, bool bAntialias, float Thickness)
+{
+	PaintGeometry.CommitTransformsIfUsingLegacyConstructor();
+
+	const FVector2f LocalSize = FVector2f(PaintGeometry.GetLocalSize());
+	TArray<FVector2f> OutlinePoints =
+	{
+		FVector2f::ZeroVector,
+		FVector2f(LocalSize.X, 0.f),
+		FVector2f(LocalSize.X, LocalSize.Y),
+		FVector2f(0.f, LocalSize.Y),
+		FVector2f::ZeroVector
+	};
+
+	MakeLines(ElementList, InLayer, PaintGeometry, MoveTemp(OutlinePoints), InDrawEffects, Tint, bAntialias, Thickness);
+}
+
 FSlateDrawElement& FSlateDrawElement::MakeBoxInternal(
 	FSlateWindowElementList& ElementList,
 	uint32 InLayer,
diff --git a/Engine/Source/Runtime/SlateCore/Private/Rendering/ElementBatcher.cpp b/Engine/Source/Runtime/SlateCore/Private/Rendering/ElementBatcher.cpp
--- a/Engine/Source/Runtime/SlateCore/Private/Rendering/ElementBatcher.cpp
+++ b/Engine/Source/Runtime/SlateCore/Private/Rendering/ElementBatcher.cpp
@@ -652,23 +652,17 @@ void FSlateElementBatcher::AddDebugQuadElement(const FSlateBoxElement& DrawEleme
 	const FColor Tint = PackVertexColor(DrawElement.GetTint());
 	const FSlateRenderTransform& RenderTransform = DrawElement.GetRenderTransform();
 	const FVector2f LocalSize = DrawElement.GetLocalSize();
-	ESlateDrawEffect InDrawEffects = DrawElement.GetDrawEffects();
+	const ESlateDrawEffect InDrawEffects = DrawElement.GetDrawEffects();
 	const int32 Layer = DrawElement.GetLayer();
 
-	FSlateRenderBatch& RenderBatch = CreateRenderBatch(Layer, FShaderParams(), nullptr, ESlateDrawPrimitive::TriangleList, ESlateShader::Default, ESlateDrawEffect::None, ESlateBatchDrawFlag::None, DrawElement);
+	FSlateRenderBatch& RenderBatch = CreateRenderBatch(Layer, FShaderParams(), nullptr, ESlateDrawPrimitive::TriangleList, ESlateShader::Default, InDrawEffects, ESlateBatchDrawFlag::None, DrawElement);
 	
-	const FColor Color = PackVertexColor(DrawElement.GetTint());
-
 	// Determine the four corners of the quad
 	FVector2f TopLeft = FVector2f::ZeroVector;
 	FVector2f TopRight = FVector2f(LocalSize.X, 0);
 	FVector2f BotLeft = FVector2f(0.f, LocalSize.Y);
 	FVector2f BotRight = FVector2f(LocalSize.X, LocalSize.Y);
 
-	// The start index of these vertices in the index buffer
-	//const uint32 IndexStart = BatchVertices.Num();
-	const uint32 IndexStart = 0;
-
 	// Add four vertices to the list of verts to be added to the vertex buffer
 	RenderBatch.AddVertex(FSlateVertex::Make<Rounding>(RenderTransform, FVector2f(TopLeft), FVector2f(0.0f,0.0f),  Tint));
 	RenderBatch.AddVertex(FSlateVertex::Make<Rounding>(RenderTransform, FVector2f(TopRight), FVector2f(1.0f,0.0f), Tint));
@@ -676,13 +670,13 @@ void FSlateElementBatcher::AddDebugQuadElement(const FSlateBoxElement& DrawEleme
 	RenderBatch.AddVertex(FSlateVertex::Make<Rounding>(RenderTransform, FVector2f(BotRight), FVector2f(1.0f,1.0f), Tint));
 
 	// Add 6 indices to the vertex buffer.  (2 tri's per quad, 3 indices per tri)
-	RenderBatch.AddIndex(IndexStart + 0);
-	RenderBatch.AddIndex(IndexStart + 1);
-	RenderBatch.AddIndex(IndexStart + 2);
+	RenderBatch.AddIndex(0);
+	RenderBatch.AddIndex(1);
+	RenderBatch.AddIndex(2);
 
-	RenderBatch.AddIndex(IndexStart + 2);
-	RenderBatch.AddIndex(IndexStart + 1);
-	RenderBatch.AddIndex(IndexStart + 3);
+	RenderBatch.AddIndex(2);
+	RenderBatch.AddIndex(1);
+	RenderBatch.AddIndex(3);
 }
 
 FORCEINLINE void IndexQuad(FSlateRenderBatch& RenderBatch, int32 TopLeft, int32 TopRight, int32 BottomRight, int32 BottomLeft)
diff --git a/Engine/Source/Runtime/SlateCore/Public/Rendering/DrawElementTypes.h b/Engine/Source/Runtime/SlateCore/Public/Rendering/DrawElementTypes.h
--- a/Engine/Source/Runtime/SlateCore/Public/Rendering/DrawElementTypes.h
+++ b/Engine/Source/Runtime/SlateCore/Public/Rendering/DrawElementTypes.h
@@ -56,13 +56,22 @@ public:
 	};
 
 	/**
-	 * Creates a wireframe quad for debug purposes
+	 * Creates a solid quad for debug purposes
 	 *
 	 * @param ElementList			The list in which to add elements
 	 * @param InLayer				The layer to draw the element on
 	 * @param PaintGeometry         DrawSpace position and dimensions; see FPaintGeometry
 	 */
-	SLATECORE_API static void MakeDebugQuad( FSlateWindowElementList& ElementList, uint32 InLayer, const FPaintGeometry& PaintGeometry, FLinearColor Tint = FLinearColor::White);
+	SLATECORE_API static void MakeDebugQuad(FSlateWindowElementList& ElementList, uint32 InLayer, const FPaintGeometry& PaintGeometry, FLinearColor Tint = FLinearColor::White);
+
+	/**
+	 * Creates an outline of a quad
+	 *
+	 * @param ElementList			The list in which to add elements
+	 * @param InLayer				The layer to draw the element on
+	 * @param PaintGeometry         DrawSpace position and dimensions; see FPaintGeometry
+	 */
+	SLATECORE_API static void MakeOutline(FSlateWindowElementList& ElementList, uint32 InLayer, const FPaintGeometry& PaintGeometry, FLinearColor Tint = FLinearColor::White, ESlateDrawEffect InDrawEffects = ESlateDrawEffect::None, bool bAntialias = true, float Thickness = 1.0f);
 
 	/**
 	 * Creates a box element based on the following diagram.  Allows for this element to be resized while maintain the border of the image
