Slate: Rename MakeOutline to MakeGeometryOutline

#rnx

diff --git a/Engine/Source/Runtime/SlateCore/Private/Rendering/DrawElementTypes.cpp b/Engine/Source/Runtime/SlateCore/Private/Rendering/DrawElementTypes.cpp
--- a/Engine/Source/Runtime/SlateCore/Private/Rendering/DrawElementTypes.cpp
+++ b/Engine/Source/Runtime/SlateCore/Private/Rendering/DrawElementTypes.cpp
@@ -243,7 +243,7 @@ void FSlateDrawElement::MakeDebugQuad( FSlateWindowElementList& ElementList, uin
 	Element.Init(ElementList, EElementType::ET_DebugQuad, InLayer, PaintGeometry, ESlateDrawEffect::None);
 }
 
-void FSlateDrawElement::MakeOutline(FSlateWindowElementList& ElementList, uint32 InLayer, const FPaintGeometry& PaintGeometry, FLinearColor Tint, ESlateDrawEffect InDrawEffects, bool bAntialias, float Thickness)
+void FSlateDrawElement::MakeGeometryOutline(FSlateWindowElementList& ElementList, uint32 InLayer, const FPaintGeometry& PaintGeometry, FLinearColor Tint, ESlateDrawEffect InDrawEffects, bool bAntialias, float Thickness)
 {
 	PaintGeometry.CommitTransformsIfUsingLegacyConstructor();
 
diff --git a/Engine/Source/Runtime/SlateCore/Public/Rendering/DrawElementTypes.h b/Engine/Source/Runtime/SlateCore/Public/Rendering/DrawElementTypes.h
--- a/Engine/Source/Runtime/SlateCore/Public/Rendering/DrawElementTypes.h
+++ b/Engine/Source/Runtime/SlateCore/Public/Rendering/DrawElementTypes.h
@@ -65,13 +65,13 @@ public:
 	SLATECORE_API static void MakeDebugQuad(FSlateWindowElementList& ElementList, uint32 InLayer, const FPaintGeometry& PaintGeometry, FLinearColor Tint = FLinearColor::White);
 
 	/**
-	 * Creates an outline of a quad
+	 * Creates an outline of the specific Paint Geometry
 	 *
 	 * @param ElementList			The list in which to add elements
 	 * @param InLayer				The layer to draw the element on
 	 * @param PaintGeometry         DrawSpace position and dimensions; see FPaintGeometry
 	 */
-	SLATECORE_API static void MakeOutline(FSlateWindowElementList& ElementList, uint32 InLayer, const FPaintGeometry& PaintGeometry, FLinearColor Tint = FLinearColor::White, ESlateDrawEffect InDrawEffects = ESlateDrawEffect::None, bool bAntialias = true, float Thickness = 1.0f);
+	SLATECORE_API static void MakeGeometryOutline(FSlateWindowElementList& ElementList, uint32 InLayer, const FPaintGeometry& PaintGeometry, FLinearColor Tint = FLinearColor::White, ESlateDrawEffect InDrawEffects = ESlateDrawEffect::None, bool bAntialias = true, float Thickness = 1.0f);
 
 	/**
 	 * Creates a box element based on the following diagram.  Allows for this element to be resized while maintain the border of the image
