Slate: Add missing param doc comments for MakeDebugQuad and MakeGeometryOutline

#rnx

diff --git a/Engine/Source/Runtime/SlateCore/Public/Rendering/DrawElementTypes.h b/Engine/Source/Runtime/SlateCore/Public/Rendering/DrawElementTypes.h
--- a/Engine/Source/Runtime/SlateCore/Public/Rendering/DrawElementTypes.h
+++ b/Engine/Source/Runtime/SlateCore/Public/Rendering/DrawElementTypes.h
@@ -61,6 +61,7 @@ public:
 	 * @param ElementList			The list in which to add elements
 	 * @param InLayer				The layer to draw the element on
 	 * @param PaintGeometry         DrawSpace position and dimensions; see FPaintGeometry
+	 * @param Tint					The color to tint the quad
 	 */
 	SLATECORE_API static void MakeDebugQuad(FSlateWindowElementList& ElementList, uint32 InLayer, const FPaintGeometry& PaintGeometry, FLinearColor Tint = FLinearColor::White);
 
@@ -70,6 +71,10 @@ public:
 	 * @param ElementList			The list in which to add elements
 	 * @param InLayer				The layer to draw the element on
 	 * @param PaintGeometry         DrawSpace position and dimensions; see FPaintGeometry
+	 * @param Tint					The color to tint the outline
+	 * @param InDrawEffects			Which effects to apply to the outline
+	 * @param bAntialias			Whether to anti-alias the outline
+	 * @param Thickness				How thick the outline should be
 	 */
 	SLATECORE_API static void MakeGeometryOutline(FSlateWindowElementList& ElementList, uint32 InLayer, const FPaintGeometry& PaintGeometry, FLinearColor Tint = FLinearColor::White, ESlateDrawEffect InDrawEffects = ESlateDrawEffect::None, bool bAntialias = true, float Thickness = 1.0f);
 
