Stop editor from exiting when viewport elements have a destroyed slate resource

UE-296277

#rnx
#tests Reproed the case without my change, and confirmed it works with it.

diff --git a/Engine/Source/Runtime/Slate/Private/Widgets/SViewport.cpp b/Engine/Source/Runtime/Slate/Private/Widgets/SViewport.cpp
--- a/Engine/Source/Runtime/Slate/Private/Widgets/SViewport.cpp
+++ b/Engine/Source/Runtime/Slate/Private/Widgets/SViewport.cpp
@@ -158,7 +158,7 @@ int32 SViewport::OnPaint( const FPaintArgs& Args, const FGeometry& AllottedGeome
 	// Only draw a quad if not rendering directly to the backbuffer
 	if( !ShouldRenderDirectly() )
 	{
-		if( ViewportInterfacePin.IsValid() && ViewportInterfacePin->GetViewportRenderTargetTexture() != nullptr )
+		if( ViewportInterfacePin.IsValid() && ViewportInterfacePin->GetViewportRenderTargetTexture() != nullptr && !ViewportInterfacePin->GetViewportRenderTargetTexture()->Debug_IsDestroyed())
 		{
 			FSlateDrawElement::MakeViewport( OutDrawElements, LayerId, AllottedGeometry.ToPaintGeometry(), ViewportInterfacePin, DrawEffects, InWidgetStyle.GetColorAndOpacityTint() );
 		}
diff --git a/Engine/Source/Runtime/SlateCore/Private/Rendering/ElementBatcher.cpp b/Engine/Source/Runtime/SlateCore/Private/Rendering/ElementBatcher.cpp
--- a/Engine/Source/Runtime/SlateCore/Private/Rendering/ElementBatcher.cpp
+++ b/Engine/Source/Runtime/SlateCore/Private/Rendering/ElementBatcher.cpp
@@ -727,6 +727,9 @@ void FSlateElementBatcher::AddBoxElements(const FSlateDrawElementArray<ElementTy
 		{
 			// The actual texture for rendering.  If the texture is atlased this is the atlas
 			Resource = ResourceProxy->Resource;
+		}
+		if(Resource != nullptr && !Resource->Debug_IsDestroyed())
+		{
 			// The width and height of the texture (non-atlased size)
 			TextureWidth = ResourceProxy->ActualSize.X != 0 ? ResourceProxy->ActualSize.X : 1;
 			TextureHeight = ResourceProxy->ActualSize.Y != 0 ? ResourceProxy->ActualSize.Y : 1;
@@ -2718,6 +2721,12 @@ void FSlateElementBatcher::AddViewportElement( const FSlateViewportElement& Draw
 	FSlateShaderResource* ViewportResource = DrawElement.RenderTargetResource;
 	ESlateShader ShaderType = ESlateShader::Default;
 
+	if (ViewportResource != nullptr && ViewportResource->Debug_IsDestroyed())
+	{
+		UE_LOG(LogSlate, Error, TEXT("Adding a viewport element with an already destroyed viewport resource."));
+		return;
+	}
+
 	if(DrawElement.bViewportTextureAlphaOnly )
 	{
 		// This is a slight hack, but the grayscale font shader is the same as the general shader except it reads alpha only textures and doesn't support tiling
