SWindow: Fix Accepts Input comment

#rnx

diff --git a/Engine/Source/Runtime/SlateCore/Public/Widgets/SWindow.h b/Engine/Source/Runtime/SlateCore/Public/Widgets/SWindow.h
--- a/Engine/Source/Runtime/SlateCore/Public/Widgets/SWindow.h
+++ b/Engine/Source/Runtime/SlateCore/Public/Widgets/SWindow.h
@@ -803,7 +803,7 @@ public:
 	/** @return true if the window accepts input; false if the window is non-interactive */
 	SLATECORE_API bool AcceptsInput() const;
 
-	/** If true the window will not accept input regardless of its type*/
+	/** If true the window will accept input if window type supports it*/
 	SLATECORE_API void SetAcceptsInput(bool bEnabled) { bAcceptsInput = bEnabled; }
 
 	/** @return true if the user decides the size of the window */
