Slate Dynamic Invalidation - Fix for `ChildOrder` invalidation not being processed for all the widgets

Currently, the value of `SWidget::SupportsInvalidationRecursive` is cached and it is updated in `FSlateInvalidationWidgetList::ProcessChildOrderInvalidation` when a `ChildOrder` invalidation is raised for this particular widget. The issue is that when multiple widgets in the same hierarchy have a `ChildOrder` invalidation raised in the same frame, only the parent widget is processed in `ProcessChildOrderInvalidation` as the invalidation list is emptied and reconstructed, losing all update flags of descendants.

This happens with a `STreeView` which contains a `SListPanel` for the pinned row widgets. When a pinned row is added, a `ChildOrder` invalidation is raised for both the `SListPanel` and the `STreeView`. Both are present in the list of widgets to be updated in `FSlateInvalidationRoot::ProcessPreUpdate` but the `STreeView` is being processed first with `ProcessChildOrderInvalidation`. This function ends up removing all the TreeView children from the invalidation list - including the `SListPanel` - before rebuilding the list. After the list is rebuilt though, the `SListPanel` is not marked anymore as having any invalidation reason, so `ProcessChildOrderInvalidation` is not called for this widget. This means `SupportsInvalidationRecursive` is not being updated for that particular widget.

To avoid this, all widgets having a `ChildOrder` have their `SupportsInvalidationRecursive` updated first in `FSlateInvalidationRoot::ProcessPreUpdate` before any widget is processed with `ProcessChildOrderInvalidation`.

diff --git a/Engine/Source/Runtime/SlateCore/Private/FastUpdate/SlateInvalidationRoot.cpp b/Engine/Source/Runtime/SlateCore/Private/FastUpdate/SlateInvalidationRoot.cpp
--- a/Engine/Source/Runtime/SlateCore/Private/FastUpdate/SlateInvalidationRoot.cpp
+++ b/Engine/Source/Runtime/SlateCore/Private/FastUpdate/SlateInvalidationRoot.cpp
@@ -877,6 +877,23 @@ void FSlateInvalidationRoot::ProcessPreUpdate()
 
 	TGuardValue<bool> Tmp(bProcessingPreUpdate, true);
 
+#if UE_SLATE_WITH_DYNAMIC_INVALIDATION
+	// We need to process the list first as the code below might recreate the list and lose the CurrentInvalidateReason
+	WidgetsNeedingPreUpdate->ForEachIndexes([this](FSlateInvalidationWidgetHeapElement& Element)
+	{
+		const FSlateInvalidationWidgetList::InvalidationWidgetType& InvalidationWidget = (*FastWidgetPathList)[Element.GetWidgetIndex()];
+		const SWidget* WidgetPtr = InvalidationWidget.GetWidget();
+		const EInvalidateWidgetReason CurrentInvalidateReason = InvalidationWidget.CurrentInvalidateReason;
+		if (WidgetPtr && EnumHasAnyFlags(CurrentInvalidateReason, EInvalidateWidgetReason::ChildOrder) )
+		{
+			// If the children have been modified in any way, the cached value of SupportsInvalidationRecursive might have changed.
+			// We call SupportsInvalidationRecursive to recompute the value and invalidate the parent's value if needed
+			const bool bUseCachedValue = false;
+			bool bNowSupportsInvalidation = WidgetPtr->SupportsInvalidationRecursive(bUseCachedValue);
+		}
+	});
+#endif
+
 	TSharedRef<SWidget> RootWidget = GetRootWidget();
 	if (FastWidgetPathList->GetRoot().Pin() != RootWidget)
 	{
diff --git a/Engine/Source/Runtime/SlateCore/Private/FastUpdate/SlateInvalidationWidgetList.cpp b/Engine/Source/Runtime/SlateCore/Private/FastUpdate/SlateInvalidationWidgetList.cpp
--- a/Engine/Source/Runtime/SlateCore/Private/FastUpdate/SlateInvalidationWidgetList.cpp
+++ b/Engine/Source/Runtime/SlateCore/Private/FastUpdate/SlateInvalidationWidgetList.cpp
@@ -717,13 +717,6 @@ bool FSlateInvalidationWidgetList::ProcessChildOrderInvalidation(FSlateInvalidat
 		}
 	}
 
-#if UE_SLATE_WITH_DYNAMIC_INVALIDATION
-	// If the children have been modified in any way, the cached value of SupportsInvalidationRecursive might have changed.
-	// We call SupportsInvalidationRecursive to recompute the value and invalidate the parent's value if needed
-	const bool bUseCachedValue = false;
-	bool bNowSupportsInvalidation = WidgetPtr->SupportsInvalidationRecursive(bUseCachedValue);
-#endif
-
 	return bIsInvalidationWidgetStillValid;
 }
 
