Fixed layout centering code centering over the available size plus the margin rather than just the available size. The previous behavior broke when the margins were uneven.

diff --git a/Engine/Source/Runtime/SlateCore/Public/Layout/LayoutUtils.h b/Engine/Source/Runtime/SlateCore/Public/Layout/LayoutUtils.h
--- a/Engine/Source/Runtime/SlateCore/Public/Layout/LayoutUtils.h
+++ b/Engine/Source/Runtime/SlateCore/Public/Layout/LayoutUtils.h
@@ -713,7 +713,7 @@ static AlignmentArrangeResult AlignChild(EFlowDirection InLayoutFlow, float Allo
 	case HAlign_Left: // same as Align_Top
 		return AlignmentArrangeResult(MarginPre, ChildSize);
 	case HAlign_Center:
-		return AlignmentArrangeResult(( AllottedSize - ChildSize ) / 2.0f + MarginPre - MarginPost, ChildSize);
+		return AlignmentArrangeResult(MarginPre + (FMath::Max((AllottedSize - TotalMargin - ChildSize), 0.f) / 2.0f), ChildSize);
 	case HAlign_Right: // same as Align_Bottom		
 		return AlignmentArrangeResult(AllottedSize - ChildSize - MarginPost, ChildSize);
 	}
