Slate Dynamic Invalidation - SLayeredImage

diff --git a/Engine/Source/Runtime/SlateCore/Private/Widgets/Images/SLayeredImage.cpp b/Engine/Source/Runtime/SlateCore/Private/Widgets/Images/SLayeredImage.cpp
--- a/Engine/Source/Runtime/SlateCore/Private/Widgets/Images/SLayeredImage.cpp
+++ b/Engine/Source/Runtime/SlateCore/Private/Widgets/Images/SLayeredImage.cpp
@@ -84,6 +84,8 @@ void SLayeredImage::AddLayer(TAttribute<const FSlateBrush*> Brush)
 	BrushAttributeType TmpBrush = BrushAttributeType(AsShared(), MoveTemp(Brush), nullptr);
 	ColorAttributeType TmpColor = ColorAttributeType(AsShared(), FLinearColor::White);
 	Layers.Emplace(MoveTemp(TmpBrush), MoveTemp(TmpColor));
+
+	Invalidate(EInvalidateWidgetReason::Paint);
 }
 
 void SLayeredImage::AddLayer(TAttribute<const FSlateBrush*> Brush, TAttribute<FSlateColor> Color)
@@ -91,6 +93,8 @@ void SLayeredImage::AddLayer(TAttribute<const FSlateBrush*> Brush, TAttribute<FS
 	BrushAttributeType TmpBrush = BrushAttributeType(AsShared(), MoveTemp(Brush), nullptr);
 	ColorAttributeType TmpColor = ColorAttributeType(AsShared(), MoveTemp(Color), FLinearColor::White);
 	Layers.Emplace(MoveTemp(TmpBrush), MoveTemp(TmpColor));
+
+	Invalidate(EInvalidateWidgetReason::Paint);
 }
 
 void SLayeredImage::SetFromSlateIcon(const FSlateIcon& InIcon)
@@ -104,11 +108,15 @@ void SLayeredImage::SetFromSlateIcon(const FSlateIcon& InIcon)
 	{
 		AddLayer(OverlayIcon);
 	}
+
+	Invalidate(EInvalidateWidgetReason::Paint);
 }
 
 void SLayeredImage::RemoveAllLayers()
 {
 	Layers.Empty();
+
+	Invalidate(EInvalidateWidgetReason::Paint);
 }
 
 int32 SLayeredImage::GetNumLayers() const
@@ -135,7 +143,10 @@ void SLayeredImage::SetLayerBrush(int32 Index, TAttribute<const FSlateBrush*> Br
 	else
 	{
 		// That layer doesn't exist
+		return;
 	}
+
+	Invalidate(EInvalidateWidgetReason::Paint);
 }
 
 void SLayeredImage::SetLayerColor(int32 Index, TAttribute<FSlateColor> Color)
@@ -151,6 +162,9 @@ void SLayeredImage::SetLayerColor(int32 Index, TAttribute<FSlateColor> Color)
 	else
 	{
 		// That layer doesn't exist!
+		return;
 	}
+
+	Invalidate(EInvalidateWidgetReason::Paint);
 }
 
diff --git a/Engine/Source/Runtime/SlateCore/Public/Widgets/Images/SLayeredImage.h b/Engine/Source/Runtime/SlateCore/Public/Widgets/Images/SLayeredImage.h
--- a/Engine/Source/Runtime/SlateCore/Public/Widgets/Images/SLayeredImage.h
+++ b/Engine/Source/Runtime/SlateCore/Public/Widgets/Images/SLayeredImage.h
@@ -7,6 +7,12 @@
 
 struct FSlateIcon;
 
+template <>
+struct TWidgetTypeTraits<class SLayeredImage>
+{
+	static constexpr bool SupportsInvalidation() { return true; }
+};
+
 /** A slate widget to draw an arbitrary number of images over top of each other */
 class SLayeredImage : public SImage
 {
