Add missing validation check with SWindow.
Fixes crash with upgrading recordings in LiveLinkHub.

diff --git a/Engine/Source/Runtime/SlateCore/Private/Widgets/SWindow.cpp b/Engine/Source/Runtime/SlateCore/Private/Widgets/SWindow.cpp
--- a/Engine/Source/Runtime/SlateCore/Private/Widgets/SWindow.cpp
+++ b/Engine/Source/Runtime/SlateCore/Private/Widgets/SWindow.cpp
@@ -1473,8 +1473,11 @@ void SWindow::HideWindow()
 
 void SWindow::EnableWindow( bool bEnable )
 {
-	NativeWindow->Enable( bEnable );
-
+	if (NativeWindow.IsValid())
+	{
+		NativeWindow->Enable( bEnable );
+	}
+	
 	for( int32 ChildIndex = 0; ChildIndex < ChildWindows.Num(); ++ChildIndex )
 	{
 		ChildWindows[ChildIndex]->EnableWindow( bEnable );
