Updated comments about dynamic Slate material shader variant. No functional code changes.

diff --git a/Engine/Shaders/Private/SlateShaderCommon.ush b/Engine/Shaders/Private/SlateShaderCommon.ush
--- a/Engine/Shaders/Private/SlateShaderCommon.ush
+++ b/Engine/Shaders/Private/SlateShaderCommon.ush
@@ -11,7 +11,7 @@
 #define ST_RoundedBox		7
 #define ST_SdfFont			8
 #define ST_MsdfFont			9
-#define ST_Dynamic			10 // This variant will be the only path, replacing most of the above after performance and correctness testing (#jira UE-319026)
+#define ST_Dynamic			10 // Special variant that uses dynamic branches for all of the above shader types
 
 #if USE_MATERIALS
 #include "/Engine/Generated/Material.ush"
diff --git a/Engine/Source/Runtime/SlateCore/Public/Rendering/RenderingCommon.h b/Engine/Source/Runtime/SlateCore/Public/Rendering/RenderingCommon.h
--- a/Engine/Source/Runtime/SlateCore/Public/Rendering/RenderingCommon.h
+++ b/Engine/Source/Runtime/SlateCore/Public/Rendering/RenderingCommon.h
@@ -78,11 +78,7 @@ enum class ESlateShader : uint8
 	SdfFont = 8,
 	/** Multi-channel signed distance field font shader */
 	MsdfFont = 9,
-	/** 
-	* Dynamic branching version.
-	* This variant will be the only path, replacing most of the above after performance and correctness testing.
-	* #jira UE-319026
-	*/
+	/** Dynamic branching version of all the shader types above */
 	Dynamic = 10,
 };
 
