Add virtual to FDragDropOperation (and its MACRO) to allow retrieving TypeString from object instance (in addition to existing static)

diff --git a/Engine/Source/Runtime/SlateCore/Public/Input/DragAndDrop.h b/Engine/Source/Runtime/SlateCore/Public/Input/DragAndDrop.h
--- a/Engine/Source/Runtime/SlateCore/Public/Input/DragAndDrop.h
+++ b/Engine/Source/Runtime/SlateCore/Public/Input/DragAndDrop.h
@@ -95,6 +95,11 @@ public:
 	{
 		return false;
 	}
+	 
+	virtual FStringView GetTypeIdString() const
+	{
+		return FStringView();
+	}
 protected:
 
 	/**
@@ -206,6 +211,7 @@ DECLARE_DELEGATE_OneParam( FOnDragDropUpdate,
 
 #define DRAG_DROP_OPERATOR_TYPE(TYPE, BASE) \
 	static const FString& GetTypeId() { static FString Type = TEXT(#TYPE); return Type; } \
+	virtual FStringView GetTypeIdString() const override { return GetTypeId(); } \
 	virtual bool IsOfTypeImpl(const FString& Type) const override { return GetTypeId() == Type || BASE::IsOfTypeImpl(Type); }
 
 
