Provide widget arguments with an `UpdateAttributes` argument which will automatically update all `TSlateAttributes` after the widget is constructed. This avoids `TSlateAttributes` using delegates not having their value updated until the invalidation root processes them.

diff --git a/Engine/Source/Runtime/Slate/Private/Framework/Docking/SPanelDrawerArea.cpp b/Engine/Source/Runtime/Slate/Private/Framework/Docking/SPanelDrawerArea.cpp
--- a/Engine/Source/Runtime/Slate/Private/Framework/Docking/SPanelDrawerArea.cpp
+++ b/Engine/Source/Runtime/Slate/Private/Framework/Docking/SPanelDrawerArea.cpp
@@ -237,15 +237,14 @@ void SPanelDrawerArea::SetupClosedLayout()
 
 void SPanelDrawerArea::SetupOpenedLayout()
 {
-	TSharedPtr<SSplitter> Splitter;
-
 	ChildSlot
 	[
-		SAssignNew(Splitter, SSplitter)
+		SNew(SSplitter)
 		.Orientation(Orient_Horizontal)
+		.UpdateAttributes(ESlateAttributeInitialization::AfterConstruct)
 		+ SSplitter::Slot()
 		.OnSlotResized(this, &SPanelDrawerArea::OnMainContentSlotResized)
-		.Value(SPanelDrawerArea::GetMainContentCoefficient())
+		.Value(this, &SPanelDrawerArea::GetMainContentCoefficient)
 		[
 			// Add Clipping for content that isn't responsive to dynamic sizes
 			SNew(SBox)
@@ -258,7 +257,7 @@ void SPanelDrawerArea::SetupOpenedLayout()
 		]
 		+ SSplitter::Slot()
 		.OnSlotResized(this, &SPanelDrawerArea::OnPanelDrawerSlotResized)
-		.Value(SPanelDrawerArea::GetPanelDrawerCoefficient())
+		.Value(this, &SPanelDrawerArea::GetPanelDrawerCoefficient)
 		[
 			// Add Clipping for content that isn't responsive to dynamic sizes
 			SNew(SBox)
@@ -272,10 +271,6 @@ void SPanelDrawerArea::SetupOpenedLayout()
 		]
 	];
 
-	// Temporary fix as values are not set by default with TSlateAttributes
-	Splitter->SlotAt(0).SetSizeValue(TAttribute<float>::CreateRaw(this, &SPanelDrawerArea::GetMainContentCoefficient));
-	Splitter->SlotAt(1).SetSizeValue(TAttribute<float>::CreateRaw(this, &SPanelDrawerArea::GetPanelDrawerCoefficient));
-
 	if (CVarNoAnimationOnTabForegroundedEvent.GetValueOnGameThread())
 	{
 		if (PanelDrawerData && PanelDrawerData->HostedTab)
diff --git a/Engine/Source/Runtime/SlateCore/Public/Types/SlateAttribute.h b/Engine/Source/Runtime/SlateCore/Public/Types/SlateAttribute.h
--- a/Engine/Source/Runtime/SlateCore/Public/Types/SlateAttribute.h
+++ b/Engine/Source/Runtime/SlateCore/Public/Types/SlateAttribute.h
@@ -176,6 +176,12 @@ ENUM_CLASS_FLAGS(ESlateAttributeBindAction)
 #	define UE_SLATE_WITH_ATTRIBUTE_DEFAULT_INITIALIZATION_ACTION ESlateAttributeBindAction::None
 #endif
 
+enum class ESlateAttributeInitialization
+{
+	None = 0, // Do not initialize Slate Attributes values
+	AfterConstruct, // Initialize all Slate Attributes values after Construct
+};
+
 /** Base struct of all SlateAttribute type. */
 struct FSlateAttributeBase
 {
diff --git a/Engine/Source/Runtime/SlateCore/Public/Widgets/DeclarativeSyntaxSupport.h b/Engine/Source/Runtime/SlateCore/Public/Widgets/DeclarativeSyntaxSupport.h
--- a/Engine/Source/Runtime/SlateCore/Public/Widgets/DeclarativeSyntaxSupport.h
+++ b/Engine/Source/Runtime/SlateCore/Public/Widgets/DeclarativeSyntaxSupport.h
@@ -665,6 +665,7 @@ struct FSlateBaseNamedArgs
 	SLATE_PRIVATE_ARGUMENT_VARIABLE(bool, ForceVolatile) = false;
 	/** If true, bound Slate Attributes will be updated once per frame. */
 	SLATE_PRIVATE_ARGUMENT_VARIABLE(bool, EnabledAttributesUpdate) = true;
+	SLATE_PRIVATE_ARGUMENT_VARIABLE(ESlateAttributeInitialization, UpdateAttributes) = ESlateAttributeInitialization::None;
 	SLATE_PRIVATE_ARGUMENT_VARIABLE(EWidgetClipping, Clipping) = EWidgetClipping::Inherit;
 	SLATE_PRIVATE_ARGUMENT_VARIABLE(EWidgetPixelSnapping, PixelSnappingMethod) = EWidgetPixelSnapping::Inherit;
 	SLATE_PRIVATE_ARGUMENT_VARIABLE(EFlowDirectionPreference, FlowDirectionPreference) = EFlowDirectionPreference::Inherit;
@@ -691,6 +692,7 @@ struct TSlateBaseNamedArgs : public FSlateBaseNamedArgs
 	SLATE_PRIVATE_ATTRIBUTE_FUNCTION(bool, IsEnabled)
 	SLATE_PRIVATE_ATTRIBUTE_FUNCTION(EVisibility, Visibility)
 	SLATE_PRIVATE_ARGUMENT_FUNCTION(bool, ForceVolatile)
+	SLATE_PRIVATE_ARGUMENT_FUNCTION(ESlateAttributeInitialization, UpdateAttributes)
 	SLATE_PRIVATE_ARGUMENT_FUNCTION(EWidgetClipping, Clipping)
 	SLATE_PRIVATE_ARGUMENT_FUNCTION(EWidgetPixelSnapping, PixelSnappingMethod)
 	SLATE_PRIVATE_ARGUMENT_FUNCTION(EFlowDirectionPreference, FlowDirectionPreference)
@@ -978,6 +980,10 @@ struct TSlateDecl
 	{
 		_Widget->SWidgetConstruct(InArgs);
 		_RequiredArgs.CallConstruct(_Widget.Get(), InArgs);
+		if (InArgs._UpdateAttributes == ESlateAttributeInitialization::AfterConstruct)
+		{
+			_Widget->UpdateAllAttributes();
+		}
 		_Widget->CacheVolatility();
 		_Widget->bIsDeclarativeSyntaxConstructionCompleted = true;
 
diff --git a/Engine/Source/Runtime/SlateCore/Public/Widgets/SWidget.h b/Engine/Source/Runtime/SlateCore/Public/Widgets/SWidget.h
--- a/Engine/Source/Runtime/SlateCore/Public/Widgets/SWidget.h
+++ b/Engine/Source/Runtime/SlateCore/Public/Widgets/SWidget.h
@@ -33,6 +33,7 @@
 #include "Widgets/SlateControlledConstruction.h"
 #include "Widgets/Accessibility/SlateWidgetAccessibleTypes.h"
 #include "WidgetPixelSnapping.h"
+#include "Types/SlateAttributeMetaData.h"
 
 class FActiveTimerHandle;
 class FArrangedChildren;
@@ -676,6 +677,12 @@ public:
 	/** @return true if the widgets will update its registered slate attributes automatically or they need to be updated manually. */
 	bool IsAttributesUpdatesEnabled() const { return bEnabledAttributesUpdate; }
 
+	/** Update all the Slate Attributes of this widget. */
+	void UpdateAllAttributes()
+	{
+		FSlateAttributeMetaData::UpdateAllAttributes(*this, FSlateAttributeMetaData::EInvalidationPermission::AllowInvalidation);
+	}
+
 	const FSlateWidgetPersistentState& GetPersistentState() const { return PersistentState; }
 	const FWidgetProxyHandle GetProxyHandle() const { return FastPathProxyHandle; }
 
