Fonts:
Fleshed out the async font loading feature to give designers more control and to give a cleaner presentation while font faces are loading when painting text widgets.
This feature is still experimental and requires further battle testing. With this feature active, the following UMG widgets (and their Slate and CommonUI counterparts) support this feature:
UTextBlock, URichTextBlock, UMultiLineText, UMultiLineTextBox, UEditableText and UEditableTextBox. In the following descriptions when "text widgets" are mentioned, these are the text widgets in question.
- Prevent font faces from async loading in the Editor leading to strange artifacts and impacting other text measurement calculations.
- Introduced UI.Slate.Font.SimulateCookedFontFilesInStandalone cvar which if set to true will write all the embedded font face data in UFonts to .ufont files in a temporary folder when running in standalone.
This simulates how font faces behave in a cooked build and actually allow font faces to be async loaded in standalone.
- Introduced EFontFacesLoadingPaintPolicy for Slate and UMG text widgets. If font faces are still loading when a text widget is painting, we either skip painting the text to leave the text widget blank or paint whatever text characters are available in the font (placeholder characters and all). This can be set on a per widget basis for the UMG text widgets in the details panel..
- Introduced an OnFontFacesFinishLoading event to the UMG text widgets that designers can override. The vent is fired when all font faces that are pending loading for the text in the text widget has finished loading. The event is a sparse delegate and only takes up 1 byte if unbound.
- Refactored FTextLayout to have a smaller memory footprint. Even with additional data members, the class is 4 bytes smaller than it was previously.
- Introduced the the Slate.Font.DetectAsyncLoadingFontFaces cvar. If true, then text widgets will keep track of font faces that are loading as its text is being shaped. This in turn will have the text widgets change the paint behavior for the text as font faces are loading and fire the vent (if any) when all font faces have finished loading. If set to false, neither the aforementioned behaviors will occur and the async font loading will be the same as before this feature was introduced.

To use this feature, ensure that both Slate.Font.DetectAsyncLoadingFontFaces and Slate.Font.AsyncLazyLoad cvars are both set to true in a cooked build. If running in standalone (-game on the command line for the Editor), ensure that UI.Slate.Font.SimulateCookedFontFilesInStandalone is also set to true.

diff --git a/Engine/Source/Runtime/Engine/Classes/Engine/FontFace.h b/Engine/Source/Runtime/Engine/Classes/Engine/FontFace.h
--- a/Engine/Source/Runtime/Engine/Classes/Engine/FontFace.h
+++ b/Engine/Source/Runtime/Engine/Classes/Engine/FontFace.h
@@ -88,6 +88,9 @@ private:
 #endif // WITH_EDITOR
 	//~ End UObject Interface
 
+#if WITH_EDITOR
+	FString GetStandaloneFilename() const;
+#endif
 public:
 	/** The filename of the font face we were created from. This may not always exist on disk, as we may have previously loaded and cached the font data inside this asset. */
 	UPROPERTY(EditAnywhere, BlueprintReadOnly, Category=FontFace)
diff --git a/Engine/Source/Runtime/Engine/Private/FontFace.cpp b/Engine/Source/Runtime/Engine/Private/FontFace.cpp
--- a/Engine/Source/Runtime/Engine/Private/FontFace.cpp
+++ b/Engine/Source/Runtime/Engine/Private/FontFace.cpp
@@ -13,11 +13,27 @@
 #include "Misc/PackageName.h"
 #include "UObject/EditorObjectVersion.h"
 #include "UObject/UE5MainStreamObjectVersion.h"
+#include "Misc/FileHelper.h"
 
 #include UE_INLINE_GENERATED_CPP_BY_NAME(FontFace)
 
 DEFINE_LOG_CATEGORY_STATIC(LogFontFace, Log, All);
 
+namespace UE::Private::FontFace
+{
+	const FString& GetStandaloneFontFaceDirectory()
+	{
+		static const FString FontFaceDirectory = FPaths::Combine(FPaths::ProjectIntermediateDir(), TEXT("FontFaces"));
+		return FontFaceDirectory;
+	}
+}
+
+static bool GSimulateCookedFontFilesInStandalone = false;
+FAutoConsoleVariableRef CVarSimulateCookedFontFilesInStandalone(
+	TEXT("UI.Slate.Font.SimulateCookedFontFilesInStandalone"),
+	GSimulateCookedFontFilesInStandalone,
+	TEXT("If true when running in standalone, the embedded font face data will be written to the Intermediate project folder and loaded on demand. This simulates the loading of .ufont files in a cooked build."));
+
 TAutoConsoleVariable<FString> CVarFontFaceDistanceFieldRasterizationMode(
 	TEXT("UI.SlateSDFText.RasterizationMode"),
 	"Bitmap",
@@ -145,9 +161,56 @@ void UFontFace::PostLoad()
 #endif // WITH_EDITORONLY_DATA
 
 	UpdateDeviceRasterizationSettings();
+
+#if WITH_EDITOR
+	// For the regular Editor, all font faces need to be loaded for editing and the font data is embedded.
+	// In standalone, we write out the font data for lazily loaded fonts to the Intermediate project folder and replace the filename of the font face.
+	// This causes lazily loaded fonts to actually be loaded from disk to simulate the behavior in a cooked build.
+	if (!GIsEditor && !GIsServer && GSimulateCookedFontFilesInStandalone && LoadingPolicy == EFontLoadingPolicy::LazyLoad)
+	{
+		if (FontFaceData->HasData())
+		{
+			FString Filename = GetStandaloneFilename();
+			FString DirectoryPath = FPaths::GetPath(Filename);
+			if (!IPlatformFile::GetPlatformPhysical().DirectoryExists(*DirectoryPath))
+			{
+				UE_LOG(LogFontFace, Display, TEXT("Standalone font face directory '%s' does not exist. Creating directory."), *DirectoryPath);
+				if (!IPlatformFile::GetPlatformPhysical().CreateDirectoryTree(*DirectoryPath))
+				{
+					UE_LOG(LogFontFace, Warning, TEXT("Unable to create standalone font face directory '%s'. Font face '%s' will not be written and async font loading of font face will be skipped."), *DirectoryPath, *Filename);
+					return;
+				}
+			}
+
+			if (FFileHelper::SaveArrayToFile(FontFaceData->GetData(), *Filename))
+			{
+				UE_LOG(LogFontFace, Display, TEXT("Written font face file '%s' to simulate loading of cooked font files in standalone mode."), *Filename);
+				SourceFilename = MoveTemp(Filename);
+				// We clear the data in the FFontFaceData so we force a synchronous load of the data 
+				TArray<uint8> EmptyData;
+				FontFaceData->SetData(MoveTemp(EmptyData));
+			}
+			else
+			{
+				UE_LOG(LogFontFace, Warning, TEXT("Failed to write font face file '%s' to simulate loading of cooked font files in standalone mode."), *Filename);
+			}
+		}
+	}
+#endif
 }
 
 #if WITH_EDITOR
+FString UFontFace::GetStandaloneFilename() const
+{
+	// This should ever only be called when in standalone 
+	ensure(!GIsEditor);
+	// This implementation follows UFontFace::GetCookedFilename. If that implementation changes, it is best to also adapt this to the change 
+	FString PackageName = GetOutermost()->GetName();
+	PackageName = FPackageName::GetLocalizedPackagePath(PackageName);
+
+	return FPaths::Combine(UE::Private::FontFace::GetStandaloneFontFaceDirectory(), PackageName) / GetName() + TEXT(".ufont");
+}
+
 void UFontFace::PostEditChangeProperty(FPropertyChangedEvent& PropertyChangedEvent)
 {
 	Super::PostEditChangeProperty(PropertyChangedEvent);
diff --git a/Engine/Source/Runtime/Slate/Private/Framework/Text/ShapedTextCache.cpp b/Engine/Source/Runtime/Slate/Private/Framework/Text/ShapedTextCache.cpp
--- a/Engine/Source/Runtime/Slate/Private/Framework/Text/ShapedTextCache.cpp
+++ b/Engine/Source/Runtime/Slate/Private/Framework/Text/ShapedTextCache.cpp
@@ -18,7 +18,7 @@ FShapedGlyphSequencePtr FShapedTextCache::FindShapedText(const FCachedShapedText
 	return nullptr;
 }
 
-FShapedGlyphSequenceRef FShapedTextCache::AddShapedText(const FCachedShapedTextKey& InKey, const TCHAR* InText)
+FShapedGlyphSequenceRef FShapedTextCache::AddShapedText(const FCachedShapedTextKey& InKey, const TCHAR* InText, TArray<TWeakPtr<FFreeTypeFace>>* OutLoadingFontFaces)
 {
 	const TSharedPtr<FSlateFontCache> FontCache = FontCachePtr.Pin();
 
@@ -34,13 +34,14 @@ FShapedGlyphSequenceRef FShapedTextCache::AddShapedText(const FCachedShapedTextK
 		InKey.FontInfo, 
 		InKey.Scale, 
 		InKey.TextContext.BaseDirection, 
-		InKey.TextContext.TextShapingMethod
+		InKey.TextContext.TextShapingMethod,
+		OutLoadingFontFaces
 		);
 
 	return AddShapedText(InKey, ShapedText);
 }
 
-FShapedGlyphSequenceRef FShapedTextCache::AddShapedText(const FCachedShapedTextKey& InKey, const TCHAR* InText, const TextBiDi::ETextDirection InTextDirection)
+FShapedGlyphSequenceRef FShapedTextCache::AddShapedText(const FCachedShapedTextKey& InKey, const TCHAR* InText, const TextBiDi::ETextDirection InTextDirection, TArray<TWeakPtr<FFreeTypeFace>>* OutLoadingFontFaces)
 {
 	const TSharedPtr<FSlateFontCache> FontCache = FontCachePtr.Pin();
 
@@ -56,7 +57,8 @@ FShapedGlyphSequenceRef FShapedTextCache::AddShapedText(const FCachedShapedTextK
 		InKey.FontInfo, 
 		InKey.Scale, 
 		InTextDirection, 
-		InKey.TextContext.TextShapingMethod
+		InKey.TextContext.TextShapingMethod,
+		OutLoadingFontFaces
 		);
 
 	return AddShapedText(InKey, ShapedText);
@@ -68,25 +70,25 @@ FShapedGlyphSequenceRef FShapedTextCache::AddShapedText(const FCachedShapedTextK
 	return InShapedText;
 }
 
-FShapedGlyphSequenceRef FShapedTextCache::FindOrAddShapedText(const FCachedShapedTextKey& InKey, const TCHAR* InText)
+FShapedGlyphSequenceRef FShapedTextCache::FindOrAddShapedText(const FCachedShapedTextKey& InKey, const TCHAR* InText, TArray<TWeakPtr<FFreeTypeFace>>* OutLoadingFontFaces)
 {
 	FShapedGlyphSequencePtr ShapedText = FindShapedText(InKey);
 
 	if (!ShapedText.IsValid())
 	{
-		ShapedText = AddShapedText(InKey, InText);
+		ShapedText = AddShapedText(InKey, InText, OutLoadingFontFaces);
 	}
 
 	return ShapedText.ToSharedRef();
 }
 
-FShapedGlyphSequenceRef FShapedTextCache::FindOrAddShapedText(const FCachedShapedTextKey& InKey, const TCHAR* InText, const TextBiDi::ETextDirection InTextDirection)
+FShapedGlyphSequenceRef FShapedTextCache::FindOrAddShapedText(const FCachedShapedTextKey& InKey, const TCHAR* InText, const TextBiDi::ETextDirection InTextDirection, TArray<TWeakPtr<FFreeTypeFace>>* OutLoadingFontFaces)
 {
 	FShapedGlyphSequencePtr ShapedText = FindShapedText(InKey);
 
 	if (!ShapedText.IsValid())
 	{
-		ShapedText = AddShapedText(InKey, InText, InTextDirection);
+		ShapedText = AddShapedText(InKey, InText, InTextDirection, OutLoadingFontFaces);
 	}
 
 	return ShapedText.ToSharedRef();
@@ -114,11 +116,10 @@ void FShapedTextCache::Clear()
 	CachedShapedText.Reset();
 }
 
-
-FVector2D ShapedTextCacheUtil::MeasureShapedText(const FShapedTextCacheRef& InShapedTextCache, const FCachedShapedTextKey& InRunKey, const FTextRange& InMeasureRange, const TCHAR* InText)
+FVector2D ShapedTextCacheUtil::MeasureShapedText(const FShapedTextCacheRef& InShapedTextCache, const FCachedShapedTextKey& InRunKey, const FTextRange& InMeasureRange, const TCHAR* InText, TArray<TWeakPtr<FFreeTypeFace>>* OutLoadingFontFaces)
 {
 	// Get the shaped text for the entire run and try and take a sub-measurement from it - this can help minimize the amount of text shaping that needs to be done when measuring text
-	FShapedGlyphSequenceRef ShapedText = InShapedTextCache->FindOrAddShapedText(InRunKey, InText);
+	FShapedGlyphSequenceRef ShapedText = InShapedTextCache->FindOrAddShapedText(InRunKey, InText, OutLoadingFontFaces);
 
 	TOptional<int32> MeasuredWidth = ShapedText->GetMeasuredWidth(InMeasureRange.BeginIndex, InMeasureRange.EndIndex);
 	if (!MeasuredWidth.IsSet())
@@ -127,7 +128,7 @@ FVector2D ShapedTextCacheUtil::MeasureShapedText(const FShapedTextCacheRef& InSh
 		MeasureKey.TextRange = InMeasureRange;
 
 		// Couldn't measure the sub-range, try and measure from a shape of the specified range
-		ShapedText = InShapedTextCache->FindOrAddShapedText(MeasureKey, InText);
+		ShapedText = InShapedTextCache->FindOrAddShapedText(MeasureKey, InText, OutLoadingFontFaces);
 		MeasuredWidth = ShapedText->GetMeasuredWidth();
 	}
 
@@ -216,10 +217,10 @@ int32 ShapedTextCacheUtil::FindCharacterIndexAtOffset(const FShapedTextCacheRef&
 	return GlyphOffsetResultValue.CharacterIndex;
 }
 
-int8 ShapedTextCacheUtil::GetShapedGlyphKerning(const FShapedTextCacheRef& InShapedTextCache, const FCachedShapedTextKey& InRunKey, const int32 InGlyphIndex, const TCHAR* InText)
+int8 ShapedTextCacheUtil::GetShapedGlyphKerning(const FShapedTextCacheRef& InShapedTextCache, const FCachedShapedTextKey& InRunKey, const int32 InGlyphIndex, const TCHAR* InText, TArray<TWeakPtr<FFreeTypeFace>>* OutLoadingFontFaces)
 {
 	// Get the shaped text for the entire run and try and get the kerning from it - this can help minimize the amount of text shaping that needs to be done when calculating kerning
-	FShapedGlyphSequenceRef ShapedText = InShapedTextCache->FindOrAddShapedText(InRunKey, InText);
+	FShapedGlyphSequenceRef ShapedText = InShapedTextCache->FindOrAddShapedText(InRunKey, InText, OutLoadingFontFaces);
 
 	TOptional<int8> Kerning = ShapedText->GetKerning(InGlyphIndex);
 	if (!Kerning.IsSet())
@@ -228,19 +229,19 @@ int8 ShapedTextCacheUtil::GetShapedGlyphKerning(const FShapedTextCacheRef& InSha
 		KerningKey.TextRange = FTextRange(InGlyphIndex, InGlyphIndex + 2);
 
 		// Couldn't get the kerning from the main run data, try and get the kerning from a shape of the specified range
-		ShapedText = InShapedTextCache->FindOrAddShapedText(KerningKey, InText);
+		ShapedText = InShapedTextCache->FindOrAddShapedText(KerningKey, InText, OutLoadingFontFaces);
 		Kerning = ShapedText->GetKerning(InGlyphIndex);
 	}
 
 	return Kerning.Get(0);
 }
 
-FShapedGlyphSequenceRef ShapedTextCacheUtil::GetShapedTextSubSequence(const FShapedTextCacheRef& InShapedTextCache, const FCachedShapedTextKey& InRunKey, const FTextRange& InTextRange, const TCHAR* InText, const TextBiDi::ETextDirection InTextDirection)
+FShapedGlyphSequenceRef ShapedTextCacheUtil::GetShapedTextSubSequence(const FShapedTextCacheRef& InShapedTextCache, const FCachedShapedTextKey& InRunKey, const FTextRange& InTextRange, const TCHAR* InText, const TextBiDi::ETextDirection InTextDirection, TArray<TWeakPtr<FFreeTypeFace>>* OutLoadingFontFaces)
 {
 	LLM_SCOPE_BYTAG(UI_Text);
 
 	// Get the shaped text for the entire run and try and make a sub-sequence from it - this can help minimize the amount of text shaping that needs to be done when drawing text
-	FShapedGlyphSequencePtr ShapedText = InShapedTextCache->FindOrAddShapedText(InRunKey, InText);
+	FShapedGlyphSequencePtr ShapedText = InShapedTextCache->FindOrAddShapedText(InRunKey, InText, OutLoadingFontFaces);
 
 	if (InRunKey.TextRange != InTextRange)
 	{
@@ -267,7 +268,7 @@ FShapedGlyphSequenceRef ShapedTextCacheUtil::GetShapedTextSubSequence(const FSha
 			else
 			{
 				// Couldn't get the sub-sequence, try and make a new shape for it instead
-				ShapedText = InShapedTextCache->FindOrAddShapedText(SubSequenceKey, InText, InTextDirection);
+				ShapedText = InShapedTextCache->FindOrAddShapedText(SubSequenceKey, InText, InTextDirection, OutLoadingFontFaces);
 			}
 		}
 	}
@@ -275,3 +276,4 @@ FShapedGlyphSequenceRef ShapedTextCacheUtil::GetShapedTextSubSequence(const FSha
 	check(ShapedText.IsValid());
 	return ShapedText.ToSharedRef();
 }
+
diff --git a/Engine/Source/Runtime/Slate/Private/Framework/Text/SlateTextLayout.cpp b/Engine/Source/Runtime/Slate/Private/Framework/Text/SlateTextLayout.cpp
--- a/Engine/Source/Runtime/Slate/Private/Framework/Text/SlateTextLayout.cpp
+++ b/Engine/Source/Runtime/Slate/Private/Framework/Text/SlateTextLayout.cpp
@@ -1,6 +1,7 @@
 // Copyright Epic Games, Inc. All Rights Reserved.
 
 #include "Framework/Text/SlateTextLayout.h"
+#include "Framework/Text/SlateAsyncFontLoadingTypes.h"
 #include "Rendering/DrawElements.h"
 #include "Fonts/FontCache.h"
 #include "Framework/Application/SlateApplication.h"
@@ -26,6 +27,7 @@ FSlateTextLayout::FSlateTextLayout(SWidget* InOwner, FTextBlockStyle InDefaultTe
 	: DefaultTextStyle(MoveTemp(InDefaultTextStyle))
 	, Children(InOwner, false)
 	, bIsPassword(false)
+	, FontFacesLoadingPaintPolicy(EFontFacesLoadingPaintPolicy::DoNotPaint)
 {
 
 }
@@ -59,6 +61,16 @@ static FAutoConsoleVariableRef CVarSlateShowTextDebugging(TEXT("Slate.ShowTextDe
 
 int32 FSlateTextLayout::OnPaint( const FPaintArgs& Args, const FGeometry& AllottedGeometry, const FSlateRect& MyCullingRect, FSlateWindowElementList& OutDrawElements, int32 LayerId, const FWidgetStyle& InWidgetStyle, bool bParentEnabled ) const
 {
+	// The layout pass would have already started shaping all of the text in this text layout and caching the shaped text as part of measurement 
+	// We should know by this point whether any of the text runs are waiting on a font to load
+	if (AreAnyFontFacesLoading())
+	{
+		if (GetFontFacesLoadingPaintPolicy() == EFontFacesLoadingPaintPolicy::DoNotPaint)
+		{
+			return LayerId;
+		}
+	}
+
 	const ESlateDrawEffect DrawEffects = bParentEnabled ? ESlateDrawEffect::None : ESlateDrawEffect::DisabledEffect;
 
 #if !(UE_BUILD_SHIPPING || UE_BUILD_TEST)
@@ -323,3 +335,39 @@ TSharedRef<IRun> FSlateTextLayout::CreateDefaultTextRun(const TSharedRef<FString
 	}
 	return FSlateTextRun::Create(FRunInfo(), NewText, DefaultTextStyle, NewRange);
 }
+
+void FSlateTextLayout::SetOnAllFontFacesFinishLoading(const FSimpleDelegate& InAllFontFacesFinishLoading)
+{
+	SWidget& OwnerWidget = Children.GetOwner();
+	TSharedPtr< FSlateTextLayoutPendingFontFacesLoadingMetaData> MetaData = OwnerWidget.GetMetaData<FSlateTextLayoutPendingFontFacesLoadingMetaData>();
+	if (MetaData)
+	{
+		MetaData->OnAllFontFacesFinishLoading = InAllFontFacesFinishLoading;
+	}
+	// No point creating a new metadata and adding it to the widget if the passed in delegate isn't bound 
+	else if (InAllFontFacesFinishLoading.IsBound())
+	{
+		TSharedRef<FSlateTextLayoutPendingFontFacesLoadingMetaData> NewMetaData = MakeShared<FSlateTextLayoutPendingFontFacesLoadingMetaData>();
+		NewMetaData->OnAllFontFacesFinishLoading = InAllFontFacesFinishLoading;
+		OwnerWidget.AddMetadata(NewMetaData);
+	}
+}
+
+void FSlateTextLayout::HandleAllFontFacesFinishLoading()
+{
+	TSharedPtr<FSlateTextLayoutPendingFontFacesLoadingMetaData> MetaData = Children.GetOwner().GetMetaData<FSlateTextLayoutPendingFontFacesLoadingMetaData>();
+	if (MetaData)
+	{
+		MetaData->OnAllFontFacesFinishLoading.ExecuteIfBound();
+	}
+}
+
+EFontFacesLoadingPaintPolicy FSlateTextLayout::GetFontFacesLoadingPaintPolicy() const
+{
+	return FontFacesLoadingPaintPolicy;
+}
+
+void FSlateTextLayout::SetFontFacesLoadingPaintPolicy(EFontFacesLoadingPaintPolicy InFontFacesLoadingPaintPolicy)
+{
+	FontFacesLoadingPaintPolicy = InFontFacesLoadingPaintPolicy;
+}
diff --git a/Engine/Source/Runtime/Slate/Private/Framework/Text/SlateTextRun.cpp b/Engine/Source/Runtime/Slate/Private/Framework/Text/SlateTextRun.cpp
--- a/Engine/Source/Runtime/Slate/Private/Framework/Text/SlateTextRun.cpp
+++ b/Engine/Source/Runtime/Slate/Private/Framework/Text/SlateTextRun.cpp
@@ -53,7 +53,7 @@ FVector2D FSlateTextRun::Measure( int32 BeginIndex, int32 EndIndex, float Scale,
 	}
 	
 	// Use the full text range (rather than the run range) so that text that spans runs will still be shaped correctly
-	return ShapedTextCacheUtil::MeasureShapedText(TextContext.ShapedTextCache, FCachedShapedTextKey(FTextRange(0, Text->Len()), Scale, TextContext, Style.Font), FTextRange(BeginIndex, EndIndex), **Text) + OutlineSize + ShadowSize;
+	return ShapedTextCacheUtil::MeasureShapedText(TextContext.ShapedTextCache, FCachedShapedTextKey(FTextRange(0, Text->Len()), Scale, TextContext, Style.Font), FTextRange(BeginIndex, EndIndex), **Text, TextContext.LoadingFontFaces) + OutlineSize + ShadowSize;
 }
 
 int8 FSlateTextRun::GetKerning(int32 CurrentIndex, float Scale, const FRunTextContext& TextContext) const
@@ -65,7 +65,7 @@ int8 FSlateTextRun::GetKerning(int32 CurrentIndex, float Scale, const FRunTextCo
 	}
 
 	// Use the full text range (rather than the run range) so that text that spans runs will still be shaped correctly
-	return ShapedTextCacheUtil::GetShapedGlyphKerning(TextContext.ShapedTextCache, FCachedShapedTextKey(FTextRange(0, Text->Len()), Scale, TextContext, Style.Font), PreviousIndex, **Text);
+	return ShapedTextCacheUtil::GetShapedGlyphKerning(TextContext.ShapedTextCache, FCachedShapedTextKey(FTextRange(0, Text->Len()), Scale, TextContext, Style.Font), PreviousIndex, **Text, TextContext.LoadingFontFaces);
 }
 
 FVector2d FSlateTextRun::GetOutlineSize(int32 StartIndex, int32 EndIndex, float Scale) const
@@ -95,6 +95,10 @@ TSharedRef< ILayoutBlock > FSlateTextRun::CreateBlock( int32 BeginIndex, int32 E
 }
 int32 FSlateTextRun::OnPaint(const FPaintArgs& PaintArgs, const FTextArgs& TextArgs, const FGeometry& AllottedGeometry, const FSlateRect& MyCullingRect, FSlateWindowElementList& OutDrawElements, int32 LayerId, const FWidgetStyle& InWidgetStyle, bool bParentEnabled) const
 {
+	if (Text->IsEmpty())
+	{
+		return LayerId;
+	}
 #if !UE_BUILD_SHIPPING
 	static const auto CVar = IConsoleManager::Get().FindConsoleVariable(TEXT("Slate.LogPaintedText"));
 	if (CVar->GetBool())
diff --git a/Engine/Source/Runtime/Slate/Private/Framework/Text/TextLayout.cpp b/Engine/Source/Runtime/Slate/Private/Framework/Text/TextLayout.cpp
--- a/Engine/Source/Runtime/Slate/Private/Framework/Text/TextLayout.cpp
+++ b/Engine/Source/Runtime/Slate/Private/Framework/Text/TextLayout.cpp
@@ -37,7 +37,7 @@ ETextFlowDirection GetDefaultTextFlowDirection()
 
 FTextLayout::FBreakCandidate FTextLayout::CreateBreakCandidate( int32& OutRunIndex, FLineModel& Line, int32 PreviousBreak, int32 CurrentBreak )
 {
-	const FRunTextContext RunTextContext(TextShapingMethod, Line.TextBaseDirection, Line.ShapedTextCache);
+	const FRunTextContext RunTextContext(TextShapingMethod, Line.TextBaseDirection, Line.ShapedTextCache, &LoadingFontFaces);
 
 	bool SuccessfullyMeasuredSlice = false;
 	int16 MaxAboveBaseline = 0;
@@ -202,7 +202,7 @@ void FTextLayout::CreateLineViewBlocks( int32 LineModelIndex, const int32 StopIn
 {
 	const FLineModel& LineModel = LineModels[ LineModelIndex ];
 
-	const FRunTextContext RunTextContext(TextShapingMethod, LineModel.TextBaseDirection, LineModel.ShapedTextCache);
+	const FRunTextContext RunTextContext(TextShapingMethod, LineModel.TextBaseDirection, LineModel.ShapedTextCache, &LoadingFontFaces);
 
 	int16 MaxAboveBaseline = 0;
 	int16 MaxBelowBaseline = 0;
@@ -552,7 +552,6 @@ float FTextLayout::GetWrappingDrawWidth() const
 void FTextLayout::FlowLayout()
 {
 	const float WrappingDrawWidth = GetWrappingDrawWidth();
-
 	TArray< TSharedRef< ILayoutBlock > > SoftLine;
 	for (int32 LineModelIndex = 0; LineModelIndex < LineModels.Num(); LineModelIndex++)
 	{
@@ -625,7 +624,7 @@ void FTextLayout::FlowLineLayout(const int32 LineModelIndex, const float Wrappin
 
 			if ( WrappingPolicy == ETextWrappingPolicy::AllowPerCharacterWrapping && !BreakWithoutTrailingWhitespaceDoesFit && IsFirstBreakOnSoftLine )
 			{
-				const FRunTextContext RunTextContext(TextShapingMethod, LineModel.TextBaseDirection, LineModel.ShapedTextCache);
+				const FRunTextContext RunTextContext(TextShapingMethod, LineModel.TextBaseDirection, LineModel.ShapedTextCache, &LoadingFontFaces);
 
 				// This is a single word that's too long to fit onto a line, so we'll fallback to wrapping this word at grapheme cluster boundaries - this may require more than a single split
 				const int32 NonBreakingStringIndexOffset = PreviousBlockEnd;
@@ -810,7 +809,7 @@ void FTextLayout::FlowHighlights()
 
 		FLineModel& LineModel = LineModels[LineView.ModelIndex];
 		
-		const FRunTextContext RunTextContext(TextShapingMethod, LineModel.TextBaseDirection, LineModel.ShapedTextCache);
+		const FRunTextContext RunTextContext(TextShapingMethod, LineModel.TextBaseDirection, LineModel.ShapedTextCache, &LoadingFontFaces);
 
 		// Insert each highlighter into every line view that's within its range, either as an underlay, or as an overlay
 		for (FTextLineHighlight& LineHighlight : LineModel.LineHighlights)
@@ -1012,6 +1011,24 @@ void FTextLayout::EndLayout()
 	{
 		EndLineLayout(LineModel);
 	}
+
+	if (LoadingFontFaces.IsEmpty())
+	{
+		// We just finished loading all fonts 
+		if (bLoadingFontFacesOnPreviousLayout)
+		{
+			HandleAllFontFacesFinishLoading();
+			bLoadingFontFacesOnPreviousLayout = !LoadingFontFaces.IsEmpty();
+		}
+	}
+	else
+	{
+		// We just started loading necessary font faces 
+		if (!bLoadingFontFacesOnPreviousLayout)
+		{
+			bLoadingFontFacesOnPreviousLayout = !LoadingFontFaces.IsEmpty();
+		}
+	}
 }
 
 void FTextLayout::EndLineLayout(FLineModel& LineModel)
@@ -1202,34 +1219,31 @@ void FTextLayout::DirtyAllLineModels(const ELineModelDirtyState::Flags InDirtyFl
 }
 
 FTextLayout::FTextLayout()
-	: LineModels()
-	, LineViews()
-	, LineViewsToJustify()
+	: ViewSize(ForceInitToZero)
+	, ScrollOffset(ForceInitToZero)
+	, LineBreakIterator() // Initialized in FTextLayout::CreateWrappingCache if no custom iterator is provided
+	, GraphemeBreakIterator(GSlateUseSharedBreakIterator ? FBreakIterator::GetCharacterBoundaryIterator() : FBreakIterator::CreateCharacterBoundaryIterator())
+	, WordBreakIterator(GSlateUseSharedBreakIterator ? FBreakIterator::GetWordBreakIterator() : FBreakIterator::CreateWordBreakIterator())
+	, Justification(ETextJustify::Left)
+	, TextBiDiDetection(TextBiDi::CreateTextBiDi())
+	, Scale(1.0f)
+	, WrappingWidth(0.0f)
+	, LineHeightPercentage(1.0f)
+	, OverHeight(0.0f)
 	, DirtyFlags( ETextLayoutDirtyState::None )
 	, TextShapingMethod( GetDefaultTextShapingMethod() )
 	, TextFlowDirection( GetDefaultTextFlowDirection() )
-	, Scale( 1.0f )
-	, WrappingWidth( 0 )
 	, WrappingPolicy( ETextWrappingPolicy::DefaultWrapping )
 	, TransformPolicy( ETextTransformPolicy::None )
-	, Margin()
-	, Justification( ETextJustify::Left )
-	, LineHeightPercentage( 1.0f )
 	, ApplyLineHeightToBottomLine( true )
-	, TextLayoutSize()
-	, OverHeight( 0.0f )
-	, ViewSize( ForceInitToZero )
-	, ScrollOffset( ForceInitToZero )
-	, LineBreakIterator() // Initialized in FTextLayout::CreateWrappingCache if no custom iterator is provided
-	, GraphemeBreakIterator(GSlateUseSharedBreakIterator  ? FBreakIterator::GetCharacterBoundaryIterator() : FBreakIterator::CreateCharacterBoundaryIterator())
-	, WordBreakIterator(GSlateUseSharedBreakIterator  ? FBreakIterator::GetWordBreakIterator() : FBreakIterator::CreateWordBreakIterator())
-	, TextBiDiDetection(TextBiDi::CreateTextBiDi())
-	, TextOverflowPolicyOverride()
+	, bLoadingFontFacesOnPreviousLayout(false)
 {
 }
 
 void FTextLayout::UpdateIfNeeded()
 {
+	// GSlateLayoutGeneration is incremented when font faces are finally done async lazily loaded
+	// This re-measures and lays out the text accordingly
 	if (CachedLayoutGeneration != GSlateLayoutGeneration)
 	{
 		CachedLayoutGeneration = GSlateLayoutGeneration;
@@ -1257,6 +1271,8 @@ void FTextLayout::UpdateLayout()
 {
 	SCOPE_CYCLE_COUNTER(STAT_SlateTextLayout);
 
+	// We clear all loading font faces every re-layout so that it can be populated anew during shaping 
+	LoadingFontFaces.Empty();
 	ClearView();
 	BeginLayout();
 
@@ -2544,6 +2560,12 @@ FTextSelection FTextLayout::GetWordAt(const FTextLocation& Location) const
 	return FTextSelection(FTextLocation(LineIndex, CurrentBreak), FTextLocation(LineIndex, PreviousBreak));
 }
 
+bool FTextLayout::AreAnyFontFacesLoading() const
+{
+	return !LoadingFontFaces.IsEmpty();
+}
+
+
 void FTextLayout::SetVisibleRegion( const FVector2D& InViewSize, const FVector2D& InScrollOffset )
 {
 	if (ViewSize != InViewSize)
diff --git a/Engine/Source/Runtime/Slate/Private/Widgets/Input/SEditableText.cpp b/Engine/Source/Runtime/Slate/Private/Widgets/Input/SEditableText.cpp
--- a/Engine/Source/Runtime/Slate/Private/Widgets/Input/SEditableText.cpp
+++ b/Engine/Source/Runtime/Slate/Private/Widgets/Input/SEditableText.cpp
@@ -77,6 +77,11 @@ void SEditableText::Construct( const FArguments& InArgs )
 	EditableTextLayout->SetDebugSourceInfo(TAttribute<FString>::Create(TAttribute<FString>::FGetter::CreateLambda([this]{ return FReflectionMetaData::GetWidgetDebugInfo(this); })));
 	EditableTextLayout->SetJustification(InArgs._Justification);
 	EditableTextLayout->SetOverflowPolicy(InArgs._OverflowPolicy);
+	EditableTextLayout->SetFontFacesLoadingPaintPolicy(InArgs._FontFacesLoadingPaintPolicy);
+	if (InArgs._OnAllFontFacesFinishLoading.IsBound())
+	{
+		EditableTextLayout->SetOnAllFontFacesFinishLoadingDelegate(InArgs._OnAllFontFacesFinishLoading);
+	}
 
 	// build context menu extender
 	MenuExtender = MakeShareable(new FExtender());
@@ -417,6 +422,14 @@ void SEditableText::SetOverflowPolicy(TOptional<ETextOverflowPolicy> InOverflowP
 	EditableTextLayout->SetOverflowPolicy(InOverflowPolicy);
 }
 
+void SEditableText::SetFontFacesLoadingPaintPolicy(EFontFacesLoadingPaintPolicy InFontFacesLoadingPaintPolicy)
+{
+	if (EditableTextLayout)
+	{
+		EditableTextLayout->SetFontFacesLoadingPaintPolicy(InFontFacesLoadingPaintPolicy);
+	}
+}
+
 bool SEditableText::AnyTextSelected() const
 {
 	return EditableTextLayout->AnyTextSelected();
diff --git a/Engine/Source/Runtime/Slate/Private/Widgets/Input/SEditableTextBox.cpp b/Engine/Source/Runtime/Slate/Private/Widgets/Input/SEditableTextBox.cpp
--- a/Engine/Source/Runtime/Slate/Private/Widgets/Input/SEditableTextBox.cpp
+++ b/Engine/Source/Runtime/Slate/Private/Widgets/Input/SEditableTextBox.cpp
@@ -90,6 +90,8 @@ void SEditableTextBox::Construct( const FArguments& InArgs )
 					.TextShapingMethod(InArgs._TextShapingMethod)
 					.TextFlowDirection( InArgs._TextFlowDirection )
 					.OverflowPolicy(InArgs._OverflowPolicy)
+					.FontFacesLoadingPaintPolicy(InArgs._FontFacesLoadingPaintPolicy)
+					.OnAllFontFacesFinishLoading(InArgs._OnAllFontFacesFinishLoading)
 				]
 			]
 		]
@@ -211,6 +213,14 @@ void SEditableTextBox::SetOverflowPolicy(TOptional<ETextOverflowPolicy> InOverfl
 	EditableText->SetOverflowPolicy(InOverflowPolicy);
 }
 
+void SEditableTextBox::SetFontFacesLoadingPaintPolicy(EFontFacesLoadingPaintPolicy InFontFacesLoadingPaintPolicy)
+{
+	if (EditableText)
+	{
+		EditableText->SetFontFacesLoadingPaintPolicy(InFontFacesLoadingPaintPolicy);
+	}
+}
+
 bool SEditableTextBox::AnyTextSelected() const
 {
 	return EditableText->AnyTextSelected();
diff --git a/Engine/Source/Runtime/Slate/Private/Widgets/Input/SMultiLineEditableTextBox.cpp b/Engine/Source/Runtime/Slate/Private/Widgets/Input/SMultiLineEditableTextBox.cpp
--- a/Engine/Source/Runtime/Slate/Private/Widgets/Input/SMultiLineEditableTextBox.cpp
+++ b/Engine/Source/Runtime/Slate/Private/Widgets/Input/SMultiLineEditableTextBox.cpp
@@ -136,6 +136,8 @@ void SMultiLineEditableTextBox::Construct( const FArguments& InArgs )
 					.TextFlowDirection(InArgs._TextFlowDirection)
 					.AllowContextMenu(InArgs._AllowContextMenu)
 					.OverflowPolicy(InArgs._OverflowPolicy)
+					.FontFacesLoadingPaintPolicy(InArgs._FontFacesLoadingPaintPolicy)
+					.OnAllFontFacesFinishLoading(InArgs._OnAllFontFacesFinishLoading)
 				]
 
 				+SVerticalBox::Slot()
@@ -425,6 +427,14 @@ void SMultiLineEditableTextBox::SetIsReadOnly(const TAttribute<bool>& InIsReadOn
 	EditableText->SetIsReadOnly(InIsReadOnly);
 }
 
+void SMultiLineEditableTextBox::SetFontFacesLoadingPaintPolicy( EFontFacesLoadingPaintPolicy InFontFacesLoadingPaintPolicy)
+{
+	if (EditableText)
+	{
+		EditableText->SetFontFacesLoadingPaintPolicy(InFontFacesLoadingPaintPolicy);
+	}
+}
+
 void SMultiLineEditableTextBox::SetError( const FText& InError )
 {
 	SetError( InError.ToString() );
diff --git a/Engine/Source/Runtime/Slate/Private/Widgets/Text/SMultiLineEditableText.cpp b/Engine/Source/Runtime/Slate/Private/Widgets/Text/SMultiLineEditableText.cpp
--- a/Engine/Source/Runtime/Slate/Private/Widgets/Text/SMultiLineEditableText.cpp
+++ b/Engine/Source/Runtime/Slate/Private/Widgets/Text/SMultiLineEditableText.cpp
@@ -91,6 +91,11 @@ void SMultiLineEditableText::Construct( const FArguments& InArgs )
 	EditableTextLayout->SetLineHeightPercentage(InArgs._LineHeightPercentage);
 	EditableTextLayout->SetDebugSourceInfo(TAttribute<FString>::Create(TAttribute<FString>::FGetter::CreateLambda([this]{ return FReflectionMetaData::GetWidgetDebugInfo(this); })));
 	EditableTextLayout->SetOverflowPolicy(InArgs._OverflowPolicy);
+	EditableTextLayout->SetFontFacesLoadingPaintPolicy(InArgs._FontFacesLoadingPaintPolicy);
+	if (InArgs._OnAllFontFacesFinishLoading.IsBound())
+	{
+		EditableTextLayout->SetOnAllFontFacesFinishLoadingDelegate(InArgs._OnAllFontFacesFinishLoading);
+	}
 
 	// build context menu extender
 	MenuExtender = MakeShareable(new FExtender);
@@ -276,6 +281,14 @@ void SMultiLineEditableText::SetClearKeyboardFocusOnCommit(const TAttribute<bool
 	bClearKeyboardFocusOnCommit = InClearKeyboardFocusOnCommit;
 }
 
+void SMultiLineEditableText::SetFontFacesLoadingPaintPolicy(EFontFacesLoadingPaintPolicy InFontFacesLoadingPaintPolicy)
+{
+	if (EditableTextLayout)
+	{
+		EditableTextLayout->SetFontFacesLoadingPaintPolicy(InFontFacesLoadingPaintPolicy);
+	}
+}
+
 void SMultiLineEditableText::OnHScrollBarMoved(const float InScrollOffsetFraction)
 {
 	EditableTextLayout->SetHorizontalScrollFraction(InScrollOffsetFraction);
diff --git a/Engine/Source/Runtime/Slate/Private/Widgets/Text/SRichTextBlock.cpp b/Engine/Source/Runtime/Slate/Private/Widgets/Text/SRichTextBlock.cpp
--- a/Engine/Source/Runtime/Slate/Private/Widgets/Text/SRichTextBlock.cpp
+++ b/Engine/Source/Runtime/Slate/Private/Widgets/Text/SRichTextBlock.cpp
@@ -55,6 +55,11 @@ void SRichTextBlock::Construct( const FArguments& InArgs )
 		TextLayoutCache = MakeUnique<FSlateTextBlockLayout>(this, TextStyle, InArgs._TextShapingMethod, InArgs._TextFlowDirection, InArgs._CreateSlateTextLayout, Marshaller.ToSharedRef(), nullptr);
 		TextLayoutCache->SetDebugSourceInfo(TAttribute<FString>::Create(TAttribute<FString>::FGetter::CreateLambda([this]{ return FReflectionMetaData::GetWidgetDebugInfo(this); })));
 		TextLayoutCache->SetTextOverflowPolicy(InArgs._OverflowPolicy.IsSet() ? InArgs._OverflowPolicy : TextStyle.OverflowPolicy);
+		TextLayoutCache->SetFontFacesLoadingPaintPolicy(InArgs._FontFacesLoadingPaintPolicy);
+		if (InArgs._OnAllFontFacesFinishLoading.IsBound())
+		{
+			TextLayoutCache->SetOnAllFontFacesFinishLoading(InArgs._OnAllFontFacesFinishLoading);
+		}
 	}
 
 	SetCanTick(false);
@@ -224,6 +229,14 @@ void SRichTextBlock::SetTextBlockScale(const float NewTextBlockScale)
 	}
 }
 
+void SRichTextBlock::SetFontFacesLoadingPaintPolicy(EFontFacesLoadingPaintPolicy InFontFacesLoadingPaintPolicy)
+{
+	if (TextLayoutCache)
+	{
+		TextLayoutCache->SetFontFacesLoadingPaintPolicy(InFontFacesLoadingPaintPolicy);
+	}
+}
+
 void SRichTextBlock::Refresh()
 {
 	TextLayoutCache->DirtyContent();
diff --git a/Engine/Source/Runtime/Slate/Private/Widgets/Text/STextBlock.cpp b/Engine/Source/Runtime/Slate/Private/Widgets/Text/STextBlock.cpp
--- a/Engine/Source/Runtime/Slate/Private/Widgets/Text/STextBlock.cpp
+++ b/Engine/Source/Runtime/Slate/Private/Widgets/Text/STextBlock.cpp
@@ -152,13 +152,18 @@ void STextBlock::Construct( const FArguments& InArgs )
 	}
 
 	SetText(InArgs._Text);
-
 	//if(!bSimpleTextMode)
 	{
 		// We use a dummy style here (as it may not be safe to call the delegates used to compute the style), but the correct style is set by ComputeDesiredSize
 		TextLayoutCache = MakeUnique<FSlateTextBlockLayout>(this, FTextBlockStyle::GetDefault(), InArgs._TextShapingMethod, InArgs._TextFlowDirection, FCreateSlateTextLayout(), FPlainTextLayoutMarshaller::Create(), InArgs._LineBreakPolicy);
 		TextLayoutCache->SetDebugSourceInfo(TAttribute<FString>::Create(TAttribute<FString>::FGetter::CreateLambda([this] { return FReflectionMetaData::GetWidgetDebugInfo(this); })));
 		TextLayoutCache->SetTextOverflowPolicy(InArgs._OverflowPolicy.IsSet() ? InArgs._OverflowPolicy : TextStyle.OverflowPolicy);
+		TextLayoutCache->SetFontFacesLoadingPaintPolicy(InArgs._FontFacesLoadingPaintPolicy);
+	}
+
+	if (InArgs._OnAllFontFacesFinishLoading.IsBound() && TextLayoutCache)
+	{
+		TextLayoutCache->SetOnAllFontFacesFinishLoading(InArgs._OnAllFontFacesFinishLoading);
 	}
 }
 
@@ -478,6 +483,14 @@ void STextBlock::SetJustification(TAttribute<ETextJustify::Type> InJustification
 	Justification.Assign(*this, MoveTemp(InJustification));
 }
 
+void STextBlock::SetFontFacesLoadingPaintPolicy(EFontFacesLoadingPaintPolicy InFontFacesLoadingPaintPolicy)
+{
+	if (TextLayoutCache)
+	{
+		TextLayoutCache->SetFontFacesLoadingPaintPolicy(InFontFacesLoadingPaintPolicy);
+	}
+}
+
 FTextBlockStyle STextBlock::GetComputedTextStyle() const
 {
 	FTextBlockStyle ComputedStyle = TextStyle;
diff --git a/Engine/Source/Runtime/Slate/Private/Widgets/Text/SlateEditableTextLayout.cpp b/Engine/Source/Runtime/Slate/Private/Widgets/Text/SlateEditableTextLayout.cpp
--- a/Engine/Source/Runtime/Slate/Private/Widgets/Text/SlateEditableTextLayout.cpp
+++ b/Engine/Source/Runtime/Slate/Private/Widgets/Text/SlateEditableTextLayout.cpp
@@ -4197,4 +4197,20 @@ void FSlateEditableTextLayout::ToggleVirtualKeyboard(bool bShow, const int32 Use
 	{
 		FSlateApplication::Get().ShowVirtualKeyboard(bShow, UserIndex, VirtualKeyboardEntry);
 	}
+}
+
+void FSlateEditableTextLayout::SetOnAllFontFacesFinishLoadingDelegate(const FSimpleDelegate& InAllFontFacesFinishLoadingDelegate)
+{
+	if (TextLayout)
+	{
+		TextLayout->SetOnAllFontFacesFinishLoading(InAllFontFacesFinishLoadingDelegate);
+	}
+}
+
+void FSlateEditableTextLayout::SetFontFacesLoadingPaintPolicy(EFontFacesLoadingPaintPolicy InFontFacesLoadingPaintPolicy)
+{
+	if (TextLayout)
+	{
+		TextLayout->SetFontFacesLoadingPaintPolicy(InFontFacesLoadingPaintPolicy);
+	}
 }
\ No newline at end of file
diff --git a/Engine/Source/Runtime/Slate/Private/Widgets/Text/SlateTextBlockLayout.cpp b/Engine/Source/Runtime/Slate/Private/Widgets/Text/SlateTextBlockLayout.cpp
--- a/Engine/Source/Runtime/Slate/Private/Widgets/Text/SlateTextBlockLayout.cpp
+++ b/Engine/Source/Runtime/Slate/Private/Widgets/Text/SlateTextBlockLayout.cpp
@@ -231,6 +231,11 @@ float FSlateTextBlockLayout::GetLayoutScale() const
 
 int32 FSlateTextBlockLayout::OnPaint(const FPaintArgs& InPaintArgs, const FGeometry& InAllottedGeometry, const FSlateRect& InClippingRect, FSlateWindowElementList& OutDrawElements, int32 LayerId, const FWidgetStyle& InWidgetStyle, bool bParentEnabled)
 {
+	if (!TextLayout || TextLayout->IsEmpty())
+	{
+		return LayerId;
+	}
+
 	// Store a new cached size with the scale
 	CachedSize = FVector2f(InAllottedGeometry.GetLocalSize());
 
@@ -401,3 +406,19 @@ float FSlateTextBlockLayout::CalculateWrappingWidth() const
 
 	return FMath::Max(0.0f, WrappingWidth);
 }
+
+void FSlateTextBlockLayout::SetOnAllFontFacesFinishLoading(const FSimpleDelegate& InAllFontFacesFinishLoading)
+{
+	if (TextLayout)
+	{
+		TextLayout->SetOnAllFontFacesFinishLoading(InAllFontFacesFinishLoading);
+	}
+}
+
+void FSlateTextBlockLayout::SetFontFacesLoadingPaintPolicy(EFontFacesLoadingPaintPolicy InFontFacesLoadingPaintPolicy)
+{
+	if (TextLayout)
+	{
+		TextLayout->SetFontFacesLoadingPaintPolicy(InFontFacesLoadingPaintPolicy);
+	}
+}
diff --git a/Engine/Source/Runtime/Slate/Public/Framework/Text/IRun.h b/Engine/Source/Runtime/Slate/Public/Framework/Text/IRun.h
--- a/Engine/Source/Runtime/Slate/Public/Framework/Text/IRun.h
+++ b/Engine/Source/Runtime/Slate/Public/Framework/Text/IRun.h
@@ -9,6 +9,7 @@ class ILayoutBlock;
 class IRunRenderer;
 enum class ETextHitPoint : uint8;
 enum class ETextShapingMethod : uint8;
+class FFreeTypeFace;
 
 struct FRunInfo
 {
@@ -86,11 +87,20 @@ struct FRunTextContext : public FShapedTextContext
 	FRunTextContext(const ETextShapingMethod InTextShapingMethod, const TextBiDi::ETextDirection InBaseDirection, FShapedTextCacheRef InShapedTextCache)
 		: FShapedTextContext(InTextShapingMethod, InBaseDirection)
 		, ShapedTextCache(InShapedTextCache)
+		, LoadingFontFaces(nullptr)
 	{
 	}
 
+	FRunTextContext(const ETextShapingMethod InTextShapingMethod, const TextBiDi::ETextDirection InBaseDirection, FShapedTextCacheRef InShapedTextCache, TArray<TWeakPtr<FFreeTypeFace>>* InLoadingFontFaces)
+		: FShapedTextContext(InTextShapingMethod, InBaseDirection)
+		, ShapedTextCache(InShapedTextCache)
+		, LoadingFontFaces(InLoadingFontFaces)
+	{
+	}
 	/** The shaped text cache that should be used by this line of text */
 	FShapedTextCacheRef ShapedTextCache;
+
+	TArray<TWeakPtr<FFreeTypeFace>>* LoadingFontFaces;
 };
 
 /** The context needed when creating a block from a run of a text */
diff --git a/Engine/Source/Runtime/Slate/Public/Framework/Text/ShapedTextCache.h b/Engine/Source/Runtime/Slate/Public/Framework/Text/ShapedTextCache.h
--- a/Engine/Source/Runtime/Slate/Public/Framework/Text/ShapedTextCache.h
+++ b/Engine/Source/Runtime/Slate/Public/Framework/Text/ShapedTextCache.h
@@ -76,11 +76,12 @@ public:
 	 * @param InText				The text to shape. InKey may specify a sub-section of the entire text
 	 * @param InTextDirection		The text direction of all of the text to be shaped. If present we do a unidirectional shape, otherwise we do a bidirectional shape
 	 * @param InShapedText			The shaped text instance to add
+	 * @param OutLoadingFontFaces	An array (if provided) that keeps track of the unique font faces that are loading for a shaped text during the shaping algorithm.  
 	 *
 	 * @return The shaped text instance
 	 */
-	SLATE_API FShapedGlyphSequenceRef AddShapedText(const FCachedShapedTextKey& InKey, const TCHAR* InText);
-	SLATE_API FShapedGlyphSequenceRef AddShapedText(const FCachedShapedTextKey& InKey, const TCHAR* InText, const TextBiDi::ETextDirection InTextDirection);
+	SLATE_API FShapedGlyphSequenceRef AddShapedText(const FCachedShapedTextKey& InKey, const TCHAR* InText, TArray<TWeakPtr<FFreeTypeFace>>* OutLoadingFontFaces = nullptr);
+	SLATE_API FShapedGlyphSequenceRef AddShapedText(const FCachedShapedTextKey& InKey, const TCHAR* InText, const TextBiDi::ETextDirection InTextDirection, TArray<TWeakPtr<FFreeTypeFace>>* OutLoadingFontFaces = nullptr);
 	SLATE_API FShapedGlyphSequenceRef AddShapedText(const FCachedShapedTextKey& InKey, FShapedGlyphSequenceRef InShapedText);
 
 	/**
@@ -89,11 +90,12 @@ public:
 	 * @param InKey					The key identifying the shaped text instance to find or add
 	 * @param InText				The text to shape if we can't find the shaped text in the cache. InKey may specify a sub-section of the entire text
 	 * @param InTextDirection		The text direction of all of the text to be shaped. If present we do a unidirectional shape, otherwise we do a bidirectional shape
+	 * @param OutLoadingFontFaces	An array (if provided) that keeps track of the unique font faces that are loading for a shaped text during the shaping algorithm.  
 	 *
 	 * @return The shaped text instance
 	 */
-	SLATE_API FShapedGlyphSequenceRef FindOrAddShapedText(const FCachedShapedTextKey& InKey, const TCHAR* InText);
-	SLATE_API FShapedGlyphSequenceRef FindOrAddShapedText(const FCachedShapedTextKey& InKey, const TCHAR* InText, const TextBiDi::ETextDirection InTextDirection);
+	SLATE_API FShapedGlyphSequenceRef FindOrAddShapedText(const FCachedShapedTextKey& InKey, const TCHAR* InText, TArray<TWeakPtr<FFreeTypeFace>>* OutLoadingFontFaces = nullptr);
+	SLATE_API FShapedGlyphSequenceRef FindOrAddShapedText(const FCachedShapedTextKey& InKey, const TCHAR* InText, const TextBiDi::ETextDirection InTextDirection, TArray<TWeakPtr<FFreeTypeFace>>* OutLoadingFontFaces = nullptr);
 
 	/**
 	 * Try and find an existing overflow shaped text instance, or add a new entry to the cache if one cannot be found.
@@ -146,10 +148,11 @@ namespace ShapedTextCacheUtil
  * @param InRunKey				The key identifying the cached shaped text for the run
  * @param InMeasureRange		The range of text that should be measured
  * @param InText				The text to shape if we can't find the shaped text in the cache. InMeasureRange may specify a sub-section of the entire text
+ * @param OutLoadingFontFaces	An array (if provided) that keeps track of the unique font faces that are loading for a text during shaping.
  *
  * @return The measured size of the shaped text
  */
-SLATE_API FVector2D MeasureShapedText(const FShapedTextCacheRef& InShapedTextCache, const FCachedShapedTextKey& InRunKey, const FTextRange& InMeasureRange, const TCHAR* InText);
+SLATE_API FVector2D MeasureShapedText(const FShapedTextCacheRef& InShapedTextCache, const FCachedShapedTextKey& InRunKey, const FTextRange& InMeasureRange, const TCHAR* InText, TArray<TWeakPtr<FFreeTypeFace>>* OutLoadingFontFaces = nullptr);
 
 /**
  * Find the character index at the specified position in pixels along the string horizontally
@@ -171,10 +174,11 @@ SLATE_API int32 FindCharacterIndexAtOffset(const FShapedTextCacheRef& InShapedTe
  * @param InRunKey				The key identifying the cached shaped text for the run
  * @param InGlyphIndex			The index of the glyph to get the kerning for (will get it between the given glyph, and it's next glyph)
  * @param InText				The text to shape if we can't find the shaped text in the cache. (InGlyphIndex, InGlyphIndex+1) will be the range used
+ * @param OutLoadingFontFaces	An array (if provided) that keeps track of the unique font faces that are loading for a text during shaping.
  *
  * @return The kerning
  */
-SLATE_API int8 GetShapedGlyphKerning(const FShapedTextCacheRef& InShapedTextCache, const FCachedShapedTextKey& InRunKey, const int32 InGlyphIndex, const TCHAR* InText);
+SLATE_API int8 GetShapedGlyphKerning(const FShapedTextCacheRef& InShapedTextCache, const FCachedShapedTextKey& InRunKey, const int32 InGlyphIndex, const TCHAR* InText, TArray<TWeakPtr<FFreeTypeFace>>* OutLoadingFontFaces = nullptr);
 
 /**
  * Extract a sub-section of a run of text into its own shaped glyph sequence
@@ -184,9 +188,10 @@ SLATE_API int8 GetShapedGlyphKerning(const FShapedTextCacheRef& InShapedTextCach
  * @param InTextRange			The range of text that should be extracted into its own shaped glyph sequence
  * @param InText				The text to shape if we can't find the shaped text in the cache. InTextRange may specify a sub-section of the entire text
  * @param InTextDirection		The text direction of all of the text to be shaped
+ * @param OutLoadingFontFaces	An array (if provided) that keeps track of the unique font faces that are loading for a text during shaping.
  *
  * @return The shaped text
  */
-SLATE_API FShapedGlyphSequenceRef GetShapedTextSubSequence(const FShapedTextCacheRef& InShapedTextCache, const FCachedShapedTextKey& InRunKey, const FTextRange& InTextRange, const TCHAR* InText, const TextBiDi::ETextDirection InTextDirection);
+SLATE_API FShapedGlyphSequenceRef GetShapedTextSubSequence(const FShapedTextCacheRef& InShapedTextCache, const FCachedShapedTextKey& InRunKey, const FTextRange& InTextRange, const TCHAR* InText, const TextBiDi::ETextDirection InTextDirection, TArray<TWeakPtr<FFreeTypeFace>>* OutLoadingFontFaces = nullptr);
 
 }
diff --git a/Engine/Source/Runtime/Slate/Public/Framework/Text/SlateAsyncFontLoadingTypes.h b/Engine/Source/Runtime/Slate/Public/Framework/Text/SlateAsyncFontLoadingTypes.h
new file mode 100644
--- /dev/null
+++ b/Engine/Source/Runtime/Slate/Public/Framework/Text/SlateAsyncFontLoadingTypes.h
@@ -0,0 +1,34 @@
+// Copyright Epic Games, Inc. All Rights Reserved.
+#pragma once
+
+#include "CoreMinimal.h"
+#include "Types/ISlateMetaData.h"
+#include "SlateAsyncFontLoadingTypes.generated.h"
+
+
+/** The policy for how to handle painting the text in text widgets if there are font faces that are still async loading while painting.*/
+UENUM()
+enum class EFontFacesLoadingPaintPolicy : uint8
+{
+	/**
+	 * The text will not be painted if any of the associated font faces are async loading.
+	 * This means the text widgets can appear blank until all of the font faces are finished async loading. 
+	 */
+	DoNotPaint,
+	/**
+	 * All glyphs supported by loaded font faces will be painted.
+	 * All glyphs that are missing will have a either a fallback character empty space painted while the font face is async loading.
+	 * This is best efffort until all font faces are finished loading and can result in visual artifacts. 
+	 */
+	PaintAvailableGlyphs
+};
+
+/** Metadata that holds the delegates if a widget wants to handle async font loading. */
+class FSlateTextLayoutPendingFontFacesLoadingMetaData : public ISlateMetaData
+{
+public:
+	SLATE_METADATA_TYPE(FSlateTextLayoutPendingFontFacesLoadingMetaData, ISlateMetaData)
+
+	
+	FSimpleDelegate OnAllFontFacesFinishLoading;
+};
diff --git a/Engine/Source/Runtime/Slate/Public/Framework/Text/SlateTextLayout.h b/Engine/Source/Runtime/Slate/Public/Framework/Text/SlateTextLayout.h
--- a/Engine/Source/Runtime/Slate/Public/Framework/Text/SlateTextLayout.h
+++ b/Engine/Source/Runtime/Slate/Public/Framework/Text/SlateTextLayout.h
@@ -9,9 +9,11 @@
 #include "Widgets/SWidget.h"
 #include "Layout/Children.h"
 
+
 class FArrangedChildren;
 class FPaintArgs;
 class FSlateWindowElementList;
+enum class EFontFacesLoadingPaintPolicy : uint8;
 
 class FSlateTextLayout : public FTextLayout
 {
@@ -32,6 +34,14 @@ public:
 
 	SLATE_API void SetIsPassword(const TAttribute<bool>& InIsPassword);
 
+	/** Sets the delegate that is Called when all font faces for the text have finished loading asynchronously. */
+	SLATE_API void SetOnAllFontFacesFinishLoading(const FSimpleDelegate& InAllFontFacesFinishLoading);
+
+	/** Gets the paint policy to be used for text if the font faces are still async loading. */
+	SLATE_API EFontFacesLoadingPaintPolicy GetFontFacesLoadingPaintPolicy() const;
+
+	/** Sets the paint policy to use for text if the font faces are still async loading. */
+	SLATE_API void SetFontFacesLoadingPaintPolicy(EFontFacesLoadingPaintPolicy InFontFacesLoadingPaintPolicy);
 protected:
 
 	SLATE_API FSlateTextLayout(SWidget* InOwner, FTextBlockStyle InDefaultTextStyle);
@@ -42,6 +52,7 @@ protected:
 
 	SLATE_API virtual TSharedRef<IRun> CreateDefaultTextRun(const TSharedRef<FString>& NewText, const FTextRange& NewRange) const override;
 
+	virtual void HandleAllFontFacesFinishLoading() override;
 protected:
 	/** Default style used by the TextLayout */
 	FTextBlockStyle DefaultTextStyle;
@@ -53,5 +64,7 @@ private:
 	/** This this layout displaying a password? */
 	TAttribute<bool> bIsPassword;
 
+	/** The paint policy to use for the text layout if any of the font faces are still async loading. */
+	EFontFacesLoadingPaintPolicy FontFacesLoadingPaintPolicy;
 	friend class FSlateTextLayoutFactory;
 };
diff --git a/Engine/Source/Runtime/Slate/Public/Framework/Text/TextLayout.h b/Engine/Source/Runtime/Slate/Public/Framework/Text/TextLayout.h
--- a/Engine/Source/Runtime/Slate/Public/Framework/Text/TextLayout.h
+++ b/Engine/Source/Runtime/Slate/Public/Framework/Text/TextLayout.h
@@ -565,6 +565,8 @@
 
 	SLATE_API FTextSelection GetWordAt(const FTextLocation& Location) const;
 
+	/** Returns true if there are any font faces that are currently loading for the text layout . Else returns false. */
+	SLATE_API bool AreAnyFontFacesLoading() const;
 protected:
 
 	SLATE_API FTextLayout();
@@ -648,6 +650,10 @@
 	 */
 	virtual TSharedRef<IRun> CreateDefaultTextRun(const TSharedRef<FString>& NewText, const FTextRange& NewRange) const = 0;
 
+	/**
+	 * Called when it is detected during a layout pass that all async loading font faces for this text layout have finished loading.
+	 */
+	virtual void HandleAllFontFacesFinishLoading() {}
 private:
 
 	SLATE_API float GetWrappingDrawWidth() const;
@@ -710,7 +716,11 @@
 		float Height;
 	};
 
-	int32 CachedLayoutGeneration = 0;
+	/** The indices for all of the line views that require justification. */
+	TSet< int32 > LineViewsToJustify;
 
+	/** Information given to use by our an external source (typically our owner widget) to help identify who owns this text layout in the case of an error */
+	TAttribute<FString> DebugSourceInfo;
+
 	/** The models for the lines of text. A LineModel represents a single string with no manual breaks. */
 	TArray< FLineModel > LineModels;
@@ -718,54 +728,18 @@
 	/** The views for the lines of text. A LineView represents a single visual line of text. Multiple LineViews can map to the same LineModel, if for example wrapping occurs. */
 	TArray< FLineView > LineViews;
 
-	/** The indices for all of the line views that require justification. */
-	TSet< int32 > LineViewsToJustify;
+	/** The array of unique font faces still loading for the text. */
+	TArray<TWeakPtr<class FFreeTypeFace>> LoadingFontFaces;
 
-	/** Whether parameters on the layout have changed which requires the view be updated. */
-	ETextLayoutDirtyState::Flags DirtyFlags;
-
-	/** The method used to shape the text within this layout */
-	ETextShapingMethod TextShapingMethod;
-
-	/** How should the text within this layout be flowed? */
-	ETextFlowDirection TextFlowDirection;
-
-	/** The scale to draw the text at */
-	float Scale;
-
-	/** The width that the text should be wrap at. If 0 or negative no wrapping occurs. */
-	float WrappingWidth;
-
-	/** The wrapping policy used by this text layout. */
-	ETextWrappingPolicy WrappingPolicy;
-
-	/** The transform policy used by this text layout. */
-	ETextTransformPolicy TransformPolicy;
-
-	/** The size of the margins to put about the text. This is an unscaled value. */
-	FMargin Margin;
-
-	/** How the text should be aligned with the margin. */
-	ETextJustify::Type Justification;
-
-	/** The percentage to modify a line height by. */
-	float LineHeightPercentage;
-
-	/** Whether or not line height should be applied to the last line. */
-	bool ApplyLineHeightToBottomLine;
-
-	/** The final size of the text layout on screen. */
-	FTextLayoutSize TextLayoutSize;
-
-	/** Extra height of the last line due to line height. */
-	float OverHeight;
-
 	/** The size of the text layout that can actually be seen from the parent widget */
 	FVector2D ViewSize;
 
 	/** The scroll offset of the text layout from the parent widget */
 	FVector2D ScrollOffset;
 
+	/** The size of the margins to put about the text. This is an unscaled value. */
+	FMargin Margin;
+
 	/** The iterator to use to detect appropriate soft-wrapping points for lines */
 	TSharedPtr<IBreakIterator> LineBreakIterator;
 
@@ -775,13 +749,50 @@
 	/** The iterator to use to detect word boundaries */
 	TSharedRef<IBreakIterator> WordBreakIterator;
 
+	/** The final size of the text layout on screen. */
+	FTextLayoutSize TextLayoutSize;
+
+	/** How the text should be aligned with the margin. */
+	ETextJustify::Type Justification;
+
 	/** Unicode BiDi text detection */
 	TUniquePtr<TextBiDi::ITextBiDi> TextBiDiDetection;
 
-	/** Information given to use by our an external source (typically our owner widget) to help identify who owns this text layout in the case of an error */
-	TAttribute<FString> DebugSourceInfo;
+	int32 CachedLayoutGeneration = 0;
 
+	/** The scale to draw the text at */
+	float Scale;
+
+	/** The width that the text should be wrap at. If 0 or negative no wrapping occurs. */
+	float WrappingWidth;
+
+	/** The percentage to modify a line height by. */
+	float LineHeightPercentage;
+
+	/** Extra height of the last line due to line height. */
+	float OverHeight;
+
 	/** Override for the text overflow policy. If unset, the style is used */
 	TOptional<ETextOverflowPolicy> TextOverflowPolicyOverride;
 
+	/** Whether parameters on the layout have changed which requires the view be updated. */
+	ETextLayoutDirtyState::Flags DirtyFlags;
+
+	/** The method used to shape the text within this layout */
+	ETextShapingMethod TextShapingMethod;
+
+	/** How should the text within this layout be flowed? */
+	ETextFlowDirection TextFlowDirection;
+
+	/** The wrapping policy used by this text layout. */
+	ETextWrappingPolicy WrappingPolicy;
+
+	/** The transform policy used by this text layout. */
+	ETextTransformPolicy TransformPolicy;
+
+	/** Whether or not line height should be applied to the last line. */
+	bool ApplyLineHeightToBottomLine : 1;
+
+	/** True if there were any font faces that were detected to be async loading in the previous layout pass. Else false. */
+	bool bLoadingFontFacesOnPreviousLayout : 1;
 };
diff --git a/Engine/Source/Runtime/Slate/Public/Widgets/Input/SEditableText.h b/Engine/Source/Runtime/Slate/Public/Widgets/Input/SEditableText.h
--- a/Engine/Source/Runtime/Slate/Public/Widgets/Input/SEditableText.h
+++ b/Engine/Source/Runtime/Slate/Public/Widgets/Input/SEditableText.h
@@ -17,6 +17,7 @@
 #include "Framework/SlateDelegates.h"
 #include "Framework/MultiBox/MultiBoxExtender.h"
 #include "Framework/Text/TextLayout.h"
+#include "Framework/Text/SlateAsyncFontLoadingTypes.h"
 #include "Widgets/Text/ISlateEditableTextWidget.h"
 
 class FActiveTimerHandle;
@@ -67,6 +68,7 @@ public:
 		, _TextShapingMethod()
 		, _TextFlowDirection()
 		, _OverflowPolicy()
+		, _FontFacesLoadingPaintPolicy(EFontFacesLoadingPaintPolicy::DoNotPaint)
 		{
 			_Clipping = EWidgetClipping::ClipToBounds;
 		}
@@ -186,6 +188,12 @@ public:
 		
 		/** Determines what happens to text that is clipped and doesnt fit within the allotted area for this widget */
 		SLATE_ARGUMENT(TOptional<ETextOverflowPolicy>, OverflowPolicy)
+
+		/** The paint policy to use for text if the font faces are still async loading. */
+		SLATE_ARGUMENT(EFontFacesLoadingPaintPolicy, FontFacesLoadingPaintPolicy)
+
+		/** Called when all font faces for the text have finished loading asynchronously */
+		SLATE_EVENT(FSimpleDelegate, OnAllFontFacesFinishLoading)
 	SLATE_END_ARGS()
 
 	/** Constructor */
@@ -360,6 +368,8 @@ public:
 	/** Sets the overflow policy for this text block */
 	SLATE_API void SetOverflowPolicy(TOptional<ETextOverflowPolicy> InOverflowPolicy);
 
+	/** Sets the paint policy to use for text if the font faces are still async loading. */
+	SLATE_API void SetFontFacesLoadingPaintPolicy(EFontFacesLoadingPaintPolicy InFontFacesLoadingPaintPolicy);
 	/** Query to see if any text is selected within the document */
 	SLATE_API bool AnyTextSelected() const;
 
diff --git a/Engine/Source/Runtime/Slate/Public/Widgets/Input/SEditableTextBox.h b/Engine/Source/Runtime/Slate/Public/Widgets/Input/SEditableTextBox.h
--- a/Engine/Source/Runtime/Slate/Public/Widgets/Input/SEditableTextBox.h
+++ b/Engine/Source/Runtime/Slate/Public/Widgets/Input/SEditableTextBox.h
@@ -15,6 +15,7 @@
 #include "Widgets/Input/SEditableText.h"
 #include "Widgets/Layout/SBorder.h"
 #include "Framework/SlateDelegates.h"
+#include "Framework/Text/SlateAsyncFontLoadingTypes.h"
 
 class IErrorReportingWidget;
 class SBox;
@@ -58,6 +59,7 @@ public:
 		, _VirtualKeyboardTrigger(EVirtualKeyboardTrigger::OnFocusByPointer)
 		, _VirtualKeyboardDismissAction(EVirtualKeyboardDismissAction::TextChangeOnDismiss)
 		, _OverflowPolicy()
+		, _FontFacesLoadingPaintPolicy(EFontFacesLoadingPaintPolicy::DoNotPaint)
 		{
 		}
 
@@ -177,6 +179,12 @@ public:
 
 		/** Determines what happens to text that is clipped and doesnt fit within the allotted area for this text box */
 		SLATE_ARGUMENT(TOptional<ETextOverflowPolicy>, OverflowPolicy)
+
+		/** The paint policy to use for text if the font faces are still async loading. */
+		SLATE_ARGUMENT(EFontFacesLoadingPaintPolicy, FontFacesLoadingPaintPolicy)
+
+		/** Called when all font faces for the text have finished loading asynchronously */
+		SLATE_EVENT(FSimpleDelegate, OnAllFontFacesFinishLoading)
 	SLATE_END_ARGS()
 
 	SLATE_API SEditableTextBox();
@@ -375,6 +383,9 @@ public:
 	/** Sets the overflow policy for this text block */
 	SLATE_API void SetOverflowPolicy(TOptional<ETextOverflowPolicy> InOverflowPolicy);
 
+	/** Sets the paint policy to use for text if the font faces are still async loading. */
+	SLATE_API void SetFontFacesLoadingPaintPolicy(EFontFacesLoadingPaintPolicy InFontFacesLoadingPaintPolicy);
+
 	/** Query to see if any text is selected within the document */
 	SLATE_API bool AnyTextSelected() const;
 
diff --git a/Engine/Source/Runtime/Slate/Public/Widgets/Input/SMultiLineEditableTextBox.h b/Engine/Source/Runtime/Slate/Public/Widgets/Input/SMultiLineEditableTextBox.h
--- a/Engine/Source/Runtime/Slate/Public/Widgets/Input/SMultiLineEditableTextBox.h
+++ b/Engine/Source/Runtime/Slate/Public/Widgets/Input/SMultiLineEditableTextBox.h
@@ -14,6 +14,7 @@
 #include "Styling/CoreStyle.h"
 #include "Framework/Text/IRun.h"
 #include "Framework/Text/TextLayout.h"
+#include "Framework/Text/SlateAsyncFontLoadingTypes.h"
 #include "Widgets/Layout/SBorder.h"
 #include "Widgets/Layout/SScrollBar.h"
 #include "Widgets/Text/SMultiLineEditableText.h"
@@ -74,6 +75,7 @@ public:
 		, _TextShapingMethod()
 		, _TextFlowDirection()
 		, _OverflowPolicy()
+		, _FontFacesLoadingPaintPolicy(EFontFacesLoadingPaintPolicy::DoNotPaint)
 		{}
 
 		/** The styling of the textbox */
@@ -248,6 +250,12 @@ public:
 
 		/** Determines what happens to text that is clipped and doesn't fit within the allotted area for this widget */
 		SLATE_ARGUMENT(TOptional<ETextOverflowPolicy>, OverflowPolicy)
+
+		/** The paint policy to use for text if the font faces are still async loading. */
+		SLATE_ARGUMENT(EFontFacesLoadingPaintPolicy, FontFacesLoadingPaintPolicy)
+
+		/** Called when all font faces for the text have finished loading asynchronously */
+		SLATE_EVENT(FSimpleDelegate, OnAllFontFacesFinishLoading)
 	SLATE_END_ARGS()
 	
 	/**
@@ -394,6 +402,9 @@ public:
 	/** Set the ReadOnly attribute */
 	SLATE_API void SetIsReadOnly(const TAttribute< bool >& InIsReadOnly);
 
+	/** Set the FontFacesLoadingPaintPolicy argument */
+	SLATE_API void SetFontFacesLoadingPaintPolicy( EFontFacesLoadingPaintPolicy InFontFacesLoadingPaintPolicy);
+
 	/**
 	 * If InError is a non-empty string the TextBox will the ErrorReporting provided during construction
 	 * If no error reporting was provided, the TextBox will create a default error reporter.
diff --git a/Engine/Source/Runtime/Slate/Public/Widgets/Text/SMultiLineEditableText.h b/Engine/Source/Runtime/Slate/Public/Widgets/Text/SMultiLineEditableText.h
--- a/Engine/Source/Runtime/Slate/Public/Widgets/Text/SMultiLineEditableText.h
+++ b/Engine/Source/Runtime/Slate/Public/Widgets/Text/SMultiLineEditableText.h
@@ -15,6 +15,7 @@
 #include "Styling/CoreStyle.h"
 #include "Framework/Text/IRun.h"
 #include "Framework/Text/TextLayout.h"
+#include "Framework/Text/SlateAsyncFontLoadingTypes.h"
 #include "Widgets/SWidget.h"
 #include "Widgets/Layout/SScrollBar.h"
 #include "Framework/SlateDelegates.h"
@@ -97,6 +98,7 @@ public:
 		, _TextShapingMethod()
 		, _TextFlowDirection()
 		, _OverflowPolicy()
+		, _FontFacesLoadingPaintPolicy(EFontFacesLoadingPaintPolicy::DoNotPaint)
 	{
 		_Clipping = EWidgetClipping::ClipToBounds;
 	}
@@ -234,6 +236,11 @@ public:
 		/** Determines what happens to text that is clipped and doesnt fit within the clip rect for this widget */
 		SLATE_ARGUMENT(TOptional<ETextOverflowPolicy>, OverflowPolicy)
 
+		/** The paint policy to use for text if the font faces are still async loading. */
+		SLATE_ARGUMENT(EFontFacesLoadingPaintPolicy, FontFacesLoadingPaintPolicy)
+
+		/** Called when all font faces for the text have finished loading asynchronously */
+		SLATE_EVENT(FSimpleDelegate, OnAllFontFacesFinishLoading)
 	SLATE_END_ARGS()
 
 	SLATE_API SMultiLineEditableText();
@@ -468,6 +475,8 @@ public:
 		OnKeyDownHandler = InOnKeyDownHandler;
 	}
 
+	/** Sets the paint policy to use for text if the font faces are still async loading. */
+	SLATE_API void SetFontFacesLoadingPaintPolicy(EFontFacesLoadingPaintPolicy InFontFacesLoadingPaintPolicy);
 	/**
 	 *	Force a single scroll operation. 
 	 */
diff --git a/Engine/Source/Runtime/Slate/Public/Widgets/Text/SRichTextBlock.h b/Engine/Source/Runtime/Slate/Public/Widgets/Text/SRichTextBlock.h
--- a/Engine/Source/Runtime/Slate/Public/Widgets/Text/SRichTextBlock.h
+++ b/Engine/Source/Runtime/Slate/Public/Widgets/Text/SRichTextBlock.h
@@ -11,6 +11,7 @@
 #include "Styling/SlateTypes.h"
 #include "Styling/CoreStyle.h"
 #include "Framework/Text/TextLayout.h"
+#include "Framework/Text/SlateAsyncFontLoadingTypes.h"
 #include "Framework/Text/SlateHyperlinkRun.h"
 #include "Framework/Text/ITextDecorator.h"
 #include "Framework/Text/TextDecorators.h"
@@ -54,6 +55,7 @@ public:
 		, _Parser()
 		, _OverflowPolicy()
 		, _MinDesiredWidth()
+		, _FontFacesLoadingPaintPolicy(EFontFacesLoadingPaintPolicy::DoNotPaint)
 	{
 		_Clipping = EWidgetClipping::OnDemand;
 	}
@@ -119,6 +121,12 @@ public:
 		/** Minimum width that this text block should be */
 		SLATE_ATTRIBUTE(float, MinDesiredWidth)
 
+		/** The paint policy to use for text if the font faces are still async loading. */
+		SLATE_ARGUMENT(EFontFacesLoadingPaintPolicy, FontFacesLoadingPaintPolicy)
+
+		/** Called when all font faces for the text have finished loading asynchronously */
+		SLATE_EVENT(FSimpleDelegate, OnAllFontFacesFinishLoading)
+
 		/** Additional decorators can be append to the widget inline. Inline decorators get precedence over decorators not specified inline. */
 		TArray< TSharedRef< ITextDecorator > > InlineDecorators;
 		SRichTextBlock::FArguments& operator + (const TSharedRef< ITextDecorator >& DecoratorToAdd)
@@ -243,6 +251,8 @@ public:
 	/** set the scale value at the TextLayout*/
 	SLATE_API void SetTextBlockScale(const float NewTextBlockScale);
 
+	/** Set the paint policy for the text layout if font faces are still async loading.*/
+	SLATE_API void SetFontFacesLoadingPaintPolicy(EFontFacesLoadingPaintPolicy InFontFacesLoadingPaintPolicy);
 protected:
 	//~ SWidget interface
 	SLATE_API virtual FVector2D ComputeDesiredSize(float) const override;
diff --git a/Engine/Source/Runtime/Slate/Public/Widgets/Text/STextBlock.h b/Engine/Source/Runtime/Slate/Public/Widgets/Text/STextBlock.h
--- a/Engine/Source/Runtime/Slate/Public/Widgets/Text/STextBlock.h
+++ b/Engine/Source/Runtime/Slate/Public/Widgets/Text/STextBlock.h
@@ -15,6 +15,7 @@
 #include "Styling/CoreStyle.h"
 #include "Framework/SlateDelegates.h"
 #include "Framework/Text/TextLayout.h"
+#include "Framework/Text/SlateAsyncFontLoadingTypes.h"
 
 class FPaintArgs;
 class FSlateWindowElementList;
@@ -22,6 +23,7 @@ class FSlateTextBlockLayout;
 class IBreakIterator;
 enum class ETextShapingMethod : uint8;
 
+
 namespace ETextRole
 {
 	enum Type
@@ -72,6 +74,7 @@ public:
 		, _LineBreakPolicy()
 		, _OverflowPolicy()
 		, _SimpleTextMode(false)
+		, _FontFacesLoadingPaintPolicy(EFontFacesLoadingPaintPolicy::DoNotPaint)
 		{
 			_Clipping = EWidgetClipping::OnDemand;
 		}
@@ -157,6 +160,11 @@ public:
 		/** Called when this text is double clicked */
 		SLATE_EVENT(FPointerEventHandler, OnDoubleClicked)
 
+		/** The paint policy to use for text if the font faces are still async loading. */
+		SLATE_ARGUMENT(EFontFacesLoadingPaintPolicy, FontFacesLoadingPaintPolicy)
+
+		/** Called when all font faces for the text have finished loading asynchronously */
+		SLATE_EVENT(FSimpleDelegate, OnAllFontFacesFinishLoading)
 	SLATE_END_ARGS()
 
 	/** Constructor */
@@ -254,6 +262,8 @@ public:
 	/** See Justification attribute */
 	SLATE_API void SetJustification(TAttribute<ETextJustify::Type> InJustification);
 
+	/** See FontFacesLoadingPaintPolicy argument */
+	SLATE_API void SetFontFacesLoadingPaintPolicy(EFontFacesLoadingPaintPolicy InFontFacesLoadingPaintPolicy);
 	// SWidget interface
 	SLATE_API virtual int32 OnPaint( const FPaintArgs& Args, const FGeometry& AllottedGeometry, const FSlateRect& MyCullingRect, FSlateWindowElementList& OutDrawElements, int32 LayerId, const FWidgetStyle& InWidgetStyle, bool bParentEnabled ) const override;
 	SLATE_API virtual FVector2D ComputeDesiredSize(float) const override;
diff --git a/Engine/Source/Runtime/Slate/Public/Widgets/Text/SlateEditableTextLayout.h b/Engine/Source/Runtime/Slate/Public/Widgets/Text/SlateEditableTextLayout.h
--- a/Engine/Source/Runtime/Slate/Public/Widgets/Text/SlateEditableTextLayout.h
+++ b/Engine/Source/Runtime/Slate/Public/Widgets/Text/SlateEditableTextLayout.h
@@ -30,6 +30,7 @@ class IBreakIterator;
 class ISlateEditableTextWidget;
 class SWindow;
 enum class ETextShapingMethod : uint8;
+enum class EFontFacesLoadingPaintPolicy : uint8;
 
 /** Class to handle the cached layout of SEditableText/SMultiLineEditableText by proxying around a FTextLayout */
 class FSlateEditableTextLayout
@@ -419,6 +420,11 @@ public:
 	/* Function to spawn virtual keyboard on demand */
 	SLATE_API void ToggleVirtualKeyboard(bool bShow, const int32 UserIndex);
 
+	/** Sets the delegate for the internal text layout that is Called when all font faces for the text have finished loading asynchronously. */
+	SLATE_API void SetOnAllFontFacesFinishLoadingDelegate(const FSimpleDelegate& InAllFontFacesFinishLoadingDelegate);
+
+	/** Sets the paint policy to use for text if the font faces are still async loading. */
+	SLATE_API void SetFontFacesLoadingPaintPolicy(EFontFacesLoadingPaintPolicy InFontFacesLoadingPaintPolicy);
 private:
 	/** Insert the given text at the current cursor position, correctly taking into account new line characters */
 	SLATE_API void InsertTextAtCursorImpl(const FString& InString);
diff --git a/Engine/Source/Runtime/Slate/Public/Widgets/Text/SlateTextBlockLayout.h b/Engine/Source/Runtime/Slate/Public/Widgets/Text/SlateTextBlockLayout.h
--- a/Engine/Source/Runtime/Slate/Public/Widgets/Text/SlateTextBlockLayout.h
+++ b/Engine/Source/Runtime/Slate/Public/Widgets/Text/SlateTextBlockLayout.h
@@ -15,6 +15,7 @@ class FSlateWindowElementList;
 class IBreakIterator;
 class ISlateRunRenderer;
 enum class ETextShapingMethod : uint8;
+enum class EFontFacesLoadingPaintPolicy : uint8;
 
 /** Class to handle the cached layout of STextBlock/SRichTextBlock by proxying around a FTextLayout */
 class FSlateTextBlockLayout
@@ -230,6 +231,15 @@ public:
 	 */
 	SLATE_API void ArrangeChildren(const FGeometry& AllottedGeometry, FArrangedChildren& ArrangedChildren) const;
 
+	/**
+	 * Sets the delegate for the internal text layout that is Called when all font faces for the text have finished loading asynchronously.
+	 */
+	SLATE_API void SetOnAllFontFacesFinishLoading(const FSimpleDelegate& InAllFontFacesFinishLoading);
+
+	/**
+	 * Sets the paint policy to use for text if the font faces are still async loading.
+	 */
+	SLATE_API void SetFontFacesLoadingPaintPolicy(EFontFacesLoadingPaintPolicy InFontFacesLoadingPaintPolicy);
 private:
 	/** Updates the text layout to contain the given text */
 	void UpdateTextLayout(const FText& InText);
diff --git a/Engine/Source/Runtime/SlateCore/Private/Fonts/FontCache.cpp b/Engine/Source/Runtime/SlateCore/Private/Fonts/FontCache.cpp
--- a/Engine/Source/Runtime/SlateCore/Private/Fonts/FontCache.cpp
+++ b/Engine/Source/Runtime/SlateCore/Private/Fonts/FontCache.cpp
@@ -1221,24 +1221,24 @@ bool FSlateFontCache::AddNewEntry( const FCharacterRenderData InRenderData, uint
 	return NewSlot != nullptr;
 }
 
-FShapedGlyphSequenceRef FSlateFontCache::ShapeBidirectionalText( const FString& InText, const FSlateFontInfo &InFontInfo, const float InFontScale, const TextBiDi::ETextDirection InBaseDirection, const ETextShapingMethod InTextShapingMethod ) const
+FShapedGlyphSequenceRef FSlateFontCache::ShapeBidirectionalText( const FString& InText, const FSlateFontInfo &InFontInfo, const float InFontScale, const TextBiDi::ETextDirection InBaseDirection, const ETextShapingMethod InTextShapingMethod , TArray<TWeakPtr<FFreeTypeFace>>* OutLoadingFontFaces) const
 {
-	return ShapeBidirectionalText(*InText, 0, InText.Len(), InFontInfo, InFontScale, InBaseDirection, InTextShapingMethod);
+	return ShapeBidirectionalText(*InText, 0, InText.Len(), InFontInfo, InFontScale, InBaseDirection, InTextShapingMethod, OutLoadingFontFaces);
 }
 
-FShapedGlyphSequenceRef FSlateFontCache::ShapeBidirectionalText( const TCHAR* InText, const int32 InTextStart, const int32 InTextLen, const FSlateFontInfo &InFontInfo, const float InFontScale, const TextBiDi::ETextDirection InBaseDirection, const ETextShapingMethod InTextShapingMethod ) const
+FShapedGlyphSequenceRef FSlateFontCache::ShapeBidirectionalText( const TCHAR* InText, const int32 InTextStart, const int32 InTextLen, const FSlateFontInfo &InFontInfo, const float InFontScale, const TextBiDi::ETextDirection InBaseDirection, const ETextShapingMethod InTextShapingMethod , TArray<TWeakPtr<FFreeTypeFace>>* OutLoadingFontFaces) const
 {
-	return TextShaper->ShapeBidirectionalText(InText, InTextStart, InTextLen, InFontInfo, InFontScale, InBaseDirection, InTextShapingMethod);
+	return TextShaper->ShapeBidirectionalText(InText, InTextStart, InTextLen, InFontInfo, InFontScale, InBaseDirection, InTextShapingMethod, OutLoadingFontFaces);
 }
 
-FShapedGlyphSequenceRef FSlateFontCache::ShapeUnidirectionalText( const FString& InText, const FSlateFontInfo &InFontInfo, const float InFontScale, const TextBiDi::ETextDirection InTextDirection, const ETextShapingMethod InTextShapingMethod ) const
+FShapedGlyphSequenceRef FSlateFontCache::ShapeUnidirectionalText( const FString& InText, const FSlateFontInfo &InFontInfo, const float InFontScale, const TextBiDi::ETextDirection InTextDirection, const ETextShapingMethod InTextShapingMethod , TArray<TWeakPtr<FFreeTypeFace>>* OutLoadingFontFaces) const
 {
-	return ShapeUnidirectionalText(*InText, 0, InText.Len(), InFontInfo, InFontScale, InTextDirection, InTextShapingMethod);
+	return ShapeUnidirectionalText(*InText, 0, InText.Len(), InFontInfo, InFontScale, InTextDirection, InTextShapingMethod, OutLoadingFontFaces);
 }
 
-FShapedGlyphSequenceRef FSlateFontCache::ShapeUnidirectionalText( const TCHAR* InText, const int32 InTextStart, const int32 InTextLen, const FSlateFontInfo &InFontInfo, const float InFontScale, const TextBiDi::ETextDirection InTextDirection, const ETextShapingMethod InTextShapingMethod ) const
+FShapedGlyphSequenceRef FSlateFontCache::ShapeUnidirectionalText( const TCHAR* InText, const int32 InTextStart, const int32 InTextLen, const FSlateFontInfo &InFontInfo, const float InFontScale, const TextBiDi::ETextDirection InTextDirection, const ETextShapingMethod InTextShapingMethod , TArray<TWeakPtr<FFreeTypeFace>>* OutLoadingFontFaces) const
 {
-	return TextShaper->ShapeUnidirectionalText(InText, InTextStart, InTextLen, InFontInfo, InFontScale, InTextDirection, InTextShapingMethod);
+	return TextShaper->ShapeUnidirectionalText(InText, InTextStart, InTextLen, InFontInfo, InFontScale, InTextDirection, InTextShapingMethod, OutLoadingFontFaces);
 }
 
 FCharacterList& FSlateFontCache::GetCharacterList( const FSlateFontInfo &InFontInfo, float FontScale, const FFontOutlineSettings& InOutlineSettings )
diff --git a/Engine/Source/Runtime/SlateCore/Private/Fonts/FontCacheCompositeFont.cpp b/Engine/Source/Runtime/SlateCore/Private/Fonts/FontCacheCompositeFont.cpp
--- a/Engine/Source/Runtime/SlateCore/Private/Fonts/FontCacheCompositeFont.cpp
+++ b/Engine/Source/Runtime/SlateCore/Private/Fonts/FontCacheCompositeFont.cpp
@@ -18,7 +18,20 @@ static bool GAsyncFontLazyLoad = false;
 FAutoConsoleVariableRef CVarAsyncLazyLoad(
 	TEXT("Slate.Font.AsyncLazyLoad"),
 	GAsyncFontLazyLoad,
-	TEXT("Causes unloaded font faces that are lazily loaded, to be loaded asynchronously, until then the font won't measure correctly.  Once complete the UI will invalidate."));
+	TEXT("Causes unloaded font faces that are lazily loaded, to be loaded asynchronously, until then the font won't measure correctly.  Once complete the UI will invalidate. Note that async font loading is disallowed in Editor/PIE, standalone and cooked builds are fine. "));
+
+namespace UE::Private::FontCacheCompositeFont
+{
+bool IsAsyncFontLoadingAllowed()
+{
+	// We only allow async font loading in standalone or cooked builds, not in Editor/PIE. We want everything to be synchronously loaded for the actual Editor
+#if WITH_EDITOR
+	return GAsyncFontLazyLoad && !GIsEditor;
+#else
+	return GAsyncFontLazyLoad;
+#endif
+}
+}
 
 DECLARE_CYCLE_STAT(TEXT("Load Font"), STAT_SlateLoadFont, STATGROUP_Slate);
 
@@ -39,6 +52,10 @@ public:
 		{
 			UE_LOG(LogSlate, Warning, TEXT("FAsyncLoadFontFaceData failed to load or process '%s' with face index %d"), *FontData.GetFontFilename(), FontData.GetSubFaceIndex());
 		}
+		else
+		{
+			UE_LOG(LogSlate, Display, TEXT("FAsyncLoadFontFaceData successfully loaded '%s' with face index %d"), *FontData.GetFontFilename(), FontData.GetSubFaceIndex());
+		}
 	}
 
 	FORCEINLINE TStatId GetStatId() const
@@ -510,7 +527,7 @@ TSharedPtr<FFreeTypeFace> FCompositeFontCache::GetFontFace(const FFontData& InFo
 			{
 			case EFontLoadingPolicy::LazyLoad:
 				{
-					if (GAsyncFontLazyLoad)
+					if (UE::Private::FontCacheCompositeFont::IsAsyncFontLoadingAllowed())
 					{
 						FAsyncTask<FAsyncLoadFontFaceData>* AsyncFontLoadTaskPtr = new FAsyncTask<FAsyncLoadFontFaceData>(InFontData);
 						LoadFontFaceTasks.Add(AsyncFontLoadTaskPtr);
diff --git a/Engine/Source/Runtime/SlateCore/Private/Fonts/SlateFontRenderer.cpp b/Engine/Source/Runtime/SlateCore/Private/Fonts/SlateFontRenderer.cpp
--- a/Engine/Source/Runtime/SlateCore/Private/Fonts/SlateFontRenderer.cpp
+++ b/Engine/Source/Runtime/SlateCore/Private/Fonts/SlateFontRenderer.cpp
@@ -39,6 +39,16 @@ FAutoConsoleVariableRef CVarEnableFontAntiAliasing(
 	EnableFontAntiAliasing,
 	TEXT("Enable or disable anti-aliasing for font rendering (0 = off, 1 = on). Enabled by default."));
 
+/**
+ * Enable or disable detection of async loading font faces. This impacts whether or not FTextLayout, FSlateTextLayout etc would be able to detect if font faces are loading asynchronously and respond accordingly.
+ * See FTextLayout::AreAnyFontFacesLoading() and FSlateTextLayout::OnPaint() to see how the detection of async loading font faces for a text widget has implications on rendering of text and execution of delegates.
+ */
+static bool GDetectAsyncLoadingFontFaces = false;
+FAutoConsoleVariableRef CVarDetectAsyncLoadingFontFaces(
+	TEXT("Slate.Font.DetectAsyncLoadingFontFaces"),
+	GDetectAsyncLoadingFontFaces,
+	TEXT("If set to true, text widgets will be able to detect if any of its font faces are loading asynchornously and respond accordingly.. This does nothing if Slate.Font.AsyncLazyLoad is set to false."));
+
 namespace SlateFontRendererUtils
 {
 
@@ -197,7 +207,7 @@ uint16 FSlateFontRenderer::GetMaxHeight(const FSlateFontInfo& InFontInfo, const
 	const FFontData& FontData = CompositeFontCache->GetDefaultFontData(InFontInfo);
 	const FFreeTypeFaceGlyphData FaceGlyphData = GetFontFaceForCodepoint(FontData, Char, InFontInfo.FontFallback);
 
-	if (FaceGlyphData.FaceAndMemory.IsValid() && FaceGlyphData.FaceAndMemory->IsFaceValid())
+	if (FaceGlyphData.FaceAndMemory.IsValid() && FaceGlyphData.FaceAndMemory->IsFaceValid() && !FaceGlyphData.FaceAndMemory->IsFaceLoading())
 	{
 		FreeTypeUtils::ApplySizeAndScale(FaceGlyphData.FaceAndMemory->GetFace(), InFontInfo.Size, InScale);
 
@@ -220,7 +230,7 @@ int16 FSlateFontRenderer::GetBaseline(const FSlateFontInfo& InFontInfo, const fl
 	const FFontData& FontData = CompositeFontCache->GetDefaultFontData(InFontInfo);
 	const FFreeTypeFaceGlyphData FaceGlyphData = GetFontFaceForCodepoint(FontData, Char, InFontInfo.FontFallback);
 
-	if (FaceGlyphData.FaceAndMemory.IsValid() && FaceGlyphData.FaceAndMemory->IsFaceValid())
+	if (FaceGlyphData.FaceAndMemory.IsValid() && FaceGlyphData.FaceAndMemory->IsFaceValid() && !FaceGlyphData.FaceAndMemory->IsFaceLoading())
 	{
 		FreeTypeUtils::ApplySizeAndScale(FaceGlyphData.FaceAndMemory->GetFace(), InFontInfo.Size, InScale);
 
@@ -335,7 +345,7 @@ bool FSlateFontRenderer::CanLoadCodepoint(const FFontData& InFontData, const UTF
 
 #if WITH_FREETYPE
 	const FFreeTypeFaceGlyphData FaceGlyphData = GetFontFaceForCodepoint(InFontData, InCodepoint, MaxFallbackLevel);
-	bReturnVal = FaceGlyphData.FaceAndMemory.IsValid() && FaceGlyphData.GlyphIndex != 0;
+	bReturnVal = FaceGlyphData.FaceAndMemory.IsValid() && !FaceGlyphData.FaceAndMemory->IsFaceLoading() && FaceGlyphData.GlyphIndex != 0;
 #endif // WITH_FREETYPE
 
 	return bReturnVal;
@@ -356,9 +366,15 @@ FFreeTypeFaceGlyphData FSlateFontRenderer::GetFontFaceForCodepoint(const FFontDa
 		// If we have valid face and memory but it just hasn't finished loading,
 		// return like we found it, so that we don't immediately trigger falling back to yet another font
 		// when it may have the glyph once it's actually done loading.
+		// If we want to detect and handle still loading font faces, we avoid clearing the face in memory so text widgets are able to react accordingly when font faces are detected to still be loading.
 		if (ReturnVal.FaceAndMemory->IsFaceLoading())
 		{
-			ReturnVal.FaceAndMemory.Reset();
+			// This would prevent the detection of fonts that are still loading while text is being passed through the shaping pass
+			// By clearing the face in memory, there is no way for text widgets to detect which font faces are loading and thus react accordingly  
+			if (!GDetectAsyncLoadingFontFaces)
+			{
+				ReturnVal.FaceAndMemory.Reset();
+			}
 			ReturnVal.GlyphIndex = 0;
 			ReturnVal.CharFallbackLevel = EFontFallback::FF_NoFallback;
 			return ReturnVal;
@@ -780,7 +796,7 @@ bool FSlateFontRenderer::GetRenderDataInternal(const FFreeTypeFaceGlyphData& InF
 FT_Face FSlateFontRenderer::GetFontFace(const FFontData& InFontData) const
 {
 	TSharedPtr<FFreeTypeFace> FaceAndMemory = CompositeFontCache->GetFontFace(InFontData);
-	return (FaceAndMemory.IsValid()) ? FaceAndMemory->GetFace() : nullptr;
+	return (FaceAndMemory.IsValid() && !FaceAndMemory->IsFaceLoading()) ? FaceAndMemory->GetFace() : nullptr;
 }
 
 #endif // WITH_FREETYPE
diff --git a/Engine/Source/Runtime/SlateCore/Private/Fonts/SlateTextShaper.cpp b/Engine/Source/Runtime/SlateCore/Private/Fonts/SlateTextShaper.cpp
--- a/Engine/Source/Runtime/SlateCore/Private/Fonts/SlateTextShaper.cpp
+++ b/Engine/Source/Runtime/SlateCore/Private/Fonts/SlateTextShaper.cpp
@@ -120,7 +120,7 @@ void FSlateTextShaper::FlushCache()
 #endif
 }
 
-FShapedGlyphSequenceRef FSlateTextShaper::ShapeBidirectionalText(const TCHAR* InText, const int32 InTextStart, const int32 InTextLen, const FSlateFontInfo& InFontInfo, const float InFontScale, const TextBiDi::ETextDirection InBaseDirection, const ETextShapingMethod TextShapingMethod)
+FShapedGlyphSequenceRef FSlateTextShaper::ShapeBidirectionalText(const TCHAR* InText, const int32 InTextStart, const int32 InTextLen, const FSlateFontInfo& InFontInfo, const float InFontScale, const TextBiDi::ETextDirection InBaseDirection, const ETextShapingMethod TextShapingMethod, TArray<TWeakPtr<FFreeTypeFace>>* OutLoadingFontFaces)
 {
 	SCOPE_CYCLE_COUNTER(STAT_SlateShapeBidirectionalText);
 
@@ -130,22 +130,22 @@ FShapedGlyphSequenceRef FSlateTextShaper::ShapeBidirectionalText(const TCHAR* In
 	TArray<FShapedGlyphEntry> GlyphsToRender;
 	for (const TextBiDi::FTextDirectionInfo& TextDirectionInfo : TextDirectionInfos)
 	{
-		PerformTextShaping(InText, TextDirectionInfo.StartIndex, TextDirectionInfo.Length, InFontInfo, InFontScale, TextDirectionInfo.TextDirection, TextShapingMethod, GlyphsToRender);
+		PerformTextShaping(InText, TextDirectionInfo.StartIndex, TextDirectionInfo.Length, InFontInfo, InFontScale, TextDirectionInfo.TextDirection, TextShapingMethod, GlyphsToRender, OutLoadingFontFaces);
 	}
 
 	return FinalizeTextShaping(MoveTemp(GlyphsToRender), InFontInfo, InFontScale, FShapedGlyphSequence::FSourceTextRange(InTextStart, InTextLen));
 }
 
-FShapedGlyphSequenceRef FSlateTextShaper::ShapeUnidirectionalText(const TCHAR* InText, const int32 InTextStart, const int32 InTextLen, const FSlateFontInfo& InFontInfo, const float InFontScale, const TextBiDi::ETextDirection InTextDirection, const ETextShapingMethod TextShapingMethod)
+FShapedGlyphSequenceRef FSlateTextShaper::ShapeUnidirectionalText(const TCHAR* InText, const int32 InTextStart, const int32 InTextLen, const FSlateFontInfo& InFontInfo, const float InFontScale, const TextBiDi::ETextDirection InTextDirection, const ETextShapingMethod TextShapingMethod, TArray<TWeakPtr<FFreeTypeFace>>* OutLoadingFontFaces)
 {
 	SCOPE_CYCLE_COUNTER(STAT_SlateShapeUnidirectionalText);
 
 	TArray<FShapedGlyphEntry> GlyphsToRender;
-	PerformTextShaping(InText, InTextStart, InTextLen, InFontInfo, InFontScale, InTextDirection, TextShapingMethod, GlyphsToRender);
+	PerformTextShaping(InText, InTextStart, InTextLen, InFontInfo, InFontScale, InTextDirection, TextShapingMethod, GlyphsToRender, OutLoadingFontFaces);
 	return FinalizeTextShaping(MoveTemp(GlyphsToRender), InFontInfo, InFontScale, FShapedGlyphSequence::FSourceTextRange(InTextStart, InTextLen));
 }
 
-void FSlateTextShaper::PerformTextShaping(const TCHAR* InText, const int32 InTextStart, const int32 InTextLen, const FSlateFontInfo& InFontInfo, const float InFontScale, const TextBiDi::ETextDirection InTextDirection, const ETextShapingMethod TextShapingMethod, TArray<FShapedGlyphEntry>& OutGlyphsToRender)
+void FSlateTextShaper::PerformTextShaping(const TCHAR* InText, const int32 InTextStart, const int32 InTextLen, const FSlateFontInfo& InFontInfo, const float InFontScale, const TextBiDi::ETextDirection InTextDirection, const ETextShapingMethod TextShapingMethod, TArray<FShapedGlyphEntry>& OutGlyphsToRender, TArray<TWeakPtr<FFreeTypeFace>>* OutLoadingFontFaces)
 {
 	check(InTextDirection != TextBiDi::ETextDirection::Mixed);
 
@@ -273,12 +273,12 @@ void FSlateTextShaper::PerformTextShaping(const TCHAR* InText, const int32 InTex
 
 		if (TextShapingMethod == ETextShapingMethod::FullShaping || (TextShapingMethod == ETextShapingMethod::Auto && TextRequiresFullShaping()))
 		{
-			PerformHarfBuzzTextShaping(InText, InTextStart, InTextLen, InFontInfo, InFontScale, InTextDirection, OutGlyphsToRender);
+			PerformHarfBuzzTextShaping(InText, InTextStart, InTextLen, InFontInfo, InFontScale, InTextDirection, OutGlyphsToRender, OutLoadingFontFaces);
 		}
 		else
 #endif // WITH_HARFBUZZ
 		{
-			PerformKerningOnlyTextShaping(InText, InTextStart, InTextLen, InFontInfo, InFontScale, OutGlyphsToRender);
+			PerformKerningOnlyTextShaping(InText, InTextStart, InTextLen, InFontInfo, InFontScale, OutGlyphsToRender, OutLoadingFontFaces);
 		}
 	}
 #endif // WITH_FREETYPE
@@ -296,7 +296,7 @@ FShapedGlyphSequenceRef FSlateTextShaper::FinalizeTextShaping(TArray<FShapedGlyp
 		const FFontData& FontData = CompositeFontCache->GetDefaultFontData(InFontInfo);
 		const FFreeTypeFaceGlyphData FaceGlyphData = FontRenderer->GetFontFaceForCodepoint(FontData, Char, InFontInfo.FontFallback);
 
-		if (FaceGlyphData.FaceAndMemory.IsValid() && FaceGlyphData.FaceAndMemory->IsFaceValid())
+		if (FaceGlyphData.FaceAndMemory.IsValid() && FaceGlyphData.FaceAndMemory->IsFaceValid() && !FaceGlyphData.FaceAndMemory->IsFaceLoading())
 		{
 			if (FMath::IsNearlyEqual(InFontInfo.GetClampSkew(), 0.f))
 			{
@@ -333,7 +333,7 @@ FShapedGlyphSequenceRef FSlateTextShaper::FinalizeTextShaping(TArray<FShapedGlyp
 
 #if WITH_FREETYPE
 
-void FSlateTextShaper::PerformKerningOnlyTextShaping(const TCHAR* InText, const int32 InTextStart, const int32 InTextLen, const FSlateFontInfo& InFontInfo, const float InFontScale, TArray<FShapedGlyphEntry>& OutGlyphsToRender) const
+void FSlateTextShaper::PerformKerningOnlyTextShaping(const TCHAR* InText, const int32 InTextStart, const int32 InTextLen, const FSlateFontInfo& InFontInfo, const float InFontScale, TArray<FShapedGlyphEntry>& OutGlyphsToRender, TArray<TWeakPtr<FFreeTypeFace>>* OutLoadingFontFaces) const
 {
 	SCOPE_CYCLE_COUNTER(STAT_PerformKerningOnlyTextShaping);
 
@@ -379,11 +379,16 @@ void FSlateTextShaper::PerformKerningOnlyTextShaping(const TCHAR* InText, const
 			float SubFontScalingFactor = 1.0f;
 			const FFontData* FontDataPtr = &CompositeFontCache->GetFontDataForCodepoint(InFontInfo, CurrentChar, SubFontScalingFactor);
 			FFreeTypeFaceGlyphData FaceGlyphData = FontRenderer->GetFontFaceForCodepoint(*FontDataPtr, CurrentChar, bShouldRenderAsWhitespace ? EFontFallback::FF_NoFallback : InFontInfo.FontFallback);
+			if (OutLoadingFontFaces && FaceGlyphData.FaceAndMemory.IsValid() && FaceGlyphData.FaceAndMemory->IsFaceLoading())
+			{
+				OutLoadingFontFaces->AddUnique(FaceGlyphData.FaceAndMemory);
+			}
+			
 
 			// If none of our fonts can render that character (as the fallback font may be missing), 
 			// try again with the fallback character, or a normal space if this character was supposed to 
 			// be whitespace (as we don't render whitespace anyway)
-			if (!FaceGlyphData.FaceAndMemory.IsValid())
+			if (!FaceGlyphData.FaceAndMemory.IsValid() || FaceGlyphData.FaceAndMemory->IsFaceLoading())
 			{
 				const TCHAR FallbackChar = bShouldRenderAsWhitespace ? TEXT(' ') : SlateFontRendererUtils::InvalidSubChar;
 				FontDataPtr = &CompositeFontCache->GetFontDataForCodepoint(InFontInfo, FallbackChar, SubFontScalingFactor);
@@ -391,7 +396,7 @@ void FSlateTextShaper::PerformKerningOnlyTextShaping(const TCHAR* InText, const
 			}
 
 			// Only scalable font types can use sub-font scaling
-			if (FaceGlyphData.FaceAndMemory.IsValid() && !FT_IS_SCALABLE(FaceGlyphData.FaceAndMemory->GetFace()))
+			if (FaceGlyphData.FaceAndMemory.IsValid() && !FaceGlyphData.FaceAndMemory->IsFaceLoading() && !FT_IS_SCALABLE(FaceGlyphData.FaceAndMemory->GetFace()))
 			{
 				SubFontScalingFactor = 1.0f;
 			}
@@ -417,7 +422,7 @@ void FSlateTextShaper::PerformKerningOnlyTextShaping(const TCHAR* InText, const
 		OutGlyphsToRender.Reserve(OutGlyphsToRender.Num() + InTextLen);
 		for (const FKerningOnlyTextSequenceEntry& KerningOnlyTextSequenceEntry : KerningOnlyTextSequence)
 		{
-			if (!KerningOnlyTextSequenceEntry.FaceAndMemory.IsValid())
+			if (!KerningOnlyTextSequenceEntry.FaceAndMemory.IsValid() || KerningOnlyTextSequenceEntry.FaceAndMemory->IsFaceLoading())
 			{
 				continue;
 			}
@@ -519,7 +524,7 @@ void FSlateTextShaper::PerformKerningOnlyTextShaping(const TCHAR* InText, const
 
 #if WITH_HARFBUZZ
 
-void FSlateTextShaper::PerformHarfBuzzTextShaping(const TCHAR* InText, const int32 InTextStart, const int32 InTextLen, const FSlateFontInfo& InFontInfo, const float InFontScale, const TextBiDi::ETextDirection InTextDirection, TArray<FShapedGlyphEntry>& OutGlyphsToRender)
+void FSlateTextShaper::PerformHarfBuzzTextShaping(const TCHAR* InText, const int32 InTextStart, const int32 InTextLen, const FSlateFontInfo& InFontInfo, const float InFontScale, const TextBiDi::ETextDirection InTextDirection, TArray<FShapedGlyphEntry>& OutGlyphsToRender, TArray<TWeakPtr<FFreeTypeFace>>* OutLoadingFontFaces)
 {
 	SCOPE_CYCLE_COUNTER(STAT_PerformHarfBuzzTextShaping);
 
@@ -589,8 +594,12 @@ void FSlateTextShaper::PerformHarfBuzzTextShaping(const TCHAR* InText, const int
 			const FFontData* FontDataPtr = &CompositeFontCache->GetFontDataForCodepoint(InFontInfo, CurrentCodepoint, SubFontScalingFactor);
 			FFreeTypeFaceGlyphData FaceGlyphData = FontRenderer->GetFontFaceForCodepoint(*FontDataPtr, CurrentCodepoint, bShouldRenderAsWhitespace ? EFontFallback::FF_NoFallback : InFontInfo.FontFallback);
 
+			if (OutLoadingFontFaces && FaceGlyphData.FaceAndMemory.IsValid() && FaceGlyphData.FaceAndMemory->IsFaceLoading())
+			{
+				OutLoadingFontFaces->AddUnique(FaceGlyphData.FaceAndMemory);
+			}
 			// If none of our fonts can render that character (as the fallback font may be missing), try again with the fallback character, or a normal space if this character was supposed to be whitespace (as we don't render whitespace anyway)
-			if (!FaceGlyphData.FaceAndMemory.IsValid())
+			if (!FaceGlyphData.FaceAndMemory.IsValid() || FaceGlyphData.FaceAndMemory->IsFaceLoading())
 			{
 				const UTF32CHAR FallbackChar = bShouldRenderAsWhitespace ? TEXT(' ') : SlateFontRendererUtils::InvalidSubChar;
 				FontDataPtr = &CompositeFontCache->GetFontDataForCodepoint(InFontInfo, FallbackChar, SubFontScalingFactor);
@@ -598,7 +607,7 @@ void FSlateTextShaper::PerformHarfBuzzTextShaping(const TCHAR* InText, const int
 			}
 
 			// Only scalable font types can use sub-font scaling
-			if (FaceGlyphData.FaceAndMemory.IsValid() && !FT_IS_SCALABLE(FaceGlyphData.FaceAndMemory->GetFace()))
+			if (FaceGlyphData.FaceAndMemory.IsValid() && !FaceGlyphData.FaceAndMemory->IsFaceLoading() && !FT_IS_SCALABLE(FaceGlyphData.FaceAndMemory->GetFace()))
 			{
 				SubFontScalingFactor = 1.0f;
 			}
diff --git a/Engine/Source/Runtime/SlateCore/Private/Fonts/SlateTextShaper.h b/Engine/Source/Runtime/SlateCore/Private/Fonts/SlateTextShaper.h
--- a/Engine/Source/Runtime/SlateCore/Private/Fonts/SlateTextShaper.h
+++ b/Engine/Source/Runtime/SlateCore/Private/Fonts/SlateTextShaper.h
@@ -65,24 +65,24 @@ class FSlateTextShaper
 public:
 	FSlateTextShaper(FFreeTypeCacheDirectory* InFTCacheDirectory, FCompositeFontCache* InCompositeFontCache, FSlateFontRenderer* InFontRenderer, FSlateFontCache* InFontCache);
 
-	FShapedGlyphSequenceRef ShapeBidirectionalText(const TCHAR* InText, const int32 InTextStart, const int32 InTextLen, const FSlateFontInfo& InFontInfo, const float InFontScale, const TextBiDi::ETextDirection InBaseDirection, const ETextShapingMethod TextShapingMethod);
+	FShapedGlyphSequenceRef ShapeBidirectionalText(const TCHAR* InText, const int32 InTextStart, const int32 InTextLen, const FSlateFontInfo& InFontInfo, const float InFontScale, const TextBiDi::ETextDirection InBaseDirection, const ETextShapingMethod TextShapingMethod, TArray<TWeakPtr<FFreeTypeFace>>* OutLoadingFontFaces);
 
-	FShapedGlyphSequenceRef ShapeUnidirectionalText(const TCHAR* InText, const int32 InTextStart, const int32 InTextLen, const FSlateFontInfo& InFontInfo, const float InFontScale, const TextBiDi::ETextDirection InTextDirection, const ETextShapingMethod TextShapingMethod);
+	FShapedGlyphSequenceRef ShapeUnidirectionalText(const TCHAR* InText, const int32 InTextStart, const int32 InTextLen, const FSlateFontInfo& InFontInfo, const float InFontScale, const TextBiDi::ETextDirection InTextDirection, const ETextShapingMethod TextShapingMethod, TArray<TWeakPtr<FFreeTypeFace>>* OutLoadingFontFaces);
 
 	/** Get rid of the harfbuzz font cache, wich store the harfbuzz fonts. */
 	void FlushCache();
 
 private:
-	void PerformTextShaping(const TCHAR* InText, const int32 InTextStart, const int32 InTextLen, const FSlateFontInfo& InFontInfo, const float InFontScale, const TextBiDi::ETextDirection InTextDirection, const ETextShapingMethod TextShapingMethod, TArray<FShapedGlyphEntry>& OutGlyphsToRender);
+	void PerformTextShaping(const TCHAR* InText, const int32 InTextStart, const int32 InTextLen, const FSlateFontInfo& InFontInfo, const float InFontScale, const TextBiDi::ETextDirection InTextDirection, const ETextShapingMethod TextShapingMethod, TArray<FShapedGlyphEntry>& OutGlyphsToRender, TArray<TWeakPtr<FFreeTypeFace>>* OutLoadingFontFaces);
 
 	FShapedGlyphSequenceRef FinalizeTextShaping(TArray<FShapedGlyphEntry> InGlyphsToRender, const FSlateFontInfo& InFontInfo, const float InFontScale, const FShapedGlyphSequence::FSourceTextRange& InSourceTextRange) const;
 
 #if WITH_FREETYPE
-	void PerformKerningOnlyTextShaping(const TCHAR* InText, const int32 InTextStart, const int32 InTextLen, const FSlateFontInfo& InFontInfo, const float InFontScale, TArray<FShapedGlyphEntry>& OutGlyphsToRender) const;
+	void PerformKerningOnlyTextShaping(const TCHAR* InText, const int32 InTextStart, const int32 InTextLen, const FSlateFontInfo& InFontInfo, const float InFontScale, TArray<FShapedGlyphEntry>& OutGlyphsToRender, TArray<TWeakPtr<FFreeTypeFace>>* OutLoadingFontFaces) const;
 #endif // WITH_FREETYPE
 
 #if WITH_HARFBUZZ
-	void PerformHarfBuzzTextShaping(const TCHAR* InText, const int32 InTextStart, const int32 InTextLen, const FSlateFontInfo& InFontInfo, const float InFontScale, const TextBiDi::ETextDirection InTextDirection, TArray<FShapedGlyphEntry>& OutGlyphsToRender);
+	void PerformHarfBuzzTextShaping(const TCHAR* InText, const int32 InTextStart, const int32 InTextLen, const FSlateFontInfo& InFontInfo, const float InFontScale, const TextBiDi::ETextDirection InTextDirection, TArray<FShapedGlyphEntry>& OutGlyphsToRender, TArray<TWeakPtr<FFreeTypeFace>>* OutLoadingFontFaces);
 #endif // WITH_HARFBUZZ
 
 	static bool InsertSubstituteGlyphs(const TCHAR* InText, const int32 InCharIndex, const TSharedRef<FShapedGlyphFaceData>& InShapedGlyphFaceData, const TSharedRef<FFreeTypeAdvanceCache>& AdvanceCache, TArray<FShapedGlyphEntry>& OutGlyphsToRender, const int16 InLetterSpacingScaled = 0);
diff --git a/Engine/Source/Runtime/SlateCore/Public/Fonts/FontCache.h b/Engine/Source/Runtime/SlateCore/Public/Fonts/FontCache.h
--- a/Engine/Source/Runtime/SlateCore/Public/Fonts/FontCache.h
+++ b/Engine/Source/Runtime/SlateCore/Public/Fonts/FontCache.h
@@ -597,6 +597,7 @@ private:
 	int32 SequenceWidth;
 	/** The set of fonts being used by the glyphs within this sequence */
 	TArray<TWeakPtr<FFreeTypeFace>> GlyphFontFaces;
+
 	/** A map of source indices to their shaped glyph data indices - used to perform efficient reverse look-up */
 	FSourceIndicesToGlyphData SourceIndicesToGlyphData;
 	/** The Font Skew parameter of the FontInfo applied to the this Shaped glype sequence */
@@ -785,9 +786,10 @@ public:
 	 * @param InFontScale			The scale to apply to the font
 	 * @param InBaseDirection		The overall reading direction of the text (see TextBiDi::ComputeBaseDirection). This will affect where some characters (such as brackets and quotes) are placed within the resultant shaped text
 	 * @param InTextShapingMethod	The text shaping method to use
+	 * @param OutLoadingFontFaces	An array (if provided) that keeps track of the unique font faces that are loading for a text during shaping.
 	 */
-	SLATECORE_API FShapedGlyphSequenceRef ShapeBidirectionalText( const FString& InText, const FSlateFontInfo &InFontInfo, const float InFontScale, const TextBiDi::ETextDirection InBaseDirection, const ETextShapingMethod InTextShapingMethod ) const;
-	SLATECORE_API FShapedGlyphSequenceRef ShapeBidirectionalText( const TCHAR* InText, const int32 InTextStart, const int32 InTextLen, const FSlateFontInfo &InFontInfo, const float InFontScale, const TextBiDi::ETextDirection InBaseDirection, const ETextShapingMethod InTextShapingMethod ) const;
+	SLATECORE_API FShapedGlyphSequenceRef ShapeBidirectionalText( const FString& InText, const FSlateFontInfo &InFontInfo, const float InFontScale, const TextBiDi::ETextDirection InBaseDirection, const ETextShapingMethod InTextShapingMethod , TArray<TWeakPtr<FFreeTypeFace>>* OutLoadingFontFaces = nullptr) const;
+	SLATECORE_API FShapedGlyphSequenceRef ShapeBidirectionalText( const TCHAR* InText, const int32 InTextStart, const int32 InTextLen, const FSlateFontInfo &InFontInfo, const float InFontScale, const TextBiDi::ETextDirection InBaseDirection, const ETextShapingMethod InTextShapingMethod , TArray<TWeakPtr<FFreeTypeFace>>* OutLoadingFontFaces = nullptr) const;
 
 	/** 
 	 * Performs text shaping on the given range of the string using the given font info. Returns you the shaped text sequence to use for text rendering via FSlateDrawElement::MakeShapedText.
@@ -801,9 +803,10 @@ public:
 	 * @param InFontScale			The scale to apply to the font
 	 * @param InTextDirection		The reading direction of the text to shape (valid values are LeftToRight or RightToLeft)
 	 * @param InTextShapingMethod	The text shaping method to use
+	 * @param OutLoadingFontFaces	An array (if provided) that keeps track of the unique font faces that are loading for a text during shaping.
 	 */
-	SLATECORE_API FShapedGlyphSequenceRef ShapeUnidirectionalText( const FString& InText, const FSlateFontInfo &InFontInfo, const float InFontScale, const TextBiDi::ETextDirection InTextDirection, const ETextShapingMethod InTextShapingMethod ) const;
-	SLATECORE_API FShapedGlyphSequenceRef ShapeUnidirectionalText( const TCHAR* InText, const int32 InTextStart, const int32 InTextLen, const FSlateFontInfo &InFontInfo, const float InFontScale, const TextBiDi::ETextDirection InTextDirection, const ETextShapingMethod InTextShapingMethod ) const;
+	SLATECORE_API FShapedGlyphSequenceRef ShapeUnidirectionalText( const FString& InText, const FSlateFontInfo &InFontInfo, const float InFontScale, const TextBiDi::ETextDirection InTextDirection, const ETextShapingMethod InTextShapingMethod , TArray<TWeakPtr<FFreeTypeFace>>* OutLoadingFontFaces = nullptr) const;
+	SLATECORE_API FShapedGlyphSequenceRef ShapeUnidirectionalText( const TCHAR* InText, const int32 InTextStart, const int32 InTextLen, const FSlateFontInfo &InFontInfo, const float InFontScale, const TextBiDi::ETextDirection InTextDirection, const ETextShapingMethod InTextShapingMethod , TArray<TWeakPtr<FFreeTypeFace>>* OutLoadingFontFaces = nullptr) const;
 
 	/**
 	 * Performs text shaping on the overflow glyph sequence for a given font. The overflow sequence is used to replace characters that are clipped
@@ -1115,3 +1118,4 @@ private:
 	FText EllipsisText;
 
 };
+
diff --git a/Engine/Source/Runtime/UMG/Private/Components/EditableText.cpp b/Engine/Source/Runtime/UMG/Private/Components/EditableText.cpp
--- a/Engine/Source/Runtime/UMG/Private/Components/EditableText.cpp
+++ b/Engine/Source/Runtime/UMG/Private/Components/EditableText.cpp
@@ -5,6 +5,7 @@
 #include "Engine/Font.h"
 #include "Engine/LocalPlayer.h"
 #include "Framework/Application/SlateApplication.h"
+#include "Framework/Text/SlateAsyncFontLoadingTypes.h"
 #include "Widgets/DeclarativeSyntaxSupport.h"
 #include "Widgets/Input/SEditableText.h"
 #include "Slate/SlateBrushAsset.h"
@@ -57,6 +58,7 @@ PRAGMA_DISABLE_DEPRECATION_WARNINGS
 	VirtualKeyboardDismissAction = EVirtualKeyboardDismissAction::TextChangeOnDismiss;
 	SetClipping(EWidgetClipping::ClipToBounds);
 	OverflowPolicy = ETextOverflowPolicy::Clip;
+	FontFacesLoadingPaintPolicy = EFontFacesLoadingPaintPolicy::DoNotPaint;
 PRAGMA_ENABLE_DEPRECATION_WARNINGS
 
 #if WITH_EDITORONLY_DATA
@@ -93,7 +95,8 @@ TSharedRef<SWidget> UEditableText::RebuildWidget()
 		.VirtualKeyboardDismissAction(VirtualKeyboardDismissAction)
 		.Justification(Justification)
 		.OverflowPolicy(OverflowPolicy)
-		.EnableIntegratedKeyboard(EnableIntegratedKeyboard);
+		.EnableIntegratedKeyboard(EnableIntegratedKeyboard)
+		.OnAllFontFacesFinishLoading(OnAllFontFacesFinishLoading.IsBound() ? BIND_UOBJECT_DELEGATE(FSimpleDelegate, SlateHandleAllFontFacesFinishLoading) : FSimpleDelegate());
 	PRAGMA_ENABLE_DEPRECATION_WARNINGS
 	
 	return MyEditableText.ToSharedRef();
@@ -123,6 +126,7 @@ void UEditableText::SynchronizeProperties()
 	MyEditableText->SetJustification(Justification);
 	MyEditableText->SetOverflowPolicy(OverflowPolicy);
 	MyEditableText->SetEnableIntegratedKeyboard(EnableIntegratedKeyboard);
+	MyEditableText->SetFontFacesLoadingPaintPolicy(FontFacesLoadingPaintPolicy);
 
 	PRAGMA_ENABLE_DEPRECATION_WARNINGS
 
@@ -439,6 +443,17 @@ void UEditableText::ToggleVirtualKeyboard(bool bShow)
 	}
 }
 
+void UEditableText::SetFontFacesLoadingPaintPolicy(EFontFacesLoadingPaintPolicy InFontFacesLoadingPaintPolicy)
+{
+	FontFacesLoadingPaintPolicy = InFontFacesLoadingPaintPolicy;
+
+	if (MyEditableText)
+	{
+		MyEditableText->SetFontFacesLoadingPaintPolicy(InFontFacesLoadingPaintPolicy);
+	}
+}
+
+
 /////////////////////////////////////////////////////
 
 #undef LOCTEXT_NAMESPACE
diff --git a/Engine/Source/Runtime/UMG/Private/Components/EditableTextBox.cpp b/Engine/Source/Runtime/UMG/Private/Components/EditableTextBox.cpp
--- a/Engine/Source/Runtime/UMG/Private/Components/EditableTextBox.cpp
+++ b/Engine/Source/Runtime/UMG/Private/Components/EditableTextBox.cpp
@@ -5,6 +5,7 @@
 #include "UObject/UE5MainStreamObjectVersion.h"
 #include "UObject/UE5ReleaseStreamObjectVersion.h"
 #include "Engine/Font.h"
+#include "Framework/Text/SlateAsyncFontLoadingTypes.h"
 #include "Widgets/DeclarativeSyntaxSupport.h"
 #include "Widgets/Input/SEditableTextBox.h"
 #include "Styling/DefaultStyleCache.h"
@@ -35,6 +36,7 @@ PRAGMA_DISABLE_DEPRECATION_WARNINGS
 	AllowContextMenu = true;
 	VirtualKeyboardDismissAction = EVirtualKeyboardDismissAction::TextChangeOnDismiss;
 	OverflowPolicy = ETextOverflowPolicy::Clip;
+	FontFacesLoadingPaintPolicy = EFontFacesLoadingPaintPolicy::DoNotPaint;
 
 	WidgetStyle = UE::Slate::Private::FDefaultStyleCache::GetRuntime().GetEditableTextBoxStyle();
 
@@ -110,7 +112,8 @@ TSharedRef<SWidget> UEditableTextBox::RebuildWidget()
 		.VirtualKeyboardTrigger(VirtualKeyboardTrigger)
 		.VirtualKeyboardDismissAction(VirtualKeyboardDismissAction)
 		.Justification(Justification)
-		.OverflowPolicy(OverflowPolicy);
+		.OverflowPolicy(OverflowPolicy)
+		.OnAllFontFacesFinishLoading(OnAllFontFacesFinishLoading.IsBound() ? BIND_UOBJECT_DELEGATE(FSimpleDelegate, SlateHandleAllFontFacesFinishLoading) : FSimpleDelegate());
 
 	return MyEditableTextBlock.ToSharedRef();
 }
@@ -142,6 +145,7 @@ void UEditableTextBox::SynchronizeProperties()
 	MyEditableTextBlock->SetVirtualKeyboardDismissAction(VirtualKeyboardDismissAction);
 	MyEditableTextBlock->SetJustification(Justification);
 	MyEditableTextBlock->SetOverflowPolicy(OverflowPolicy);
+	MyEditableTextBlock->SetFontFacesLoadingPaintPolicy(FontFacesLoadingPaintPolicy);
 
 	ShapedTextOptions.SynchronizeShapedTextProperties(*MyEditableTextBlock);
 }
@@ -426,9 +430,17 @@ const FText UEditableTextBox::GetPaletteCategory()
 {
 	return LOCTEXT("Input", "Input");
 }
-
 #endif
 
+void UEditableTextBox::SetFontFacesLoadingPaintPolicy(EFontFacesLoadingPaintPolicy InFontFacesLoadingPaintPolicy)
+{
+	FontFacesLoadingPaintPolicy = InFontFacesLoadingPaintPolicy;
+
+	if (MyEditableTextBlock)
+	{
+		MyEditableTextBlock->SetFontFacesLoadingPaintPolicy(InFontFacesLoadingPaintPolicy);
+	}
+}
 /////////////////////////////////////////////////////
 
 #undef LOCTEXT_NAMESPACE
diff --git a/Engine/Source/Runtime/UMG/Private/Components/MultiLineEditableText.cpp b/Engine/Source/Runtime/UMG/Private/Components/MultiLineEditableText.cpp
--- a/Engine/Source/Runtime/UMG/Private/Components/MultiLineEditableText.cpp
+++ b/Engine/Source/Runtime/UMG/Private/Components/MultiLineEditableText.cpp
@@ -78,6 +78,7 @@ TSharedRef<SWidget> UMultiLineEditableText::RebuildWidget()
 	.VirtualKeyboardDismissAction(VirtualKeyboardDismissAction)
 	.OnTextChanged(BIND_UOBJECT_DELEGATE(FOnTextChanged, HandleOnTextChanged))
 	.OnTextCommitted(BIND_UOBJECT_DELEGATE(FOnTextCommitted, HandleOnTextCommitted))
+	.OnAllFontFacesFinishLoading(OnAllFontFacesFinishLoading.IsBound() ? BIND_UOBJECT_DELEGATE(FSimpleDelegate, SlateHandleAllFontFacesFinishLoading) : FSimpleDelegate())
 	;
 	
 	return MyMultiLineEditableText.ToSharedRef();
diff --git a/Engine/Source/Runtime/UMG/Private/Components/MultiLineEditableTextBox.cpp b/Engine/Source/Runtime/UMG/Private/Components/MultiLineEditableTextBox.cpp
--- a/Engine/Source/Runtime/UMG/Private/Components/MultiLineEditableTextBox.cpp
+++ b/Engine/Source/Runtime/UMG/Private/Components/MultiLineEditableTextBox.cpp
@@ -91,6 +91,7 @@ TSharedRef<SWidget> UMultiLineEditableTextBox::RebuildWidget()
 		.VirtualKeyboardDismissAction(VirtualKeyboardDismissAction)
 		.OnTextChanged(BIND_UOBJECT_DELEGATE(FOnTextChanged, HandleOnTextChanged))
 		.OnTextCommitted(BIND_UOBJECT_DELEGATE(FOnTextCommitted, HandleOnTextCommitted))
+		.OnAllFontFacesFinishLoading(OnAllFontFacesFinishLoading.IsBound() ? BIND_UOBJECT_DELEGATE(FSimpleDelegate, SlateHandleAllFontFacesFinishLoading) : FSimpleDelegate())
 		;
 
 	return MyEditableTextBlock.ToSharedRef();
diff --git a/Engine/Source/Runtime/UMG/Private/Components/RichTextBlock.cpp b/Engine/Source/Runtime/UMG/Private/Components/RichTextBlock.cpp
--- a/Engine/Source/Runtime/UMG/Private/Components/RichTextBlock.cpp
+++ b/Engine/Source/Runtime/UMG/Private/Components/RichTextBlock.cpp
@@ -73,11 +73,11 @@ TSharedRef<SWidget> URichTextBlock::RebuildWidget()
 	CreateDecorators(CreatedDecorators);
 
 	TSharedRef<FRichTextLayoutMarshaller> Marshaller = FRichTextLayoutMarshaller::Create(CreateMarkupParser(), CreateMarkupWriter(), CreatedDecorators, StyleInstance.Get());
-PRAGMA_DISABLE_DEPRECATION_WARNINGS
-	MyRichTextBlock =
-		SNew(SRichTextBlock)
+	PRAGMA_DISABLE_DEPRECATION_WARNINGS
+	MyRichTextBlock = SNew(SRichTextBlock)
 		.TextStyle(bOverrideDefaultStyle ? &DefaultTextStyleOverride : &DefaultTextStyle)
-		.Marshaller(Marshaller);
+		.Marshaller(Marshaller)
+		.OnAllFontFacesFinishLoading(OnAllFontFacesFinishLoading.IsBound() ? BIND_UOBJECT_DELEGATE(FSimpleDelegate, SlateHandleAllFontFacesFinishLoading) : FSimpleDelegate());
 PRAGMA_ENABLE_DEPRECATION_WARNINGS
 	return MyRichTextBlock.ToSharedRef();
 }
diff --git a/Engine/Source/Runtime/UMG/Private/Components/TextBlock.cpp b/Engine/Source/Runtime/UMG/Private/Components/TextBlock.cpp
--- a/Engine/Source/Runtime/UMG/Private/Components/TextBlock.cpp
+++ b/Engine/Source/Runtime/UMG/Private/Components/TextBlock.cpp
@@ -329,6 +329,7 @@ TSharedRef<SWidget> UTextBlock::RebuildWidget()
  		[
  			SAssignNew(MyTextBlock, STextBlock)
 			.SimpleTextMode(bSimpleTextMode)
+			.OnAllFontFacesFinishLoading(OnAllFontFacesFinishLoading.IsBound() ? BIND_UOBJECT_DELEGATE(FSimpleDelegate, SlateHandleAllFontFacesFinishLoading) : FSimpleDelegate())
  		];
  		return RetWidget.ToSharedRef();
  	}
@@ -336,7 +337,8 @@ TSharedRef<SWidget> UTextBlock::RebuildWidget()
 	{
 		MyTextBlock =
 			SNew(STextBlock)
-			.SimpleTextMode(bSimpleTextMode);
+			.SimpleTextMode(bSimpleTextMode)
+			.OnAllFontFacesFinishLoading(OnAllFontFacesFinishLoading.IsBound() ? BIND_UOBJECT_DELEGATE(FSimpleDelegate, SlateHandleAllFontFacesFinishLoading) : FSimpleDelegate());
 
 		//if (IsDesignTime())
 		//{
diff --git a/Engine/Source/Runtime/UMG/Private/Components/TextWidgetTypes.cpp b/Engine/Source/Runtime/UMG/Private/Components/TextWidgetTypes.cpp
--- a/Engine/Source/Runtime/UMG/Private/Components/TextWidgetTypes.cpp
+++ b/Engine/Source/Runtime/UMG/Private/Components/TextWidgetTypes.cpp
@@ -2,6 +2,7 @@
 
 #include "Components/TextWidgetTypes.h"
 #include "Fonts/FontCache.h"
+#include "Framework/Text/SlateAsyncFontLoadingTypes.h"
 
 #include UE_INLINE_GENERATED_CPP_BY_NAME(TextWidgetTypes)
 
@@ -24,6 +25,7 @@ UTextLayoutWidget::UTextLayoutWidget(const FObjectInitializer& ObjectInitializer
 	Margin = FMargin(0.0f);
 	LineHeightPercentage = 1.0f;
 	ApplyLineHeightToBottomLine = true;
+	FontFacesLoadingPaintPolicy = EFontFacesLoadingPaintPolicy::DoNotPaint;
 }
 
 void UTextLayoutWidget::SetJustification(ETextJustify::Type InJustification)
@@ -89,6 +91,15 @@ void UTextLayoutWidget::SetApplyLineHeightToBottomLine(bool InApplyLineHeightToB
 	}
 }
 
+void UTextLayoutWidget::SetFontFacesLoadingPaintPolicy(EFontFacesLoadingPaintPolicy InFontFacesLoadingPaintPolicy)
+{
+	if (FontFacesLoadingPaintPolicy != InFontFacesLoadingPaintPolicy)
+	{
+		FontFacesLoadingPaintPolicy = InFontFacesLoadingPaintPolicy;
+		OnFontFacesLoadingPaintPolicyChange(InFontFacesLoadingPaintPolicy);
+	}
+}
+
 void UTextLayoutWidget::SetMargin(const FMargin& InMargin)
 {
 	if (Margin != InMargin)
diff --git a/Engine/Source/Runtime/UMG/Public/Components/EditableText.h b/Engine/Source/Runtime/UMG/Public/Components/EditableText.h
--- a/Engine/Source/Runtime/UMG/Public/Components/EditableText.h
+++ b/Engine/Source/Runtime/UMG/Public/Components/EditableText.h
@@ -17,6 +17,7 @@ class SEditableText;
 class UMaterialInterface;
 class USlateBrushAsset;
 class USlateWidgetStyleAsset;
+enum class EFontFacesLoadingPainPolicy : uint8;
 
 /**
  * Editable text box widget
@@ -30,7 +31,7 @@ public:
 
 	DECLARE_DYNAMIC_MULTICAST_DELEGATE_OneParam(FOnEditableTextChangedEvent, const FText&, Text);
 	DECLARE_DYNAMIC_MULTICAST_DELEGATE_TwoParams(FOnEditableTextCommittedEvent, const FText&, Text, ETextCommit::Type, CommitMethod);
-
+	DECLARE_DYNAMIC_MULTICAST_SPARSE_DELEGATE(FOnAllFontFacesFinishLoading_EditableText, UEditableText, OnAllFontFacesFinishLoading);
 public:
 	UE_DEPRECATED(5.1, "Direct access to Text is deprecated. Please use the getter or setter.")
 	/** The text content for this editable text box widget */
@@ -248,6 +249,7 @@ public:
 	UFUNCTION(BlueprintCallable, Category = "VirtualKeyboard")
 	UMG_API void ToggleVirtualKeyboard(bool bShow); 
 
+	UMG_API void SetFontFacesLoadingPaintPolicy(EFontFacesLoadingPaintPolicy InFontFacesLoadingPaintPolicy);
 public:
 	
 	//~ Begin UWidget Interface
@@ -275,6 +277,10 @@ protected:
 	UMG_API virtual TSharedPtr<SWidget> GetAccessibleWidget() const override;
 #endif
 
+	void SlateHandleAllFontFacesFinishLoading()
+	{
+		OnAllFontFacesFinishLoading.Broadcast();
+	}
 protected:
 	TSharedPtr<SEditableText> MyEditableText;
 
@@ -288,4 +294,13 @@ private:
 	/** Whether the integrated keyboard is enabled for this editable text. */
 	UPROPERTY(EditAnywhere, BlueprintReadWrite, Getter, Setter, Category = Mobile, AdvancedDisplay, meta = (AllowPrivateAccess = "true"))
 	bool EnableIntegratedKeyboard;
+
+protected:
+	/** Controls how painting the text in the widget should behave when font faces for the text is still loading. */
+	UPROPERTY(EditAnywhere, Setter, Category = Font)
+	EFontFacesLoadingPaintPolicy FontFacesLoadingPaintPolicy;
+
+	/** Called when all font faces for the text has finished async loading. Is not called if font faces are synchronously loaded. */
+	UPROPERTY(BlueprintAssignable, Category = "Widget Event", meta = (DisplayName = "OnAllFontFacesFinishLoading(Editable Text)"))
+	FOnAllFontFacesFinishLoading_EditableText OnAllFontFacesFinishLoading;
 };
diff --git a/Engine/Source/Runtime/UMG/Public/Components/EditableTextBox.h b/Engine/Source/Runtime/UMG/Public/Components/EditableTextBox.h
--- a/Engine/Source/Runtime/UMG/Public/Components/EditableTextBox.h
+++ b/Engine/Source/Runtime/UMG/Public/Components/EditableTextBox.h
@@ -34,7 +34,7 @@ public:
 	DECLARE_DYNAMIC_MULTICAST_DELEGATE_TwoParams(FOnEditableTextBoxCommittedEvent, const FText&, Text, ETextCommit::Type, CommitMethod);
 	
 	DECLARE_MULTICAST_DELEGATE_TwoParams(FOnEditableTextBoxCursorMovedWithSelectionEvent, const FTextLocation&, const FTextSelection&);
-
+	DECLARE_DYNAMIC_MULTICAST_SPARSE_DELEGATE(FOnAllFontFacesFinishLoading_EditableTextBox, UEditableTextBox, OnAllFontFacesFinishLoading);
 public:
 	UE_DEPRECATED(5.1, "Direct access to Text is deprecated. Please use the getter or setter.")
 	/** The text content for this editable text box widget */
@@ -262,6 +262,7 @@ public:
 	UFUNCTION(BlueprintCallable, Category="Widget", meta=(DisplayName="SetForegroundColor (Text Box)"))
 	UMG_API void SetForegroundColor(FLinearColor color);
 
+	UMG_API void SetFontFacesLoadingPaintPolicy(EFontFacesLoadingPaintPolicy InFontFacesLoadingPaintPolicy);
 	//~ Begin UWidget Interface
 	UMG_API virtual void SynchronizeProperties() override;
 	//~ End UWidget Interface
@@ -288,6 +289,10 @@ protected:
 	UMG_API virtual TSharedPtr<SWidget> GetAccessibleWidget() const override;
 #endif
 
+	void SlateHandleAllFontFacesFinishLoading()
+	{
+		OnAllFontFacesFinishLoading.Broadcast();
+	}
 protected:
 	TSharedPtr<SEditableTextBox> MyEditableTextBlock;
 
@@ -302,4 +307,13 @@ private:
 	UPROPERTY()
 	bool bIsFontDeprecationDone;
 #endif
+
+protected:
+	/** Controls how painting the text in the widget should behave when font faces for the text is still loading. */
+	UPROPERTY(EditAnywhere, Setter, Category = Font)
+	EFontFacesLoadingPaintPolicy FontFacesLoadingPaintPolicy;
+
+	/** Called when all font faces for the text has finished async loading. Is not called if font faces are synchronously loaded. */
+	UPROPERTY(BlueprintAssignable, Category = "Widget Event", meta = (DisplayName = "OnAllFontFacesFinishLoading(Editable Text)"))
+	FOnAllFontFacesFinishLoading_EditableTextBox OnAllFontFacesFinishLoading;
 };
diff --git a/Engine/Source/Runtime/UMG/Public/Components/TextWidgetTypes.h b/Engine/Source/Runtime/UMG/Public/Components/TextWidgetTypes.h
--- a/Engine/Source/Runtime/UMG/Public/Components/TextWidgetTypes.h
+++ b/Engine/Source/Runtime/UMG/Public/Components/TextWidgetTypes.h
@@ -61,6 +61,8 @@ struct FShapedTextOptions
 	ETextFlowDirection TextFlowDirection;
 };
 
+enum class EFontFacesLoadingPaintPolicy : uint8;
+
 
 /**
  * Base class for all widgets that use a text layout.
@@ -72,6 +74,8 @@ class UTextLayoutWidget : public UWidget
 	GENERATED_UCLASS_BODY()
 
 public:
+	DECLARE_DYNAMIC_MULTICAST_SPARSE_DELEGATE(FOnAllFontFacesFinishLoading_TextLayoutWidget, UTextLayoutWidget, OnAllFontFacesFinishLoading);
+
 	bool GetAutoWrapText() const { return AutoWrapText; }
 	float GetWrapTextAt() const { return WrapTextAt; }
 
@@ -84,6 +88,7 @@ public:
 	UMG_API void SetWrapTextAt(float InWrapTextAt);
 	UMG_API void SetLineHeightPercentage(float InLineHeightPercentage);
 	UMG_API void SetApplyLineHeightToBottomLine(bool InApplyLineHeightToBottomLine);
+	UMG_API void SetFontFacesLoadingPaintPolicy(EFontFacesLoadingPaintPolicy InFontFacesLoadingPaintPolicy);
 	UMG_API void SetMargin(const FMargin& InMargin);
 
 protected:
@@ -94,6 +99,7 @@ protected:
 	virtual void OnWrapTextAtChanged(float InWrapTextAt) {};
 	virtual void OnLineHeightPercentageChanged(float InLineHeightPercentage) {};
 	virtual void OnApplyLineHeightToBottomLineChanged(bool InApplyLineHeightToBottomLine) {};
+	virtual void OnFontFacesLoadingPaintPolicyChange(EFontFacesLoadingPaintPolicy InFontFacesLoadingPaintPolicy) {};
 	virtual void OnMarginChanged(const FMargin& InMargin) {};
 
 	/** Synchronize the properties with the given widget. A template as the Slate widgets conform to the same API, but don't derive from a common base. */
@@ -109,12 +115,22 @@ protected:
 		InWidget.SetMargin(Margin);
 		InWidget.SetLineHeightPercentage(LineHeightPercentage);
 		InWidget.SetApplyLineHeightToBottomLine(ApplyLineHeightToBottomLine);
+		InWidget.SetFontFacesLoadingPaintPolicy(FontFacesLoadingPaintPolicy);
+	}
+
+	void SlateHandleAllFontFacesFinishLoading()
+	{
+		OnAllFontFacesFinishLoading.Broadcast();
 	}
 
 	/** Controls how the text within this widget should be shaped. */
 	UPROPERTY(EditAnywhere, BlueprintReadWrite, Setter, Category=Localization, AdvancedDisplay, meta=(ShowOnlyInnerProperties))
 	FShapedTextOptions ShapedTextOptions;
 
+	/** Called when all font faces for the text has finished async loading. Is not called if font faces are synchronously loaded. */
+	UPROPERTY(BlueprintAssignable, Category = "Font|Event")
+	FOnAllFontFacesFinishLoading_TextLayoutWidget OnAllFontFacesFinishLoading;
+
 	/** How the text should be aligned with the margin. */
 	UPROPERTY(EditAnywhere, BlueprintReadWrite, BlueprintSetter=SetJustification, Setter, Category=Appearance)
 	TEnumAsByte<ETextJustify::Type> Justification;
@@ -131,6 +147,10 @@ protected:
 	UPROPERTY(EditAnywhere, BlueprintReadWrite, Setter, Category=Appearance, AdvancedDisplay)
 	bool ApplyLineHeightToBottomLine;
 
+	/** Controls how painting the text in the widget should behave when font faces for the text is still loading. */
+	UPROPERTY(EditAnywhere, Setter, Category = Font)
+	EFontFacesLoadingPaintPolicy FontFacesLoadingPaintPolicy;
+
 	/** Whether text wraps onto a new line when it's length exceeds this width; if this value is zero or negative, no wrapping occurs. */
 	UPROPERTY(EditAnywhere, BlueprintReadWrite, Setter, Category=Wrapping)
 	float WrapTextAt;
@@ -142,5 +162,4 @@ protected:
 	/** The amount to scale each lines height by. */
 	UPROPERTY(EditAnywhere, BlueprintReadWrite, Setter, Category=Appearance, AdvancedDisplay)
 	float LineHeightPercentage;
-
 };
