Adding console variables for Slate Paint and Prepass debug views to only show widgets in invalidation roots

#Synced-CL 47563119

diff --git a/Engine/Source/Developer/SlateReflector/Private/Widgets/SSlateOptions.cpp b/Engine/Source/Developer/SlateReflector/Private/Widgets/SSlateOptions.cpp
--- a/Engine/Source/Developer/SlateReflector/Private/Widgets/SSlateOptions.cpp
+++ b/Engine/Source/Developer/SlateReflector/Private/Widgets/SSlateOptions.cpp
@@ -107,6 +107,7 @@ void SSlateOptions::Construct( const FArguments& InArgs )
 		{
 			FSlateIcon Icon(FWidgetReflectorStyle::GetStyleSetName(), "Icon.Empty");
 			AddMenuEntry(MenuBuilder, Icon, LOCTEXT("Paint_OnlyGameWindow", "Only Game Window"), TEXT("SlateDebugger.Paint.OnlyGameWindow"));
+			AddMenuEntry(MenuBuilder, Icon, LOCTEXT("Paint_OnlyInvalidationRoot", "Only Invalidation Root"), TEXT("SlateDebugger.Paint.OnlyInvalidationRoot"));
 			AddMenuEntry(MenuBuilder, Icon, LOCTEXT("Paint_ToggleWidgetNameList", "Toggle Widget Name List"), TEXT("SlateDebugger.Paint.EnableWidgetNameList"));
 			AddMenuEntry(MenuBuilder, Icon, LOCTEXT("Paint_Draw Border", "Draw Border"), TEXT("SlateDebugger.Paint.DrawBorder"));
 			AddMenuEntry(MenuBuilder, Icon, LOCTEXT("Paint_Draw Fill", "Draw Fill"), TEXT("SlateDebugger.Paint.DrawFill"));
@@ -115,6 +116,7 @@ void SSlateOptions::Construct( const FArguments& InArgs )
 		{
 			FSlateIcon Icon(FWidgetReflectorStyle::GetStyleSetName(), "Icon.Empty");
 			AddMenuEntry(MenuBuilder, Icon, LOCTEXT("Prepass_OnlyGameWindow", "Only Game Window"), TEXT("SlateDebugger.Prepass.OnlyGameWindow"));
+			AddMenuEntry(MenuBuilder, Icon, LOCTEXT("Prepass_OnlyInvalidationRoot", "Only Invalidation Root"), TEXT("SlateDebugger.Prepass.OnlyInvalidationRoot"));
 			AddMenuEntry(MenuBuilder, Icon, LOCTEXT("Prepass_ToggleWidgetNameList", "Toggle Widget Name List"), TEXT("SlateDebugger.Prepass.EnableWidgetNameList"));
 			AddMenuEntry(MenuBuilder, Icon, LOCTEXT("Prepass_Draw Border", "Draw Border"), TEXT("SlateDebugger.Prepass.DrawBorder"));
 			AddMenuEntry(MenuBuilder, Icon, LOCTEXT("Prepass_Draw Fill", "Draw Fill"), TEXT("SlateDebugger.Prepass.DrawFill"));
diff --git a/Engine/Source/Runtime/SlateCore/Private/Debugging/ConsoleSlateDebuggerPaint.cpp b/Engine/Source/Runtime/SlateCore/Private/Debugging/ConsoleSlateDebuggerPaint.cpp
--- a/Engine/Source/Runtime/SlateCore/Private/Debugging/ConsoleSlateDebuggerPaint.cpp
+++ b/Engine/Source/Runtime/SlateCore/Private/Debugging/ConsoleSlateDebuggerPaint.cpp
@@ -55,7 +55,12 @@ FConsoleSlateDebuggerPaint::FConsoleSlateDebuggerPaint()
 		TEXT("SlateDebugger.Paint.OnlyGameWindow")
 		, bDebugGameWindowOnly
 		, TEXT("Option to only debug the game window")
-		, FConsoleVariableDelegate::CreateRaw(this, &FConsoleSlateDebuggerPaint::HandleDebugGameWindowOnlyChanged))
+		, FConsoleVariableDelegate::CreateRaw(this, &FConsoleSlateDebuggerPaint::RebuildWidgetListOnVariableChanged))
+	, OnlyInvalidationRootRefCVar(
+		TEXT("SlateDebugger.Paint.OnlyInvalidationRoot")
+		, bInvalidationRootOnly
+		, TEXT("Option to only debug widgets within an Invalidation Root")
+		, FConsoleVariableDelegate::CreateRaw(this, &FConsoleSlateDebuggerPaint::RebuildWidgetListOnVariableChanged))
 	, DrawBorderEnabledRefCVar(
 		TEXT("SlateDebugger.Paint.DrawBorder")
 		, bDrawBorder
@@ -83,6 +88,7 @@ void FConsoleSlateDebuggerPaint::LoadConfig()
 	GConfig->GetBool(*Section, TEXT("bLogWarningIfWidgetIsPaintedMoreThanOnce"), bLogWarningIfWidgetIsPaintedMoreThanOnce, *GEditorPerProjectIni);
 
 	OnlyGameWindowRefCVar->Set(bDebugGameWindowOnly, ECVF_SetByProjectSetting);
+	OnlyInvalidationRootRefCVar->Set(bInvalidationRootOnly, ECVF_SetByProjectSetting);
 	EnableWidgetsNameListRefCVar->Set(bDisplayWidgetsNameList, ECVF_SetByProjectSetting);
 	DrawFillEnabledRefCVar->Set(bDrawBox, ECVF_SetByProjectSetting);
 	DrawBorderEnabledRefCVar->Set(bDrawBorder, ECVF_SetByProjectSetting);
diff --git a/Engine/Source/Runtime/SlateCore/Private/Debugging/ConsoleSlateDebuggerPaint.h b/Engine/Source/Runtime/SlateCore/Private/Debugging/ConsoleSlateDebuggerPaint.h
--- a/Engine/Source/Runtime/SlateCore/Private/Debugging/ConsoleSlateDebuggerPaint.h
+++ b/Engine/Source/Runtime/SlateCore/Private/Debugging/ConsoleSlateDebuggerPaint.h
@@ -57,6 +57,7 @@ private:
 	FAutoConsoleVariableRef MaxNumberOfWidgetInListRefCVar;
 	FAutoConsoleVariableRef LogWarningIfWidgetIsPaintedMoreThanOnceRefCVar;
 	FAutoConsoleVariableRef OnlyGameWindowRefCVar;
+	FAutoConsoleVariableRef OnlyInvalidationRootRefCVar;
 	FAutoConsoleVariableRef DrawBorderEnabledRefCVar;
 	FAutoConsoleVariableRef DrawFillEnabledRefCVar;
 };
diff --git a/Engine/Source/Runtime/SlateCore/Private/Debugging/ConsoleSlateDebuggerPassBase.cpp b/Engine/Source/Runtime/SlateCore/Private/Debugging/ConsoleSlateDebuggerPassBase.cpp
--- a/Engine/Source/Runtime/SlateCore/Private/Debugging/ConsoleSlateDebuggerPassBase.cpp
+++ b/Engine/Source/Runtime/SlateCore/Private/Debugging/ConsoleSlateDebuggerPassBase.cpp
@@ -23,6 +23,7 @@ FConsoleSlateDebuggerPassBase::FConsoleSlateDebuggerPassBase()
 	, bLogWidgetName(false)
 	, bLogWidgetNameOnce(false)
 	, bDebugGameWindowOnly(true)
+	, bInvalidationRootOnly(false)
 	, MostRecentColor(0.75f, 0.0f, 0.0f, 0.02f)
 	, LeastRecentColor(0.0f, 0.75f, 0.0f, 0.5f)
 	, DrawWidgetNameColor(FColorList::SpicyPink)
@@ -40,17 +41,18 @@ void FConsoleSlateDebuggerPassBase::LoadConfig()
 {
 	const FString Section = *GetConfigSection();
 	GConfig->GetBool(*Section, TEXT("bDebugGameWindowOnly"), bDebugGameWindowOnly, *GEditorPerProjectIni);
+	GConfig->GetBool(*Section, TEXT("bInvalidationRootOnly"), bInvalidationRootOnly, *GEditorPerProjectIni);
 	GConfig->GetBool(*Section, TEXT("bDisplayWidgetsNameList"), bDisplayWidgetsNameList, *GEditorPerProjectIni);
 	GConfig->GetBool(*Section, TEXT("bUseWidgetPathAsName"), bUseWidgetPathAsName, *GEditorPerProjectIni);
 	GConfig->GetBool(*Section, TEXT("bDrawBox"), bDrawBox, *GEditorPerProjectIni);
 	GConfig->GetBool(*Section, TEXT("bDrawBorder"), bDrawBorder, *GEditorPerProjectIni);
 	GConfig->GetBool(*Section, TEXT("bLogWidgetName"), bLogWidgetName, *GEditorPerProjectIni);
 	FColor TmpColor;
-	if (GConfig->GetColor(*Section, TEXT("MostRecentColor"), TmpColor, *GEditorPerProjectIni))
+	if (GConfig->GetColor(*Section, TEXT("MostRecentPaintColor"), TmpColor, *GEditorPerProjectIni))
 	{
 		MostRecentColor = TmpColor;
 	}
-	if (GConfig->GetColor(*Section, TEXT("LeastRecentColor"), TmpColor, *GEditorPerProjectIni))
+	if (GConfig->GetColor(*Section, TEXT("LeastRecentPaintColor"), TmpColor, *GEditorPerProjectIni))
 	{
 		LeastRecentColor = TmpColor;
 	}
@@ -66,6 +68,7 @@ void FConsoleSlateDebuggerPassBase::SaveConfig()
 {
 	const FString Section = *GetConfigSection();
 	GConfig->SetBool(*Section, TEXT("bDebugGameWindowOnly"), bDebugGameWindowOnly, *GEditorPerProjectIni);
+	GConfig->SetBool(*Section, TEXT("bInvalidationRootOnly"), bInvalidationRootOnly, *GEditorPerProjectIni);
 	GConfig->SetBool(*Section, TEXT("bDisplayWidgetsNameList"), bDisplayWidgetsNameList, *GEditorPerProjectIni);
 	GConfig->SetBool(*Section, TEXT("bUseWidgetPathAsName"), bUseWidgetPathAsName, *GEditorPerProjectIni);
 	GConfig->SetBool(*Section, TEXT("bDrawBox"), bDrawBox, *GEditorPerProjectIni);
@@ -142,7 +145,7 @@ void FConsoleSlateDebuggerPassBase::SaveConfigOnVariableChanged(IConsoleVariable
 	SaveConfig();
 }
 
-void FConsoleSlateDebuggerPassBase::HandleDebugGameWindowOnlyChanged(IConsoleVariable* Variable)
+void FConsoleSlateDebuggerPassBase::RebuildWidgetListOnVariableChanged(IConsoleVariable* Variable)
 {
 	SaveConfig();
 	if (bEnabled)
@@ -209,14 +212,31 @@ void FConsoleSlateDebuggerPassBase::HandlePaintDebugInfo(const FPaintArgs& InArg
 		);
 	};
 
-	for (const TPair<FConsoleSlateDebuggerUtility::TSWidgetId, FWidgetInfo>& Pair : Widgets)
+	for (TPair<FConsoleSlateDebuggerUtility::TSWidgetId, FWidgetInfo>& Pair : Widgets)
 	{
-		const FWidgetInfo& WidgetInfo = Pair.Value;
+		FWidgetInfo& WidgetInfo = Pair.Value;
 		if (WidgetInfo.Window != PaintWindow)
 		{
 			continue;
 		}
-		
+
+		if (bInvalidationRootOnly & !WidgetInfo.bInInvalidationRoot)
+		{
+			continue;
+		}
+
+		if (TSharedPtr<const SWidget> Widget = WidgetInfo.Widget.Pin())
+		{
+			// Skip the widget if collapsed or not in a window anymore (ex: in an unactive tab)
+			if (Widget->GetProxyHandle().GetWidgetVisibility(Widget.Get()).IsCollapsed() || !FSlateApplicationBase::Get().FindWidgetWindow(Widget.ToSharedRef()))
+			{
+				continue;
+			}
+			// Update with latest size and location
+			WidgetInfo.PaintLocation = Widget->GetPersistentState().AllottedGeometry.GetAbsolutePosition();
+			WidgetInfo.PaintSize = Widget->GetPersistentState().AllottedGeometry.GetAbsoluteSize();
+		}
+
 		if (WidgetInfo.LastUpdatedFrame == GFrameNumber)
 		{
 			++NumberOfWidgetsUpdatedThisFrame;
@@ -351,6 +371,17 @@ FConsoleSlateDebuggerPassBase::FWidgetInfo& FConsoleSlateDebuggerPassBase::AddUp
 	FoundItem->Widget = Widget.AsWeak();
 	FoundItem->PaintLocation = Widget.GetPersistentState().AllottedGeometry.GetAbsolutePosition();
 	FoundItem->PaintSize = Widget.GetPersistentState().AllottedGeometry.GetAbsoluteSize();
+	FoundItem->bInInvalidationRoot = Widget.GetProxyHandle().GetInvalidationRootHandle().GetInvalidationRoot() != nullptr;
+	if (!FoundItem->bInInvalidationRoot)
+	{
+		// If the widget was added to the Invalidation Root after it stopped caching, the widget would not be inside an invalidation root
+		TSharedPtr<const SWidget> InvalidationRoot = Widget.AsShared();
+		while (InvalidationRoot && !InvalidationRoot->Advanced_IsInvalidationRoot())
+		{
+			InvalidationRoot = InvalidationRoot->GetParentWidget();
+		}
+		FoundItem->bInInvalidationRoot = InvalidationRoot.IsValid();
+	}
 	FoundItem->LastUpdatedFrame = LastUpdatedFrame;
 	FoundItem->LastUpdatedTime = FSlateApplicationBase::Get().GetCurrentTime();
 	if (GFrameNumber != FoundItem->LastUpdatedFrame) // if we do not have the time at which it was drawn, we estimate
diff --git a/Engine/Source/Runtime/SlateCore/Private/Debugging/ConsoleSlateDebuggerPassBase.h b/Engine/Source/Runtime/SlateCore/Private/Debugging/ConsoleSlateDebuggerPassBase.h
--- a/Engine/Source/Runtime/SlateCore/Private/Debugging/ConsoleSlateDebuggerPassBase.h
+++ b/Engine/Source/Runtime/SlateCore/Private/Debugging/ConsoleSlateDebuggerPassBase.h
@@ -51,15 +51,16 @@ protected:
 	
 	struct FWidgetInfo 
 	{
-		FConsoleSlateDebuggerUtility::TSWindowId Window;
+		FConsoleSlateDebuggerUtility::TSWindowId Window = FConsoleSlateDebuggerUtility::InvalidWindowId;
 		TWeakPtr<const SWidget> Widget; // used to check if the widget has been destroyed
 		FVector2f PaintLocation;
 		FVector2f PaintSize;
 		FString WidgetName;
-		
-		uint32 LastUpdatedFrame;
-		double LastUpdatedTime;
-		int32 UpdateCount;
+
+		bool bInInvalidationRoot = false;
+		uint32 LastUpdatedFrame = 0;
+		double LastUpdatedTime = 0.0;
+		int32 UpdateCount = 0;
 	};
 
 	/**
@@ -81,7 +82,7 @@ protected:
 	/* Should be called by CVars that should only save the current config */
 	void SaveConfigOnVariableChanged(IConsoleVariable* Variable);
 	/* Should be called by the CVar that changes the scope of the widgets retrieved. Ends up rebuilding the widget list */
-	void HandleDebugGameWindowOnlyChanged(IConsoleVariable* Variable);
+	void RebuildWidgetListOnVariableChanged(IConsoleVariable* Variable);
 	/* Called on EndFrame to reset the counts and remove deleted widgets */
 	void HandleEndFrame();
 	/* Called by Slate to draw additional elements on screen. */
@@ -103,6 +104,7 @@ protected:
 	bool bLogWidgetName;
 	bool bLogWidgetNameOnce;
 	bool bDebugGameWindowOnly;
+	bool bInvalidationRootOnly;
 	FLinearColor MostRecentColor;
 	FLinearColor LeastRecentColor;
 	FLinearColor DrawWidgetNameColor;
diff --git a/Engine/Source/Runtime/SlateCore/Private/Debugging/ConsoleSlateDebuggerPrepass.cpp b/Engine/Source/Runtime/SlateCore/Private/Debugging/ConsoleSlateDebuggerPrepass.cpp
--- a/Engine/Source/Runtime/SlateCore/Private/Debugging/ConsoleSlateDebuggerPrepass.cpp
+++ b/Engine/Source/Runtime/SlateCore/Private/Debugging/ConsoleSlateDebuggerPrepass.cpp
@@ -50,7 +50,12 @@ FConsoleSlateDebuggerPrepass::FConsoleSlateDebuggerPrepass()
 		TEXT("SlateDebugger.Prepass.OnlyGameWindow")
 		, bDebugGameWindowOnly
 		, TEXT("Option to only debug the game window")
-		, FConsoleVariableDelegate::CreateRaw(this, &FConsoleSlateDebuggerPrepass::HandleDebugGameWindowOnlyChanged))
+		, FConsoleVariableDelegate::CreateRaw(this, &FConsoleSlateDebuggerPrepass::RebuildWidgetListOnVariableChanged))
+	, OnlyInvalidationRootRefCVar(
+		TEXT("SlateDebugger.Prepass.OnlyInvalidationRoot")
+		, bInvalidationRootOnly
+		, TEXT("Option to only debug widgets within an Invalidation Root")
+		, FConsoleVariableDelegate::CreateRaw(this, &FConsoleSlateDebuggerPrepass::RebuildWidgetListOnVariableChanged))
 	, DrawBorderEnabledRefCVar(
 		TEXT("SlateDebugger.Prepass.DrawBorder")
 		, bDrawBorder
@@ -76,6 +81,7 @@ void FConsoleSlateDebuggerPrepass::LoadConfig()
 	FConsoleSlateDebuggerPassBase::LoadConfig();
 	
 	OnlyGameWindowRefCVar->Set(bDebugGameWindowOnly, ECVF_SetByProjectSetting);
+	OnlyInvalidationRootRefCVar->Set(bInvalidationRootOnly, ECVF_SetByProjectSetting);
 	EnableWidgetsNameListRefCVar->Set(bDisplayWidgetsNameList, ECVF_SetByProjectSetting);
 	DrawFillEnabledRefCVar->Set(bDrawBox, ECVF_SetByProjectSetting);
 	DrawBorderEnabledRefCVar->Set(bDrawBorder, ECVF_SetByProjectSetting);
diff --git a/Engine/Source/Runtime/SlateCore/Private/Debugging/ConsoleSlateDebuggerPrepass.h b/Engine/Source/Runtime/SlateCore/Private/Debugging/ConsoleSlateDebuggerPrepass.h
--- a/Engine/Source/Runtime/SlateCore/Private/Debugging/ConsoleSlateDebuggerPrepass.h
+++ b/Engine/Source/Runtime/SlateCore/Private/Debugging/ConsoleSlateDebuggerPrepass.h
@@ -53,6 +53,7 @@ private:
 	FAutoConsoleCommand ToggleWidgetsNameListCommand;
 	FAutoConsoleVariableRef MaxNumberOfWidgetInListRefCVar;
 	FAutoConsoleVariableRef OnlyGameWindowRefCVar;
+	FAutoConsoleVariableRef OnlyInvalidationRootRefCVar;
 	FAutoConsoleVariableRef DrawBorderEnabledRefCVar;
 	FAutoConsoleVariableRef DrawFillEnabledRefCVar;
 };
