This CL addresses three linked issues in the Dynamic Invalidation system:

1. **Unsafe pointer access in `ProcessPreUpdate`**
`FSlateInvalidationRoot::ProcessPreUpdate` currently handles `ChildOrder` invalidation and accesses the widget pointer before invalidation is processed. This is unsafe because the widget may have been destroyed, leaving the pointer stale or reused.

2. **Loss of `ChildOrder` invalidation knowledge**
For Dynamic Invalidation to work correctly, `SupportDynamicInvalidation` must be called on each widget with a `ChildOrder` invalidation. However, `FSlateInvalidationRoot` often rebuilds its list, discarding the `InvalidationReason` and losing awareness of `ChildOrder` invalidation. This leads to cases where Dynamic Invalidation allows invalidation even when some widgets in the tree do not support it.

3. **Widgets with multiple parents**
Some widgets are reused in multiple slots across different parents without being detached first. This results in widgets having multiple parents, which is not supported by the invalidation system.

Proposed solution:
- Move `ChildOrder` invalidation handling back to `SWidget::Invalidate` instead of `FSlateInvalidationRoot`. This ensures widgets are valid and makes it easier to detect multiple-parent cases.
- Update `SupportsInvalidationRecursive` to detect widgets with multiple parents and disable invalidation in such cases. This is supported by calling `InvalidateCachedSupportsInvalidation` from `SWidget::Invalidate`.
- Extend `SupportsInvalidationRecursive` to cache whether a widget supports invalidation, which increases performance.

#Synced-CL 48711348

diff --git a/Engine/Source/Developer/SlateReflector/Private/Models/WidgetReflectorNode.cpp b/Engine/Source/Developer/SlateReflector/Private/Models/WidgetReflectorNode.cpp
--- a/Engine/Source/Developer/SlateReflector/Private/Models/WidgetReflectorNode.cpp
+++ b/Engine/Source/Developer/SlateReflector/Private/Models/WidgetReflectorNode.cpp
@@ -992,7 +992,7 @@ bool FWidgetReflectorNodeUtils::GetWidgetSupportsInvalidation(const TSharedPtr<c
 bool FWidgetReflectorNodeUtils::GetWidgetSupportsInvalidationRecursive(const TSharedPtr<const SWidget>& InWidget)
 {
 #if UE_SLATE_WITH_DYNAMIC_INVALIDATION
-	return InWidget.IsValid() ? InWidget->SupportsInvalidationRecursive() : false;
+	return InWidget.IsValid() ? InWidget->SupportsInvalidationRecursive(SWidget::EInvalidationStrategy::UseCachedValue) : false;
 #else
 	return false;
 #endif
diff --git a/Engine/Source/Developer/SlateReflector/Private/Widgets/SWidgetReflectorTreeWidgetItem.cpp b/Engine/Source/Developer/SlateReflector/Private/Widgets/SWidgetReflectorTreeWidgetItem.cpp
--- a/Engine/Source/Developer/SlateReflector/Private/Widgets/SWidgetReflectorTreeWidgetItem.cpp
+++ b/Engine/Source/Developer/SlateReflector/Private/Widgets/SWidgetReflectorTreeWidgetItem.cpp
@@ -73,80 +73,95 @@ TSharedRef<SWidget> SReflectorTreeWidgetItem::GenerateWidgetForColumn(const FNam
 		static FName InvalidationPanelName = TEXT("SInvalidationPanel");
 
 		const TSharedPtr<SWidget> LiveWidget = WidgetInfo->GetLiveWidget();
-		TSharedPtr<SWidget> InvalidationPanel = LiveWidget;
-		while (InvalidationPanel && InvalidationPanel->GetType() != InvalidationPanelName)
+		if (LiveWidget && LiveWidget != SNullWidget::NullWidget)
 		{
-			InvalidationPanel = InvalidationPanel->GetParentWidget();
-		}
+			TSharedPtr<SWidget> InvalidationPanel = LiveWidget;
+			while (InvalidationPanel && InvalidationPanel->GetType() != InvalidationPanelName)
+			{
+				InvalidationPanel = InvalidationPanel->GetParentWidget();
+			}
 
-		if (InvalidationPanel && StaticCastSharedPtr<SInvalidationPanel>(InvalidationPanel)->CanUseDynamicInvalidation()) // if we found an Invalidation Panel using Dynamic Invalidation
-		{
-			const bool CanCache = StaticCastSharedPtr<SInvalidationPanel>(InvalidationPanel)->GetCanCache();
-			// we always show a color block for the Dynamic Invalidation Panel, but for its children we only show the color if it can cache
-			const bool bShowCachingIndicator = CanCache || LiveWidget == InvalidationPanel; 
-			if (bShowCachingIndicator)
+			TOptional<TPair<FLinearColor, FText>> IndicatorColorAndTooltip;
+			if (WidgetInfo->GetParentNode() && LiveWidget->GetParentWidget() != WidgetInfo->GetParentNode()->GetLiveWidget())
 			{
-				TPair<FLinearColor, FText> ColorAndTooltip;
-			
-				if (LiveWidget == InvalidationPanel)
-				{
-					if (!CanCache)
-					{
-						ColorAndTooltip = TPair<FLinearColor, FText> {
-							FLinearColor(1.0f, 0.0f, 0.0f), // Red for disabled Dynamic Invalidation Panel
-							LOCTEXT("DynamicInvalidationPanel_Disabled", "Dynamic Invalidation Panel cannot cache")
-						};
-					}
-					else if (LiveWidget->SupportsInvalidationRecursive())
-					{
-						ColorAndTooltip = TPair<FLinearColor, FText> {
-							FLinearColor(0.0f, 0.5f, 1.0f), // Blue for Dynamic Invalidation Panel that recursively supports invalidation
-							LOCTEXT("DynamicInvalidationPanel_SupportingInvalidationRecursive", "Dynamic Invalidation Panel is caching all its children")
-						} ;
-					}
-					else
-					{
-						ColorAndTooltip = TPair<FLinearColor, FText> {
-							FLinearColor(1.f, 0.f, 1.f), // Magenta for Dynamic Invalidation Panel that is not able to recursively supports invalidation
-							LOCTEXT("DynamicInvalidationPanel_NotSupportingInvalidationRecursive", "Dynamic Invalidation Panel supports invalidation but not all its children do")
-						};
-					}
-				}
-				else
+				IndicatorColorAndTooltip = TPair<FLinearColor, FText> {
+					FLinearColor(0.5f, 0.0f, 0.f), // Dark Red for widget with wrong parent
+					LOCTEXT("DynamicInvalidationPanel_WidgetWithWrongParent", "Widget's 'GetParentWidget' is not its current parent.")
+				};
+			}
+			else if (InvalidationPanel && StaticCastSharedPtr<SInvalidationPanel>(InvalidationPanel)->CanUseDynamicInvalidation()) // if we found an Invalidation Panel using Dynamic Invalidation
+			{
+				const bool CanCache = StaticCastSharedPtr<SInvalidationPanel>(InvalidationPanel)->GetCanCache();
+				// we always show a color block for the Dynamic Invalidation Panel, but for its children we only show the color if it can cache
+				const bool bShowCachingIndicator = CanCache || LiveWidget == InvalidationPanel;
+				if (bShowCachingIndicator)
 				{
-					if (!LiveWidget->SupportsInvalidation())
-					{
-						ColorAndTooltip = TPair<FLinearColor, FText> {
-							FLinearColor(1.0f, 0.0f, 0.0f), // Red for widgets not supporting invalidation
-							LOCTEXT("DynamicInvalidationPanel_WidgetNotSupportingInvalidation", "Widget does not support invalidation")
-						};
-					}
-					else if (LiveWidget->SupportsInvalidationRecursive())
+					if (LiveWidget == InvalidationPanel)
 					{
-						ColorAndTooltip = TPair<FLinearColor, FText> {
-							FLinearColor(0.0f, 1.0f, 0.0f), // Green for widgets recursively supporting invalidation
-							LOCTEXT("DynamicInvalidationPanel_WidgetSupportingInvalidationRecursive", "Widget and children all support invalidation")
-						};
+						if (!CanCache)
+						{
+							IndicatorColorAndTooltip = TPair<FLinearColor, FText> {
+								FLinearColor(1.0f, 0.0f, 0.0f), // Red for disabled Dynamic Invalidation Panel
+								LOCTEXT("DynamicInvalidationPanel_Disabled", "Dynamic Invalidation Panel cannot cache")
+							};
+						}
+						else if (LiveWidget->SupportsInvalidationRecursive(EInvalidationStrategy::UseCachedValue))
+						{
+							IndicatorColorAndTooltip = TPair<FLinearColor, FText> {
+								FLinearColor(0.0f, 0.5f, 1.0f), // Blue for Dynamic Invalidation Panel that recursively supports invalidation
+								LOCTEXT("DynamicInvalidationPanel_SupportingInvalidationRecursive", "Dynamic Invalidation Panel is caching all its children")
+							} ;
+						}
+						else
+						{
+							IndicatorColorAndTooltip = TPair<FLinearColor, FText> {
+								FLinearColor(1.f, 0.f, 1.f), // Magenta for Dynamic Invalidation Panel that is not able to recursively supports invalidation
+								LOCTEXT("DynamicInvalidationPanel_NotSupportingInvalidationRecursive", "Dynamic Invalidation Panel supports invalidation but not all its children do")
+							};
+						}
 					}
 					else
 					{
-						ColorAndTooltip = TPair<FLinearColor, FText> {
-							FLinearColor(1.0f, .5f, 0.0f), // Orange for widgets supporting invalidation but having children that do not
-							LOCTEXT("DynamicInvalidationPanel_WidgetNotSupportingInvalidationRecursive", "Widget supports invalidation but not all its children do")
-						};
+						if (!LiveWidget->SupportsInvalidation())
+						{
+							IndicatorColorAndTooltip = TPair<FLinearColor, FText> {
+								FLinearColor(1.0f, 0.0f, 0.0f), // Red for widgets not supporting invalidation
+								LOCTEXT("DynamicInvalidationPanel_WidgetNotSupportingInvalidation", "Widget Class does not support invalidation")
+							};
+						}
+						else if (LiveWidget->SupportsInvalidationRecursive(EInvalidationStrategy::UseCachedValue))
+						{
+							IndicatorColorAndTooltip = TPair<FLinearColor, FText> {
+								FLinearColor(0.0f, 1.0f, 0.0f), // Green for widgets recursively supporting invalidation
+								LOCTEXT("DynamicInvalidationPanel_WidgetSupportingInvalidationRecursive", "Widget and children all support invalidation")
+							};
+						}
+						else
+						{
+							IndicatorColorAndTooltip = TPair<FLinearColor, FText> {
+								FLinearColor(1.0f, .5f, 0.0f), // Orange for widgets supporting invalidation but having children that do not
+								LOCTEXT("DynamicInvalidationPanel_WidgetNotSupportingInvalidationRecursive", "Widget supports invalidation but not all its children do")
+							};
+						}
 					}
 				}
+			}
 
+			if (IndicatorColorAndTooltip.IsSet())
+			{
 				HorizontalBox->AddSlot()
-				.AutoWidth()
-				.Padding(2.0f, 0.0f)
-				.VAlign(VAlign_Center)
-				[
-					SNew(SColorBlock)
-					.Size(FVector2D{10.0f, 10.0f})
-					.Color(ColorAndTooltip.Key)
-					.ToolTipText(ColorAndTooltip.Value)
-				];
+					.AutoWidth()
+					.Padding(2.0f, 0.0f)
+					.VAlign(VAlign_Center)
+					[
+						SNew(SColorBlock)
+						.Size(FVector2D{10.0f, 10.0f})
+						.Color(IndicatorColorAndTooltip.GetValue().Key)
+						.ToolTipText(FText::Format(LOCTEXT("DynamicInvalidationPanel_TooltipFormat", "{0}\nCached Result: {1}"),
+							IndicatorColorAndTooltip.GetValue().Value,
+							FText::FromString(LexToString(LiveWidget->Debug_GetCachedInvalidationResult()))
+						))
+					];
 			}
 		}
 #endif
diff --git a/Engine/Source/Runtime/Slate/Private/Widgets/SInvalidationPanel.cpp b/Engine/Source/Runtime/Slate/Private/Widgets/SInvalidationPanel.cpp
--- a/Engine/Source/Runtime/Slate/Private/Widgets/SInvalidationPanel.cpp
+++ b/Engine/Source/Runtime/Slate/Private/Widgets/SInvalidationPanel.cpp
@@ -253,10 +253,6 @@ bool SInvalidationPanel::GetUseDynamicInvalidation() const
 void SInvalidationPanel::SetUseDynamicInvalidation(bool InUseDynamicInvalidation)
 {
 	bUseDynamicInvalidation = InUseDynamicInvalidation;
-	
-	// Make sure we update the cached value
-	const bool bUseCachedValue = false;
-	SupportsInvalidationRecursive(bUseCachedValue);
 }
 
 bool SInvalidationPanel::CanUseDynamicInvalidation() const
@@ -286,13 +282,23 @@ FChildren* SInvalidationPanel::Debug_GetChildrenForReflector()
 #endif
 
 #if UE_SLATE_WITH_DYNAMIC_INVALIDATION
-bool SInvalidationPanel::SupportsInvalidationRecursive(bool bUseCachedValue) const
+bool SInvalidationPanel::SupportsInvalidationRecursive(EInvalidationStrategy InInvalidationStrategy) const
 {
-	// For an invalidation panel without bUseDynamicInvalidation, we force the caching of all children, so we always support invalidation
-	return !CanUseDynamicInvalidation() || SWidget::SupportsInvalidationRecursive(bUseCachedValue);
+	const bool bUseStaticInvalidation = !CanUseDynamicInvalidation();
+	if (bUseStaticInvalidation)
+	{
+		// For an invalidation panel not using dynamic invalidation, we force the caching of all children, so we always support invalidation
+		return true;
+	}
+	return SWidget::SupportsInvalidationRecursive(InInvalidationStrategy);
 }
 #endif
 
+UE_TRACE_EVENT_BEGIN(Cpu, SInvalidationPanel_SupportsInvalidationRecursive, NoSync)
+	UE_TRACE_EVENT_FIELD(bool, DuringOnPaint)
+	UE_TRACE_EVENT_FIELD(UE::Trace::WideString, CreatedInLocation)
+UE_TRACE_EVENT_END()
+
 int32 SInvalidationPanel::OnPaint( const FPaintArgs& Args, const FGeometry& AllottedGeometry, const FSlateRect& MyCullingRect, FSlateWindowElementList& OutDrawElements, int32 LayerId, const FWidgetStyle& InWidgetStyle, bool bParentEnabled ) const
 {
 #if SLATE_VERBOSE_NAMED_EVENTS
@@ -304,7 +310,16 @@ int32 SInvalidationPanel::OnPaint( const FPaintArgs& Args, const FGeometry& Allo
 	SInvalidationPanel* MutableThis = const_cast<SInvalidationPanel*>(this);
 
 #if UE_SLATE_WITH_DYNAMIC_INVALIDATION
-	const bool bSupportsInvalidationRecursive = SupportsInvalidationRecursive();
+	bool bSupportsInvalidationRecursive;
+	{
+		#if UE_SLATE_TRACE_ENABLED
+			UE_TRACE_LOG_SCOPED_T(Cpu, SInvalidationPanel_SupportsInvalidationRecursive, SlateChannel)
+				<< SInvalidationPanel_SupportsInvalidationRecursive.DuringOnPaint(true)
+				<< SInvalidationPanel_SupportsInvalidationRecursive.CreatedInLocation(*GetCreatedInLocation().ToString());
+		#endif // UE_SLATE_TRACE_ENABLED
+
+		bSupportsInvalidationRecursive = SupportsInvalidationRecursive(EInvalidationStrategy::UseCachedValue);
+	}
 	const bool bCanCacheValue = GetCanCache();
 	const bool bCanCacheThisFrame = bCanCacheValue & bSupportsInvalidationRecursive;
 #if WITH_SLATE_DEBUGGING
@@ -389,6 +404,22 @@ bool SInvalidationPanel::CustomPrepass(float LayoutScaleMultiplier)
 		{
 			SetNeedsSlowPath(true);
 		}
+#if UE_SLATE_WITH_DYNAMIC_INVALIDATION
+		else if (!NeedsSlowPath())
+		{
+#if UE_SLATE_TRACE_ENABLED
+			UE_TRACE_LOG_SCOPED_T(Cpu, SInvalidationPanel_SupportsInvalidationRecursive, SlateChannel)
+				<< SInvalidationPanel_SupportsInvalidationRecursive.DuringOnPaint(false)
+				<< SInvalidationPanel_SupportsInvalidationRecursive.CreatedInLocation(*GetCreatedInLocation().ToString());
+#endif // UE_SLATE_TRACE_ENABLED
+
+			if (!SupportsInvalidationRecursive(EInvalidationStrategy::UseCachedValue))
+			{
+				SetNeedsSlowPath(true);
+			}
+		}
+#endif
+
 		ProcessInvalidation();
 		if (NeedsSlowPath())
 		{
diff --git a/Engine/Source/Runtime/Slate/Public/Widgets/SInvalidationPanel.h b/Engine/Source/Runtime/Slate/Public/Widgets/SInvalidationPanel.h
--- a/Engine/Source/Runtime/Slate/Public/Widgets/SInvalidationPanel.h
+++ b/Engine/Source/Runtime/Slate/Public/Widgets/SInvalidationPanel.h
@@ -80,7 +80,7 @@ public:
 	 */
 	SLATE_API bool CanUseDynamicInvalidation() const;
 
-	SLATE_API virtual bool SupportsInvalidationRecursive(bool bUseCachedValue = true) const override;
+	SLATE_API virtual bool SupportsInvalidationRecursive(EInvalidationStrategy InInvalidationStrategy) const override;
 #endif
 
 	//~ SWidget overrides
diff --git a/Engine/Source/Runtime/SlateCore/Private/FastUpdate/SlateInvalidationRoot.cpp b/Engine/Source/Runtime/SlateCore/Private/FastUpdate/SlateInvalidationRoot.cpp
--- a/Engine/Source/Runtime/SlateCore/Private/FastUpdate/SlateInvalidationRoot.cpp
+++ b/Engine/Source/Runtime/SlateCore/Private/FastUpdate/SlateInvalidationRoot.cpp
@@ -831,6 +831,7 @@ bool FSlateInvalidationRoot::PaintFastPath(const FSlateInvalidationContext& Cont
 	if (bExecuteSlowPath)
 	{
 		SCOPED_NAMED_EVENT(Slate_PaintSlowPath, FColor::Red);
+		Context.WindowElementList->PopCachedElementData(); // to ensure we are not caching anymore
 		CachedMaxLayerId = PaintSlowPath(Context);
 	}
 
@@ -877,23 +878,6 @@ void FSlateInvalidationRoot::ProcessPreUpdate()
 
 	TGuardValue<bool> Tmp(bProcessingPreUpdate, true);
 
-#if UE_SLATE_WITH_DYNAMIC_INVALIDATION
-	// We need to process the list first as the code below might recreate the list and lose the CurrentInvalidateReason
-	WidgetsNeedingPreUpdate->ForEachIndexes([this](FSlateInvalidationWidgetHeapElement& Element)
-	{
-		const FSlateInvalidationWidgetList::InvalidationWidgetType& InvalidationWidget = (*FastWidgetPathList)[Element.GetWidgetIndex()];
-		const SWidget* WidgetPtr = InvalidationWidget.GetWidget();
-		const EInvalidateWidgetReason CurrentInvalidateReason = InvalidationWidget.CurrentInvalidateReason;
-		if (WidgetPtr && EnumHasAnyFlags(CurrentInvalidateReason, EInvalidateWidgetReason::ChildOrder) )
-		{
-			// If the children have been modified in any way, the cached value of SupportsInvalidationRecursive might have changed.
-			// We call SupportsInvalidationRecursive to recompute the value and invalidate the parent's value if needed
-			const bool bUseCachedValue = false;
-			bool bNowSupportsInvalidation = WidgetPtr->SupportsInvalidationRecursive(bUseCachedValue);
-		}
-	});
-#endif
-
 	TSharedRef<SWidget> RootWidget = GetRootWidget();
 	if (FastWidgetPathList->GetRoot().Pin() != RootWidget)
 	{
diff --git a/Engine/Source/Runtime/SlateCore/Private/Rendering/DrawElements.cpp b/Engine/Source/Runtime/SlateCore/Private/Rendering/DrawElements.cpp
--- a/Engine/Source/Runtime/SlateCore/Private/Rendering/DrawElements.cpp
+++ b/Engine/Source/Runtime/SlateCore/Private/Rendering/DrawElements.cpp
@@ -215,7 +215,10 @@ void FSlateWindowElementList::PushCachedElementData(FSlateCachedElementData& Cac
 
 void FSlateWindowElementList::PopCachedElementData()
 {
-	CachedElementDataListStack.Pop();
+	if (!CachedElementDataListStack.IsEmpty())
+	{
+		CachedElementDataListStack.Pop();
+	}
 }
 
 FSlateDrawElement& FSlateWindowElementList::AddUninitializedLookup(EElementType InElementType)
diff --git a/Engine/Source/Runtime/SlateCore/Private/Widgets/SNullWidget.cpp b/Engine/Source/Runtime/SlateCore/Private/Widgets/SNullWidget.cpp
--- a/Engine/Source/Runtime/SlateCore/Private/Widgets/SNullWidget.cpp
+++ b/Engine/Source/Runtime/SlateCore/Private/Widgets/SNullWidget.cpp
@@ -59,6 +59,13 @@ public:
 		// Nothing to arrange; Null Widgets do not have children.
 	}
 
+#if UE_SLATE_WITH_DYNAMIC_INVALIDATION
+	virtual bool SupportsInvalidationRecursive(EInvalidationStrategy InInvalidationStrategy) const override
+	{
+		return true;
+	}
+#endif
+
 protected:
 	virtual FVector2D ComputeDesiredSize(float) const override final
 	{
diff --git a/Engine/Source/Runtime/SlateCore/Private/Widgets/SWidget.cpp b/Engine/Source/Runtime/SlateCore/Private/Widgets/SWidget.cpp
--- a/Engine/Source/Runtime/SlateCore/Private/Widgets/SWidget.cpp
+++ b/Engine/Source/Runtime/SlateCore/Private/Widgets/SWidget.cpp
@@ -730,6 +730,7 @@ void SWidget::AssignParentWidget(TSharedPtr<SWidget> InParent)
 
 	//@todo We should update inherited visibility and volatility here but currently we are relying on ChildOrder invalidation to do that for us
 
+	const TSharedPtr<SWidget> PreviousParent = ParentWidgetPtr.Pin();
 	ParentWidgetPtr = InParent;
 #if WITH_ACCESSIBILITY
 	if (FSlateApplicationBase::IsInitialized())
@@ -737,6 +738,12 @@ void SWidget::AssignParentWidget(TSharedPtr<SWidget> InParent)
 		FSlateApplicationBase::Get().GetAccessibleMessageHandler()->MarkDirty();
 	}
 #endif
+	if (PreviousParent.IsValid())
+	{
+		// We want to be sure that the previous parent gets a ChildOrder invalidation for Dynamic Invalidation,
+		// otherwise we might not be able to catch widgets having the wrong parents
+		PreviousParent->Invalidate(EInvalidateWidgetReason::ChildOrder);
+	}
 	if (InParent.IsValid())
 	{
 		InParent->Invalidate(EInvalidateWidgetReason::ChildOrder);
@@ -1280,6 +1287,14 @@ void SWidget::Invalidate(EInvalidateWidgetReason InvalidateReason)
 		InvalidateReason |= EInvalidateWidgetReason::Layout;
 	}
 
+#if UE_SLATE_WITH_DYNAMIC_INVALIDATION
+	if (EnumHasAnyFlags(InvalidateReason, EInvalidateWidgetReason::ChildOrder))
+	{
+		// This is done here rather than in the FSlateInvalidationRoot to help catching all cases of widgets having wrong parents
+		InvalidateCachedSupportsInvalidation();
+	}
+#endif
+
 	// NB Advanced_InvalidateVolatility needs to be called to update the bCachedVolatility
 	const bool bVolatilityChanged = EnumHasAnyFlags(InvalidateReason, EInvalidateWidgetReason::Volatility) ? Advanced_InvalidateVolatility() : false;
 	if(FastPathProxyHandle.IsValid(this))
@@ -1969,87 +1984,143 @@ bool SWidget::SupportsInvalidation() const
 	return bSupportsInvalidation;
 }
 
-UE_TRACE_EVENT_BEGIN(Cpu, SWidget_SupportsInvalidationRecursive, NoSync)
-	UE_TRACE_EVENT_FIELD(UE::Trace::WideString, WidgetClassName)
-	UE_TRACE_EVENT_FIELD(UE::Trace::WideString, CreatedInLocation)
-	UE_TRACE_EVENT_FIELD(bool, UseCachedValue)
-	UE_TRACE_EVENT_FIELD(bool, OldCanBeCached)
-UE_TRACE_EVENT_END()
-
-bool SWidget::SupportsInvalidationRecursive(bool bUseCachedValue) const
+bool SWidget::SupportsInvalidationRecursive(EInvalidationStrategy InInvalidationStrategy) const
 {
-	if (!SupportsInvalidation())
+	if ((InInvalidationStrategy == EInvalidationStrategy::UseCachedValue) &
+		(CachedSupportsInvalidationRecursive != ECachedInvalidationResult::NeedsRefresh))
 	{
-		CachedSupportsInvalidationRecursive = false;
-		return false;
+		return CachedSupportsInvalidationRecursive == ECachedInvalidationResult::SupportsInvalidation;
 	}
 
-	// If the cached value was true, then all children supported invalidation and the widget would have been invalidated if needed.
-	// If the cached value was false, at least one of the descendant widget did not support invalidation and we therefore need to recheck everytime
-	if (CachedSupportsInvalidationRecursive && bUseCachedValue)
+	if (!SupportsInvalidation())
 	{
-		return true;
+		check(CachedSupportsInvalidationRecursive == ECachedInvalidationResult::ClassNotSupportingInvalidation);
+		return false;
 	}
+	check(CachedSupportsInvalidationRecursive != ECachedInvalidationResult::ClassNotSupportingInvalidation);
 
-#if UE_SLATE_TRACE_ENABLED
-	UE_TRACE_LOG_SCOPED_T(Cpu, SWidget_SupportsInvalidationRecursive, SlateChannel)
-		<< SWidget_SupportsInvalidationRecursive.WidgetClassName(*TypeOfWidget.ToString())
-		<< SWidget_SupportsInvalidationRecursive.CreatedInLocation(*CreatedInLocation.ToString())
-		<< SWidget_SupportsInvalidationRecursive.UseCachedValue(bUseCachedValue)
-		<< SWidget_SupportsInvalidationRecursive.OldCanBeCached(CachedSupportsInvalidationRecursive);
-#endif // UE_SLATE_TRACE_ENABLED
 
-	const bool OldCachedValue = CachedSupportsInvalidationRecursive;
-	CachedSupportsInvalidationRecursive = true;
+	const ECachedInvalidationResult PreviousCachedValue = CachedSupportsInvalidationRecursive;
+	// Sets the value to Refreshing to avoid checking this widget multiple times
+	CachedSupportsInvalidationRecursive = ECachedInvalidationResult::Refreshing;
 
-	const FChildren* Children = const_cast<SWidget*>(this)->GetAllChildren();
+	FChildren* Children = const_cast<SWidget*>(this)->GetAllChildren();
 	for (int32 ChildIndex = 0; ChildIndex < Children->Num(); ++ChildIndex)
 	{
-		const TSharedRef<const SWidget> ChildWidget = Children->GetChildAt(ChildIndex);
-		if (ChildWidget->ParentWidgetPtr != nullptr && !ChildWidget->SupportsInvalidationRecursive())
+		const TSharedRef<SWidget> ChildWidget = Children->GetChildAt(ChildIndex);
+		if (ChildWidget == SNullWidget::NullWidget)
+		{
+			continue;
+		}
+
+		if (!ChildWidget->ParentWidgetPtr.HasSameObject(this))
+		{
+			// Here the child  widget has another widget as parent, so we need to mark them as not supporting invalidation
+			CachedSupportsInvalidationRecursive = ECachedInvalidationResult::ChildrenNotSupportingInvalidation;
+			if (ChildWidget->CachedSupportsInvalidationRecursive != ECachedInvalidationResult::ClassNotSupportingInvalidation)
+			{
+				ChildWidget->CachedSupportsInvalidationRecursive = ECachedInvalidationResult::InvalidParent;
+			}
+			break;
+		}
+
+		// If we had an invalid parent but we don't have one anymore, we recheck if we are invalidation compatible, otherwise we use the cached value
+		const EInvalidationStrategy Strategy = ChildWidget->CachedSupportsInvalidationRecursive == ECachedInvalidationResult::InvalidParent ?
+			EInvalidationStrategy::RecomputeSelf : EInvalidationStrategy::UseCachedValue;
+
+		if (!ChildWidget->SupportsInvalidationRecursive(Strategy))
 		{
-			CachedSupportsInvalidationRecursive = false;
+			CachedSupportsInvalidationRecursive = ECachedInvalidationResult::ChildrenNotSupportingInvalidation;
 			break;
 		}
 	}
-	
-	// If the value has changed, we invalidate the parent
-	if (CachedSupportsInvalidationRecursive != OldCachedValue)
+
+	// If we have checked all the children and none of them were invalid, then we are supporting invalidation
+	if (CachedSupportsInvalidationRecursive == ECachedInvalidationResult::Refreshing)
+	{
+		CachedSupportsInvalidationRecursive = ECachedInvalidationResult::SupportsInvalidation;
+	}
+
+	// If the value has changed, we invalidate the parent unless it is currently recomputing
+	if (CachedSupportsInvalidationRecursive != PreviousCachedValue)
 	{
-		if (const TSharedPtr<SWidget> ParentWidget = GetParentWidget())
+		const TSharedPtr<SWidget> ParentWidget = GetParentWidget();
+		if (ParentWidget && ParentWidget->CachedSupportsInvalidationRecursive != ECachedInvalidationResult::Refreshing)
 		{
-			if (!CachedSupportsInvalidationRecursive)
+			if (CachedSupportsInvalidationRecursive != ECachedInvalidationResult::SupportsInvalidation)
 			{
 				// If we supported invalidation but we do not anymore, tell the parent they do not support invalidation anymore, no need for them to check
-				ParentWidget->SetSelfAndParentAsNotSupportingInvalidationRecursive();
+				ParentWidget->SetSelfAndParentsAsNotSupportingInvalidation();
 			}
 			else
 			{
-				// If we now support invalidation, ask the parent to recompute.
-				const bool bParentUseCachedValue = false;
-				bool bNowSupportsInvalidation = ParentWidget->SupportsInvalidationRecursive(bParentUseCachedValue);
+				// If we now support invalidation, ask the parent to recompute
+				ParentWidget->SupportsInvalidationRecursive(EInvalidationStrategy::RecomputeSelf);
 			}
 		}
 	}
 
-	return CachedSupportsInvalidationRecursive;
+	return CachedSupportsInvalidationRecursive == ECachedInvalidationResult::SupportsInvalidation;
 }
 
-void SWidget::SetSelfAndParentAsNotSupportingInvalidationRecursive()
+void SWidget::SetSelfAndParentsAsNotSupportingInvalidation()
 {
-	// if this widget was supporting invalidation and it does not anymore, it needs to tell its parent.
-	// Otherwise, its parent already knew so no need to do anything.
-	if (CachedSupportsInvalidationRecursive)
+	TSharedPtr<SWidget> Widget = (this == &SNullWidget::NullWidget.Get()) ? GetParentWidget() : AsShared();
+	while (Widget)
 	{
-		CachedSupportsInvalidationRecursive = false;
-		if (const TSharedPtr<SWidget> ParentWidget = GetParentWidget())
+		// If this widget was supporting invalidation and it does not anymore, it needs to tell its parent.
+		// Otherwise, its parent already knew so no need to do anything.
+		if (Widget->CachedSupportsInvalidationRecursive == ECachedInvalidationResult::SupportsInvalidation)
+		{
+			Widget->CachedSupportsInvalidationRecursive = ECachedInvalidationResult::ChildrenNotSupportingInvalidation;
+			Widget = Widget->GetParentWidget();
+		}
+		else
 		{
-			ParentWidget->SetSelfAndParentAsNotSupportingInvalidationRecursive();
+			return;
 		}
 	}
 }
 
-#endif
+UE_TRACE_EVENT_BEGIN(Cpu, SWidget_InvalidateCachedSupportsInvalidation, NoSync)
+	UE_TRACE_EVENT_FIELD(UE::Trace::WideString, CreatedInLocation)
+UE_TRACE_EVENT_END()
+
+void SWidget::InvalidateCachedSupportsInvalidation()
+{
+#if UE_SLATE_TRACE_ENABLED
+	UE_TRACE_LOG_SCOPED_T(Cpu, SWidget_InvalidateCachedSupportsInvalidation, SlateChannel)
+		<< SWidget_InvalidateCachedSupportsInvalidation.CreatedInLocation(*GetCreatedInLocation().ToString());
+#endif // UE_SLATE_TRACE_ENABLED
+
+	TSharedPtr<SWidget> Widget = (this == &SNullWidget::NullWidget.Get()) ? GetParentWidget() : AsShared();
+	while (Widget)
+	{
+		if ((Widget->CachedSupportsInvalidationRecursive == ECachedInvalidationResult::ClassNotSupportingInvalidation)
+			| (Widget->CachedSupportsInvalidationRecursive == ECachedInvalidationResult::NeedsRefresh))
+		{
+			return;
+		}
+
+		Widget->CachedSupportsInvalidationRecursive = ECachedInvalidationResult::NeedsRefresh;
+		if (GSlateIsOnFastUpdatePath & Widget->Advanced_IsInvalidationRoot())
+		{
+			// If we are currently in fast path, we want to stop and do a slow path as we might have invalid widgets that have been added
+			const_cast<FSlateInvalidationRoot*>(Widget->Advanced_AsInvalidationRoot())->InvalidateRootLayout(this);
+		}
+
+		Widget = Widget->GetParentWidget();
+	}
+}
+
+void SWidget::SetClassSupportsInvalidation(bool bClassSupportsInvalidation)
+{
+	bSupportsInvalidation = bClassSupportsInvalidation;
+	CachedSupportsInvalidationRecursive = bSupportsInvalidation ?
+		ECachedInvalidationResult::NeedsRefresh : ECachedInvalidationResult::ClassNotSupportingInvalidation;
+}
+
+#endif // ~UE_SLATE_WITH_DYNAMIC_INVALIDATION
 
 #if WITH_ACCESSIBILITY
 TSharedRef<FSlateAccessibleWidget> SWidget::CreateAccessibleWidget()
diff --git a/Engine/Source/Runtime/SlateCore/Public/Widgets/DeclarativeSyntaxSupport.h b/Engine/Source/Runtime/SlateCore/Public/Widgets/DeclarativeSyntaxSupport.h
--- a/Engine/Source/Runtime/SlateCore/Public/Widgets/DeclarativeSyntaxSupport.h
+++ b/Engine/Source/Runtime/SlateCore/Public/Widgets/DeclarativeSyntaxSupport.h
@@ -929,7 +929,7 @@ struct TSlateDecl
 		_Widget->SetDebugInfo(InType, InFile, OnLine, sizeof(WidgetType));
 
 #if UE_SLATE_WITH_DYNAMIC_INVALIDATION
-		_Widget->bSupportsInvalidation = TWidgetTypeTraits<WidgetType>::SupportsInvalidation();
+		_Widget->SetClassSupportsInvalidation(TWidgetTypeTraits<WidgetType>::SupportsInvalidation());
 #endif
 	}
 
diff --git a/Engine/Source/Runtime/SlateCore/Public/Widgets/SWidget.h b/Engine/Source/Runtime/SlateCore/Public/Widgets/SWidget.h
--- a/Engine/Source/Runtime/SlateCore/Public/Widgets/SWidget.h
+++ b/Engine/Source/Runtime/SlateCore/Public/Widgets/SWidget.h
@@ -1534,21 +1534,54 @@ public:
 
 #if UE_SLATE_WITH_DYNAMIC_INVALIDATION
 
-	/** Returns true if this widget's class supports invalidation. */
+	/** Returns true if this widget's class supports invalidation. This value never changes. */
 	SLATECORE_API bool SupportsInvalidation() const;
 
+	enum class EInvalidationStrategy : uint8
+	{
+		/** Use the last cached value unless it needs refreshing. */
+		UseCachedValue,
+		/** Check again if this widget and its direct children support invalidation. */
+		RecomputeSelf,
+	};
+
+	enum class ECachedInvalidationResult : uint8
+	{
+		/** The last result has been invalidated and needs to be refreshed. This will happen in the next call to SupportsInvalidationRecursive. */
+		NeedsRefresh,
+		/** Used internally in SupportsInvalidationRecursive to avoid refreshing the same widget multiple times. */
+		Refreshing,
+		/** This widget and all its descendants support invalidation and have a valid parent. */
+		SupportsInvalidation,
+		/** The class of this widget does not support invalidation. This result never changes. */
+		ClassNotSupportingInvalidation,
+		/** Descendants of this widget are not supporting invalidation or have invalid parents. */
+		ChildrenNotSupportingInvalidation,
+		/**
+		 * This widget currently has an invalid parent, meaning it is a child of widget A but its Parent points to another Widget B (or null).
+		 * It is likely the case that this widget has been reused in multiple slots which is not allowed with invalidation.
+		 */
+		InvalidParent,
+	};
+
 	/**
 	 * Returns true if this widget and all of its children support invalidation.
-	 * This value is cached and reused unless bUseCachedValue is set to false.
+	 * This value is cached, and will be reused if InInvalidationStrategy is set to UseCachedValue.
 	 */
-	SLATECORE_API virtual bool SupportsInvalidationRecursive(bool bUseCachedValue = true) const;
+	SLATECORE_API virtual bool SupportsInvalidationRecursive(EInvalidationStrategy InInvalidationStrategy) const;
 private:
-	/** Fast path to invalidate all parents when a descendant widget is not supporting invalidation.. */
-	SLATECORE_API void SetSelfAndParentAsNotSupportingInvalidationRecursive();
-;
-	/** The cached value of the last SupportsInvalidationRecursive call. */
-	mutable bool CachedSupportsInvalidationRecursive = false;
-#endif
+	/** Fast path to invalidate all ancestors when a descendant widget is not supporting invalidation. */
+	void SetSelfAndParentsAsNotSupportingInvalidation();
+
+	/** Sets this widget and its ancestors as needing to have their SupportsInvalidationRecursive refreshed. */
+	void InvalidateCachedSupportsInvalidation();
+
+	/** Store the value from the TWidgetTypeTraits<>::SupportsInvalidation just after the widget is created */
+	SLATECORE_API void SetClassSupportsInvalidation(bool bClassSupportsInvalidation);
+
+public:
+	ECachedInvalidationResult Debug_GetCachedInvalidationResult() const { return CachedSupportsInvalidationRecursive; }
+#endif // ~UE_SLATE_WITH_DYNAMIC_INVALIDATION
 
 protected:
 	/**
@@ -1828,6 +1861,11 @@ protected:
 	EAccessibleBehavior AccessibleSummaryBehavior;
 #endif
 
+#if UE_SLATE_WITH_DYNAMIC_INVALIDATION
+	/** The cached value of the last SupportsInvalidationRecursive call. */
+	mutable ECachedInvalidationResult CachedSupportsInvalidationRecursive = ECachedInvalidationResult::ChildrenNotSupportingInvalidation;
+#endif
+
 	/**
 	 * Set to true if all content of the widget should clip to the bounds of this widget.
 	 */
@@ -1997,3 +2035,38 @@ bool SetWidgetAttribute(SWidget& ThisWidget, TAttribute<TargetValueType>& Target
 
 	return false;
 }
+
+
+#if UE_SLATE_WITH_DYNAMIC_INVALIDATION
+
+inline const TCHAR* LexToString(SWidget::EInvalidationStrategy Strategy)
+{
+#define INVALIDATION_STRATEGY_LEXTOSTRING(Name) case SWidget::EInvalidationStrategy::Name: return TEXT(#Name);
+	switch (Strategy)
+	{
+		INVALIDATION_STRATEGY_LEXTOSTRING(UseCachedValue)
+		INVALIDATION_STRATEGY_LEXTOSTRING(RecomputeSelf)
+	default:
+		return TEXT("Unknown");
+	}
+#undef INVALIDATION_STRATEGY_LEXTOSTRING
+}
+
+inline const TCHAR* LexToString(SWidget::ECachedInvalidationResult Result)
+{
+#define INVALIDATION_RESULT_LEXTOSTRING(Name) case SWidget::ECachedInvalidationResult::Name: return TEXT(#Name);
+	switch (Result)
+	{
+		INVALIDATION_RESULT_LEXTOSTRING(NeedsRefresh)
+		INVALIDATION_RESULT_LEXTOSTRING(Refreshing)
+		INVALIDATION_RESULT_LEXTOSTRING(SupportsInvalidation)
+		INVALIDATION_RESULT_LEXTOSTRING(ClassNotSupportingInvalidation)
+		INVALIDATION_RESULT_LEXTOSTRING(ChildrenNotSupportingInvalidation)
+		INVALIDATION_RESULT_LEXTOSTRING(InvalidParent)
+	default:
+		return TEXT("Unknown");
+	}
+#undef INVALIDATION_RESULT_LEXTOSTRING
+}
+
+#endif // ~UE_SLATE_WITH_DYNAMIC_INVALIDATION
\ No newline at end of file
