Minor fix for Dynamic Invalidation

#Synced-CL 48711348

diff --git a/Engine/Source/Runtime/SlateCore/Private/Widgets/SWidget.cpp b/Engine/Source/Runtime/SlateCore/Private/Widgets/SWidget.cpp
--- a/Engine/Source/Runtime/SlateCore/Private/Widgets/SWidget.cpp
+++ b/Engine/Source/Runtime/SlateCore/Private/Widgets/SWidget.cpp
@@ -2103,7 +2103,7 @@ void SWidget::InvalidateCachedSupportsInvalidation()
 		}
 
 		Widget->CachedSupportsInvalidationRecursive = ECachedInvalidationResult::NeedsRefresh;
-		if (GSlateIsOnFastUpdatePath & Widget->Advanced_IsInvalidationRoot())
+		if (GSlateIsOnFastUpdatePath && Widget->Advanced_IsInvalidationRoot())
 		{
 			// If we are currently in fast path, we want to stop and do a slow path as we might have invalid widgets that have been added
 			const_cast<FSlateInvalidationRoot*>(Widget->Advanced_AsInvalidationRoot())->InvalidateRootLayout(this);
