Slate Core: Added widget hierarchy to the Slate Crash Reporter. Added debug info to SImage's OnPaint  method.

#Synced-CL 48758403

diff --git a/Engine/Source/Runtime/SlateCore/Private/Debugging/SlateCrashReporterHandler.cpp b/Engine/Source/Runtime/SlateCore/Private/Debugging/SlateCrashReporterHandler.cpp
--- a/Engine/Source/Runtime/SlateCore/Private/Debugging/SlateCrashReporterHandler.cpp
+++ b/Engine/Source/Runtime/SlateCore/Private/Debugging/SlateCrashReporterHandler.cpp
@@ -3,6 +3,7 @@
 #include "Debugging/SlateCrashReporterHandler.h"
 
 #include "HAL/ThreadManager.h"
+#include "Layout/ChildrenBase.h"
 #include "Misc/ScopeLock.h"
 
 #include "Widgets/SWidget.h"
@@ -35,7 +36,8 @@ private:
 	struct FContext
 	{
 		FContext(const SWidget& Widget)
-			: Type(Widget.GetType())
+			: WidgetWeak(Widget.AsWeak())
+			, Type(Widget.GetType())
 			, Tag(Widget.GetTag())
 			, CreatedInLocation(Widget.GetCreatedInLocation())
 			, ContextName(NAME_None)
@@ -48,6 +50,7 @@ private:
 			ContextName = InContextName;
 		}
 
+		TWeakPtr<const SWidget> WidgetWeak;
 		FName Type;
 		FName Tag;
 		FName CreatedInLocation;
@@ -146,6 +149,11 @@ public:
 					Context.Tag.AppendString(Builder);
 					Builder.Append(TEXT("' Location:'"));
 					Context.CreatedInLocation.AppendString(Builder);
+					if (TSharedPtr<const SWidget> Widget = Context.WidgetWeak.Pin())
+					{
+						Builder.Append(TEXT("' Path:'"));
+						AddWidgetPath(Widget);
+					}
 					Builder.Append(TEXT("'\n"));
 				}
 			}
@@ -160,7 +168,72 @@ public:
 private:
 	FCriticalSection CriticalSection;
 	TArray<FContext> Contexts[NumberOfThread];
-	TStringBuilder<4096> Builder;
+	TStringBuilder<16384> Builder;
+
+	void AddWidgetPath(TSharedPtr<const SWidget> InLeafNode)
+	{
+		constexpr int32 StackSize = 500;
+
+		SWidget* Stack[StackSize];
+		int32 StackPtr = 0;
+
+		Stack[StackPtr] = const_cast<SWidget*>(InLeafNode.Get());
+
+		while (Stack[StackPtr] && StackPtr < (StackSize - 1))
+		{
+			++StackPtr;
+			Stack[StackPtr] = Stack[StackPtr - 1]->GetParentWidget().Get();
+		}
+
+		// If the final parent was null, decrease the ptr to a valid widget.
+		if (!Stack[StackPtr])
+		{
+			--StackPtr;
+		}
+
+		auto GetChildIndex = [](SWidget* InParent, SWidget* InChild) -> int32
+			{
+				if (FChildren* Children = InParent->GetChildren())
+				{
+					for (int32 Index = 0; Index < Children->Num(); ++Index)
+					{
+						if (&*Children->GetChildAt(Index) == InChild)
+						{
+							return Index;
+						}
+					}
+				}
+
+				return -1;
+			};
+
+		while (StackPtr >= 0)
+		{
+			Builder.Append(*Stack[StackPtr]->GetTypeAsString());
+
+			const FName WidgetTag = Stack[StackPtr]->GetTag();
+
+			if (!WidgetTag.IsNone())
+			{
+				Builder.Appendf(TEXT(" (%s)"), *WidgetTag.ToString());
+			}
+
+			if (StackPtr > 0)
+			{
+				if (FChildren* Children = Stack[StackPtr]->GetChildren())
+				{
+					if (Children->Num() > 1)
+					{
+						Builder.Appendf(TEXT(" [%i]"), GetChildIndex(Stack[StackPtr], Stack[StackPtr - 1]));
+					}
+				}
+
+				Builder.Append(TEXT(": "));
+			}
+
+			--StackPtr;
+		}
+	}
 };
 
 } // namespace
diff --git a/Engine/Source/Runtime/SlateCore/Private/Widgets/Images/SImage.cpp b/Engine/Source/Runtime/SlateCore/Private/Widgets/Images/SImage.cpp
--- a/Engine/Source/Runtime/SlateCore/Private/Widgets/Images/SImage.cpp
+++ b/Engine/Source/Runtime/SlateCore/Private/Widgets/Images/SImage.cpp
@@ -6,7 +6,7 @@
 #if WITH_ACCESSIBILITY
 #include "Widgets/Accessibility/SlateCoreAccessibleWidgets.h"
 #endif
-
+#include "Debugging/SlateCrashReporterHandler.h"
 
 SLATE_IMPLEMENT_WIDGET(SImage)
 void SImage::PrivateRegisterAttributes(FSlateAttributeInitializer& AttributeInitializer)
@@ -43,6 +43,10 @@ void SImage::Construct( const FArguments& InArgs )
 
 int32 SImage::OnPaint( const FPaintArgs& Args, const FGeometry& AllottedGeometry, const FSlateRect& MyCullingRect, FSlateWindowElementList& OutDrawElements, int32 LayerId, const FWidgetStyle& InWidgetStyle, bool bParentEnabled ) const
 {
+#if UE_WITH_SLATE_CRASHREPORTER
+	UE_SLATE_CRASH_REPORTER_CONTEXT_SCOPE(*this, "ImagePaint");
+#endif
+
 	const FSlateBrush* ImageBrush = ImageAttribute.Get();
 
 	if ((ImageBrush != nullptr) && (ImageBrush->DrawAs != ESlateBrushDrawType::NoDrawType))
