Build issue 1149702

#Synced-CL 48758403

diff --git a/Engine/Source/Runtime/SlateCore/Private/Debugging/SlateCrashReporterHandler.cpp b/Engine/Source/Runtime/SlateCore/Private/Debugging/SlateCrashReporterHandler.cpp
--- a/Engine/Source/Runtime/SlateCore/Private/Debugging/SlateCrashReporterHandler.cpp
+++ b/Engine/Source/Runtime/SlateCore/Private/Debugging/SlateCrashReporterHandler.cpp
@@ -179,7 +179,7 @@ private:
 
 		Stack[StackPtr] = const_cast<SWidget*>(InLeafNode.Get());
 
-		while (Stack[StackPtr] && StackPtr < (StackSize - 1))
+		while (Stack[StackPtr] && StackPtr < (StackSize - 1)) //-V781
 		{
 			++StackPtr;
 			Stack[StackPtr] = Stack[StackPtr - 1]->GetParentWidget().Get();
