Change a SlateDebugger ensure to a log as it triggers a lot of false positives.

#Synced-CL 48711348

diff --git a/Engine/Source/Runtime/SlateCore/Private/Debugging/ConsoleSlateDebuggerPassBase.cpp b/Engine/Source/Runtime/SlateCore/Private/Debugging/ConsoleSlateDebuggerPassBase.cpp
--- a/Engine/Source/Runtime/SlateCore/Private/Debugging/ConsoleSlateDebuggerPassBase.cpp
+++ b/Engine/Source/Runtime/SlateCore/Private/Debugging/ConsoleSlateDebuggerPassBase.cpp
@@ -358,8 +358,9 @@ FConsoleSlateDebuggerPassBase::FWidgetInfo& FConsoleSlateDebuggerPassBase::AddUp
 		FoundItem->Window = WindowId;
 		FoundItem->WidgetName = bUseWidgetPathAsName ? FReflectionMetaData::GetWidgetPath(Widget) : FReflectionMetaData::GetWidgetDebugInfo(Widget);
 	}
-	else if (!ensure(FoundItem->Window == WindowId))
+	else if (FoundItem->Window != WindowId)
 	{
+		UE_LOG(LogSlateDebugger, Log, TEXT("The WindowId of this widget '%s' has changed since the last time it was seen."), *FoundItem->WidgetName);
 		FoundItem->Window = WindowId;
 	}
 	
