An `ensure` can incorrectly trigger when reverting from FastPath back to SlowPath mid-execution.
If a widget `Visibility` is changed at this time, `SWidget::Invalidate` calls `FastPathProxyHandle.MarkWidgetDirty_NoCheck`, which then invokes `FSlateInvalidationRoot::InvalidateWidget`. At that moment, if the system is transitioning back into SlowPath, the flag `bNeedsSlowPath` will be true, and `Proxy.CurrentInvalidateReason` will be left unchanged. This mismatch ends up triggering the ensure later on.
To avoid this, we chose to not do these debug checks if `bNeedsSlowPath` is true, as they are not relevant anymore.

#Synced-CL 48711348

diff --git a/Engine/Source/Runtime/SlateCore/Private/Widgets/SWidget.cpp b/Engine/Source/Runtime/SlateCore/Private/Widgets/SWidget.cpp
--- a/Engine/Source/Runtime/SlateCore/Private/Widgets/SWidget.cpp
+++ b/Engine/Source/Runtime/SlateCore/Private/Widgets/SWidget.cpp
@@ -797,7 +797,8 @@ void SWidget::UpdateWidgetProxy(int32 NewLayerId, FSlateCachedElementsHandle& Ca
 	PersistentState.OutgoingLayerId = NewLayerId;
 
 #if WITH_SLATE_DEBUGGING
-	if (FastPathProxyHandle.IsValid(this))
+	const FSlateInvalidationRoot* InvalidationRoot = FastPathProxyHandle.GetInvalidationRootHandle().GetInvalidationRoot();
+	if (FastPathProxyHandle.IsValid(this) && InvalidationRoot && !InvalidationRoot->NeedsSlowPath())
 	{
 		FWidgetProxy& MyProxy = FastPathProxyHandle.GetProxy();
 #if UE_SLATE_WITH_INVALIDATIONWIDGETLIST_DEBUGGING
