Fix to not draw hidden widgets in the Slate Debugger Paint/Prepass Views

#Synced-CL 48711348

diff --git a/Engine/Source/Runtime/SlateCore/Private/Debugging/ConsoleSlateDebuggerPassBase.cpp b/Engine/Source/Runtime/SlateCore/Private/Debugging/ConsoleSlateDebuggerPassBase.cpp
--- a/Engine/Source/Runtime/SlateCore/Private/Debugging/ConsoleSlateDebuggerPassBase.cpp
+++ b/Engine/Source/Runtime/SlateCore/Private/Debugging/ConsoleSlateDebuggerPassBase.cpp
@@ -228,7 +228,7 @@ void FConsoleSlateDebuggerPassBase::HandlePaintDebugInfo(const FPaintArgs& InArg
 		if (TSharedPtr<const SWidget> Widget = WidgetInfo.Widget.Pin())
 		{
 			// Skip the widget if collapsed or not in a window anymore (ex: in an unactive tab)
-			if (Widget->GetProxyHandle().GetWidgetVisibility(Widget.Get()).IsCollapsed() || !FSlateApplicationBase::Get().FindWidgetWindow(Widget.ToSharedRef()))
+			if (!Widget->GetProxyHandle().GetWidgetVisibility(Widget.Get()).IsVisible() || !FSlateApplicationBase::Get().FindWidgetWindow(Widget.ToSharedRef()))
 			{
 				continue;
 			}
