Some Slate Runtime Validations are failing (`Verify Hittest Grid` and `Verify Cached Elements`) because in SlowPath, widgets newly created on tick will have invalid widget handles until they are updated on the next frame.
This CL ensures these widgets are processed the same way in SlowPath and FastPath.

#Synced-CL 49309628

diff --git a/Engine/Source/Runtime/SlateCore/Private/FastUpdate/SlateInvalidationRoot.cpp b/Engine/Source/Runtime/SlateCore/Private/FastUpdate/SlateInvalidationRoot.cpp
--- a/Engine/Source/Runtime/SlateCore/Private/FastUpdate/SlateInvalidationRoot.cpp
+++ b/Engine/Source/Runtime/SlateCore/Private/FastUpdate/SlateInvalidationRoot.cpp
@@ -415,7 +415,6 @@ FSlateInvalidationResult FSlateInvalidationRoot::PaintInvalidationRoot(const FSl
 		}
 
 		Result.bRepaintedWidgets = true;
-
 	}
 	else if (!FastWidgetPathList->IsEmpty())
 	{
@@ -423,14 +422,14 @@ FSlateInvalidationResult FSlateInvalidationRoot::PaintInvalidationRoot(const FSl
 		check(RootWidget == FastWidgetPathList->GetRoot().Pin());
 
 		Result.bRepaintedWidgets = PaintFastPath(Context);
+	}
 
-		// Once everything is painted and ticked, new widget (SListView, SRichBox, SScrollBox) might have created new widget.
-		//It might create a 1 frame hitch because the widget is in the cached and should not.
-		if (!bNeedsSlowPath && WidgetsNeedingPreUpdate->Num() > 0)
-		{
-			SCOPED_NAMED_EVENT(Slate_InvalidationProcessing_PreUpdate, FColor::Blue);
-			ProcessPreUpdate();
-		}
+	// Once everything is painted and ticked, new widget (SListView, SRichBox, SScrollBox) might have created new widget in both SlowPath and FastPath
+	// It might create a 1 frame hitch because the widget is in the cached and should not.
+	if (!bNeedsSlowPath && WidgetsNeedingPreUpdate->Num() > 0)
+	{
+		SCOPED_NAMED_EVENT(Slate_InvalidationProcessing_PreUpdate, FColor::Blue);
+		ProcessPreUpdate();
 	}
 
 	if (Context.bAllowFastPathUpdate)
@@ -832,6 +831,7 @@ bool FSlateInvalidationRoot::PaintFastPath(const FSlateInvalidationContext& Cont
 	{
 		SCOPED_NAMED_EVENT(Slate_PaintSlowPath, FColor::Red);
 		Context.WindowElementList->PopCachedElementData(); // to ensure we are not caching anymore
+		ClearAllFastPathData(true);
 		CachedMaxLayerId = PaintSlowPath(Context);
 	}
 
@@ -1677,7 +1677,7 @@ void VerifyHittest(SWidget* InvalidationRootWidget, FSlateInvalidationWidgetList
 			{
 				UE_SLATE_LOG_ERROR_IF_FALSE(WidgetList.IsValidIndex(ParentWidetIndex) && EnumHasAllFlags(WidgetList[ParentWidetIndex].CurrentInvalidateReason, EInvalidateWidgetReason::ChildOrder)
 					, CVarSlateInvalidationRootVerifyHittestGrid
-					, TEXT("The parent widget doesn't has ChildOrder invalidation. The hittest grid has an invalid widget."));
+					, TEXT("The parent widget doesn't have ChildOrder invalidation. The hittest grid has an invalid widget."));
 			}
 		}
 	}
diff --git a/Engine/Source/Runtime/SlateCore/Private/Widgets/SWidget.cpp b/Engine/Source/Runtime/SlateCore/Private/Widgets/SWidget.cpp
--- a/Engine/Source/Runtime/SlateCore/Private/Widgets/SWidget.cpp
+++ b/Engine/Source/Runtime/SlateCore/Private/Widgets/SWidget.cpp
@@ -845,6 +845,13 @@ void SWidget::SetFastPathProxyHandle(const FWidgetProxyHandle& Handle, FSlateInv
 	check(this != &SNullWidget::NullWidget.Get());
 
 	FastPathProxyHandle = Handle;
+	if (FSlateInvalidationRoot* Root = FastPathProxyHandle.GetInvalidationRootHandle().GetInvalidationRoot())
+	{
+		if (FHittestGrid* HittestGrid = Root->GetHittestGrid())
+		{
+			HittestGrid->UpdateWidget(this, Handle.GetWidgetSortOrder());
+		}
+	}
 
 	bInheritedVolatility = bParentVolatile;
 
