[Slate] Add Profiler event to track the time spent adding new entries in the SDF font atlas.

#tests Validated in Unreal Insight that the profile point was visible.
#rnx
#Synced-CL 49178553

diff --git a/Engine/Source/Runtime/SlateCore/Private/Fonts/FontCache.cpp b/Engine/Source/Runtime/SlateCore/Private/Fonts/FontCache.cpp
--- a/Engine/Source/Runtime/SlateCore/Private/Fonts/FontCache.cpp
+++ b/Engine/Source/Runtime/SlateCore/Private/Fonts/FontCache.cpp
@@ -1455,6 +1455,8 @@ FSdfGlyphFontAtlasData FSlateFontCache::GetSdfGlyphFontAtlasData(const FShapedGl
 		}
 	}
 
+	TRACE_CPUPROFILER_EVENT_SCOPE_STR("FSlateFontCache::GetSdfGlyphFontAtlasData_AddEntry");
+
 	// Not cached at all... create a new entry
 	FSlateSdfGenerator::FRequestDescriptor SdfRequestDescriptor = { 
 		GlyphKey.FontFace,
