This CL defers any call to `FSlateInvalidationRoot::HandleInvalidateAllWidgets` until it is safe to do so at the end of `ProcessInvalidation`

#Synced-CL 49309628

diff --git a/Engine/Source/Runtime/SlateCore/Private/FastUpdate/SlateInvalidationRoot.cpp b/Engine/Source/Runtime/SlateCore/Private/FastUpdate/SlateInvalidationRoot.cpp
--- a/Engine/Source/Runtime/SlateCore/Private/FastUpdate/SlateInvalidationRoot.cpp
+++ b/Engine/Source/Runtime/SlateCore/Private/FastUpdate/SlateInvalidationRoot.cpp
@@ -219,6 +219,7 @@ FSlateInvalidationRoot::FSlateInvalidationRoot()
 	, bProcessingPostUpdate(false)
 	, bBuildingWidgetList(false)
 	, bProcessingChildOrderInvalidation(false)
+	, bDeferInvalidateAllWidgetsAndClearResources(TOptional<bool>{})
 #if WITH_SLATE_DEBUGGING
 	, LastPaintType(ESlateInvalidationPaintType::None)
 	, PerformanceStat()
@@ -1367,7 +1368,13 @@ bool FSlateInvalidationRoot::ProcessInvalidation()
 
 		bWidgetsNeedRepaint = ProcessPostUpdate();
 	}
-	
+
+	if (bDeferInvalidateAllWidgetsAndClearResources.IsSet())
+	{
+		HandleInvalidateAllWidgets(bDeferInvalidateAllWidgetsAndClearResources.GetValue());
+		bDeferInvalidateAllWidgetsAndClearResources = {};
+	}
+
 	if (bNeedsSlowPath)
 	{
 		WidgetsNeedingPreUpdate->Reset(true);
@@ -1438,8 +1445,15 @@ void FSlateInvalidationRoot::SetNeedsSlowPath(bool InNeedsSlowPath)
 
 void FSlateInvalidationRoot::HandleInvalidateAllWidgets(bool bClearResourcesImmediately)
 {
-	Advanced_ResetInvalidation(bClearResourcesImmediately);
-	OnRootInvalidated();
+	if (bProcessingPreUpdate | bProcessingAttributeUpdate | bProcessingPrepassUpdate | bProcessingPostUpdate)
+	{
+		bDeferInvalidateAllWidgetsAndClearResources = bClearResourcesImmediately;
+	}
+	else
+	{
+		Advanced_ResetInvalidation(bClearResourcesImmediately);
+		OnRootInvalidated();
+	}
 }
 
 void FSlateInvalidationRoot::Advanced_ResetInvalidation(bool bClearResourcesImmediately)
diff --git a/Engine/Source/Runtime/SlateCore/Public/FastUpdate/SlateInvalidationRoot.h b/Engine/Source/Runtime/SlateCore/Public/FastUpdate/SlateInvalidationRoot.h
--- a/Engine/Source/Runtime/SlateCore/Public/FastUpdate/SlateInvalidationRoot.h
+++ b/Engine/Source/Runtime/SlateCore/Public/FastUpdate/SlateInvalidationRoot.h
@@ -228,6 +228,12 @@ private:
 	bool bBuildingWidgetList;
 	bool bProcessingChildOrderInvalidation;
 
+	/**
+	 * If set, HandleInvalidateAllWidgets will be called at the end of ProcessInvalidation when it is safe to do so.
+	 * The value of this variable will be passed as the `bClearResourcesImmediately` parameter for HandleInvalidateAllWidgets
+	 */
+	TOptional<bool> bDeferInvalidateAllWidgetsAndClearResources;
+
 #if WITH_SLATE_DEBUGGING
 	ESlateInvalidationPaintType LastPaintType;
 	FPerformanceStat PerformanceStat;
