Adding overloads to `FContainerInitializer::AddContainedAttribute` to avoid compilation warnings or errors when doing a narrowing conversion from `SIZE_T` to `FSlateAttributeDescriptor::OffsetType`.

#Synced-CL 49309628

diff --git a/Engine/Source/Runtime/SlateCore/Private/Types/SlateAttributeDescriptor.cpp b/Engine/Source/Runtime/SlateCore/Private/Types/SlateAttributeDescriptor.cpp
--- a/Engine/Source/Runtime/SlateCore/Private/Types/SlateAttributeDescriptor.cpp
+++ b/Engine/Source/Runtime/SlateCore/Private/Types/SlateAttributeDescriptor.cpp
@@ -13,9 +13,9 @@ FSlateAttributeDescriptor::FAttribute::FAttribute(FName InName, OffsetType InOff
 	, Offset(InOffset)
 	, Prerequisite()
 	, ContainerName()
-	, SortOrder(DefaultSortOrder(InOffset))
-	, InvalidationReason(MoveTemp(InReason))
 	, OnValueChanged()
+	, InvalidationReason(MoveTemp(InReason))
+	, SortOrder(DefaultSortOrder(InOffset))
 	, AttributeType(SlateAttributePrivate::ESlateAttributeType::Member)
 	, bAffectVisibility(false)
 {
@@ -26,9 +26,9 @@ FSlateAttributeDescriptor::FAttribute::FAttribute(FName InContainerName, FName I
 	, Offset(InOffset)
 	, Prerequisite()
 	, ContainerName(InContainerName)
-	, SortOrder(DefaultSortOrder(InOffset))
-	, InvalidationReason(MoveTemp(InReason))
 	, OnValueChanged()
+	, InvalidationReason(MoveTemp(InReason))
+	, SortOrder(DefaultSortOrder(InOffset))
 	, AttributeType(SlateAttributePrivate::ESlateAttributeType::Contained)
 	, bAffectVisibility(false)
 {
@@ -133,12 +133,23 @@ FSlateAttributeDescriptor::FContainerInitializer::FAttributeEntry FSlateAttribut
 	return Descriptor.AddContainedAttribute(ContainerName, AttributeName, Offset, Reason);
 }
 
+FSlateAttributeDescriptor::FContainerInitializer::FAttributeEntry FSlateAttributeDescriptor::FContainerInitializer::AddContainedAttribute(FName AttributeName, SIZE_T Offset, const FInvalidateWidgetReasonAttribute& Reason)
+{
+	check(Offset <= std::numeric_limits<OffsetType>::max());
+	return Descriptor.AddContainedAttribute(ContainerName, AttributeName, (OffsetType)Offset, Reason);
+}
 
 FSlateAttributeDescriptor::FContainerInitializer::FAttributeEntry FSlateAttributeDescriptor::FContainerInitializer::AddContainedAttribute(FName AttributeName, OffsetType Offset, FInvalidateWidgetReasonAttribute&& Reason)
 {
 	return Descriptor.AddContainedAttribute(ContainerName, AttributeName, Offset, MoveTemp(Reason));
 }
 
+FSlateAttributeDescriptor::FContainerInitializer::FAttributeEntry FSlateAttributeDescriptor::FContainerInitializer::AddContainedAttribute(FName AttributeName, SIZE_T Offset, FInvalidateWidgetReasonAttribute&& Reason)
+{
+	check(Offset <= std::numeric_limits<OffsetType>::max());
+	return Descriptor.AddContainedAttribute(ContainerName, AttributeName, (OffsetType)Offset, MoveTemp(Reason));
+}
+
 
 void FSlateAttributeDescriptor::FContainerInitializer::OverrideInvalidationReason(FName AttributeName, const FInvalidateWidgetReasonAttribute& Reason)
 {
diff --git a/Engine/Source/Runtime/SlateCore/Public/Types/SlateAttributeDescriptor.h b/Engine/Source/Runtime/SlateCore/Public/Types/SlateAttributeDescriptor.h
--- a/Engine/Source/Runtime/SlateCore/Public/Types/SlateAttributeDescriptor.h
+++ b/Engine/Source/Runtime/SlateCore/Public/Types/SlateAttributeDescriptor.h
@@ -109,7 +109,7 @@ public:
 	using OffsetType = uint32;
 
 	/** The default sort order that define in which order attributes will be updated. */
-	static constexpr OffsetType DefaultSortOrder(OffsetType Offset) { return Offset * 100; }
+	static constexpr uint32 DefaultSortOrder(OffsetType Offset) { return Offset * 100; }
 
 
 	/** */
@@ -192,10 +192,10 @@ public:
 		OffsetType Offset;
 		FName Prerequisite;
 		FName ContainerName; // if the container IsNone, then the container is the SWidget itself.
+		FAttributeValueChangedDelegate OnValueChanged;
+		FInvalidateWidgetReasonAttribute InvalidationReason;
 		uint32 SortOrder;
 		uint8 ContainerIndex;
-		FInvalidateWidgetReasonAttribute InvalidationReason;
-		FAttributeValueChangedDelegate OnValueChanged;
 		SlateAttributePrivate::ESlateAttributeType AttributeType;
 		bool bAffectVisibility;
 	};
@@ -240,7 +240,9 @@ public:
 		};
 
 		SLATECORE_API FAttributeEntry AddContainedAttribute(FName AttributeName, OffsetType Offset, const FInvalidateWidgetReasonAttribute& ReasonGetter);
+		SLATECORE_API FAttributeEntry AddContainedAttribute(FName AttributeName, SIZE_T Offset, const FInvalidateWidgetReasonAttribute& ReasonGetter);
 		SLATECORE_API FAttributeEntry AddContainedAttribute(FName AttributeName, OffsetType Offset, FInvalidateWidgetReasonAttribute&& ReasonGetter);
+		SLATECORE_API FAttributeEntry AddContainedAttribute(FName AttributeName, SIZE_T Offset, FInvalidateWidgetReasonAttribute&& ReasonGetter);
 
 	public:
 		/** Change the InvalidationReason of an attribute defined in a base class. */
