Slate Debugger Paint - Fix for `OnlyInvalidationRoot` option not working as expected when turned off

#Synced-CL 49551883

diff --git a/Engine/Source/Runtime/SlateCore/Private/Debugging/ConsoleSlateDebuggerPassBase.cpp b/Engine/Source/Runtime/SlateCore/Private/Debugging/ConsoleSlateDebuggerPassBase.cpp
--- a/Engine/Source/Runtime/SlateCore/Private/Debugging/ConsoleSlateDebuggerPassBase.cpp
+++ b/Engine/Source/Runtime/SlateCore/Private/Debugging/ConsoleSlateDebuggerPassBase.cpp
@@ -228,7 +228,8 @@ void FConsoleSlateDebuggerPassBase::HandlePaintDebugInfo(const FPaintArgs& InArg
 		if (TSharedPtr<const SWidget> Widget = WidgetInfo.Widget.Pin())
 		{
 			// Skip the widget if collapsed or not in a window anymore (ex: in an unactive tab)
-			if (!Widget->GetProxyHandle().GetWidgetVisibility(Widget.Get()).IsVisible() || !FSlateApplicationBase::Get().FindWidgetWindow(Widget.ToSharedRef()))
+			if ( (Widget->GetProxyHandle().IsValid(Widget.Get()) && !Widget->GetProxyHandle().GetWidgetVisibility(Widget.Get()).IsVisible())
+				|| !FSlateApplicationBase::Get().FindWidgetWindow(Widget.ToSharedRef()))
 			{
 				continue;
 			}
