Slate Dynamic Invalidation - Handling invalid widgets created in Slow Path

This CL ensures that if widgets were created on Tick, we will check first that dynamic invalidation is still supported before calling `ProcessPreUpdate`. This allows Dynamic Invalidation to be disabled if widgets with invalid parents were just created in both Slow Path and Fast Path, without triggering a `check` or an `ensure`.

#Synced-CL 49551883

diff --git a/Engine/Source/Runtime/SlateCore/Private/FastUpdate/SlateInvalidationRoot.cpp b/Engine/Source/Runtime/SlateCore/Private/FastUpdate/SlateInvalidationRoot.cpp
--- a/Engine/Source/Runtime/SlateCore/Private/FastUpdate/SlateInvalidationRoot.cpp
+++ b/Engine/Source/Runtime/SlateCore/Private/FastUpdate/SlateInvalidationRoot.cpp
@@ -425,6 +425,14 @@ FSlateInvalidationResult FSlateInvalidationRoot::PaintInvalidationRoot(const FSl
 		Result.bRepaintedWidgets = PaintFastPath(Context);
 	}
 
+#if UE_SLATE_WITH_DYNAMIC_INVALIDATION
+	if (!bNeedsSlowPath)
+	{
+		// In case widgets were added during Slow Path or PrePass, we need to recheck if we can call ProcessPreUpdate
+		bNeedsSlowPath = !GetRootWidget()->SupportsInvalidationRecursive(SWidget::EInvalidationStrategy::UseCachedValue);
+	}
+#endif
+
 	// Once everything is painted and ticked, new widget (SListView, SRichBox, SScrollBox) might have created new widget in both SlowPath and FastPath
 	// It might create a 1 frame hitch because the widget is in the cached and should not.
 	if (!bNeedsSlowPath && WidgetsNeedingPreUpdate->Num() > 0)
