fix pinned SListView items when used outside of the editor

#rnx

#Synced-CL 50053775

diff --git a/Engine/Source/Editor/EditorStyle/Private/StarshipStyle.cpp b/Engine/Source/Editor/EditorStyle/Private/StarshipStyle.cpp
--- a/Engine/Source/Editor/EditorStyle/Private/StarshipStyle.cpp
+++ b/Engine/Source/Editor/EditorStyle/Private/StarshipStyle.cpp
@@ -870,8 +870,6 @@ void FStarshipEditorStyle::FStyle::SetupGeneralStyles()
 			.SetActiveHoveredBrush(FSlateNoResource())
 			.SetInactiveHoveredBrush(FSlateNoResource())
 			);
-	
-		Set("ListView.PinnedItemShadow", new IMAGE_BRUSH("Starship/ListView/PinnedItemShadow", FVector2f(16.f, 8.f)));
 	}
 	
 	// Spinboxes
diff --git a/Engine/Source/Runtime/SlateCore/Private/Styling/StarshipCoreStyle.cpp b/Engine/Source/Runtime/SlateCore/Private/Styling/StarshipCoreStyle.cpp
--- a/Engine/Source/Runtime/SlateCore/Private/Styling/StarshipCoreStyle.cpp
+++ b/Engine/Source/Runtime/SlateCore/Private/Styling/StarshipCoreStyle.cpp
@@ -2149,6 +2149,8 @@ void FStarshipCoreStyle::SetupTableViewStyles(TSharedRef<FStyle>& Style)
 
 	Style->Set("TableView.Row", DefaultTableRowStyle);
 
+	Style->Set("ListView.PinnedItemShadow", new IMAGE_BRUSH("Starship/CoreWidgets/TableView/PinnedItemShadow", FVector2f(16.f, 8.f)));
+
 
 	const FTableRowStyle SimpleTableRowStyle = FTableRowStyle()
 		.SetEvenRowBackgroundBrush(FSlateNoResource())
