[SlateCore]
* Fixed non unity compile error

diff --git a/Engine/Source/Runtime/SlateCore/Private/Fonts/FontGeometryPreprocessing.cpp b/Engine/Source/Runtime/SlateCore/Private/Fonts/FontGeometryPreprocessing.cpp
--- a/Engine/Source/Runtime/SlateCore/Private/Fonts/FontGeometryPreprocessing.cpp
+++ b/Engine/Source/Runtime/SlateCore/Private/Fonts/FontGeometryPreprocessing.cpp
@@ -4,6 +4,10 @@
 
 #if WITH_EDITOR && WITH_FREETYPE
 
+#if PLATFORM_MAC
+#include <TargetConditionals.h>
+#endif
+
 #define MSDFGEN_USE_SKIA
 #include "ThirdParty/skia/skia-simplify.h"
 
