Fixed non unity mac errors

diff --git a/Engine/Source/Runtime/SlateCore/Private/Fonts/PreprocessedFontGeometry.cpp b/Engine/Source/Runtime/SlateCore/Private/Fonts/PreprocessedFontGeometry.cpp
--- a/Engine/Source/Runtime/SlateCore/Private/Fonts/PreprocessedFontGeometry.cpp
+++ b/Engine/Source/Runtime/SlateCore/Private/Fonts/PreprocessedFontGeometry.cpp
@@ -2,6 +2,10 @@
 
 #include "Fonts/PreprocessedFontGeometry.h"
 
+#if PLATFORM_MAC
+#include "TargetConditionals.h"
+#endif
+
 // Only enable Skia-dependent features of MSDFgen for Editor build
 #if WITH_EDITOR
 #define MSDFGEN_USE_SKIA
diff --git a/Engine/Source/Runtime/SlateCore/Private/Fonts/SlateSdfGenerator.cpp b/Engine/Source/Runtime/SlateCore/Private/Fonts/SlateSdfGenerator.cpp
--- a/Engine/Source/Runtime/SlateCore/Private/Fonts/SlateSdfGenerator.cpp
+++ b/Engine/Source/Runtime/SlateCore/Private/Fonts/SlateSdfGenerator.cpp
@@ -15,6 +15,10 @@
 // This option will log the time it took to generate SDF for each glyph
 //#define SLATE_LOG_SDF_GENERATOR_TIME 1
 
+#if PLATFORM_MAC
+#include "TargetConditionals.h"
+#endif
+
 // Only enable Skia-dependent features of MSDFgen for Editor build
 #if WITH_EDITOR
 #define MSDFGEN_USE_SKIA
