Fonts:
Fixed a bug in the composite font resolution.
- The problem was because font ranges could overlap, e.g (1, 1000) Font A, (200, 300) Font B, (350, 400) Font C
- Given a unicode character say 325, the binary search will find the first range with the higher lower bound (350, 400) in this case.
It would only search the range it found and the previous range (200, 300) in this example. Because neither range coudl support 325, the search will return that no font range and font support the unicode character.
- The fix is to normalize the cached font ranges and remove all overlaps and to ensure that all font ranges are disjoint.
- This feature is off by default and gated by [REDACTED_PERSON]. Set that to true in your project's DefaultEngine.ini to enable this feature.

diff --git a/Engine/Source/Runtime/SlateCore/Private/Fonts/FontCacheCompositeFont.cpp b/Engine/Source/Runtime/SlateCore/Private/Fonts/FontCacheCompositeFont.cpp
--- a/Engine/Source/Runtime/SlateCore/Private/Fonts/FontCacheCompositeFont.cpp
+++ b/Engine/Source/Runtime/SlateCore/Private/Fonts/FontCacheCompositeFont.cpp
@@ -1,36 +1,43 @@
 // Copyright Epic Games, Inc. All Rights Reserved.
 
 #include "Fonts/FontCacheCompositeFont.h"
-#include "Fonts/FontCacheFreeType.h"
 #include "SlateGlobals.h"
-#include "Fonts/FontBulkData.h"
-#include "Misc/FileHelper.h"
 #include "Algo/BinarySearch.h"
-#include "Internationalization/Culture.h"
-#include "Internationalization/Internationalization.h"
 #include "Application/SlateApplicationBase.h"
 #include "Async/Async.h"
+#include "Fonts/FontBulkData.h"
+#include "Fonts/FontCacheFreeType.h"
 #include "HAL/PlatformProcess.h"
+#include "Internationalization/Culture.h"
+#include "Internationalization/Internationalization.h"
+#include "Misc/FileHelper.h"
 #include "Serialization/MemoryReader.h"
 #include "Trace/SlateMemoryTags.h"
 
+
 static bool GAsyncFontLazyLoad = false;
 FAutoConsoleVariableRef CVarAsyncLazyLoad(
 	TEXT("Slate.Font.AsyncLazyLoad"),
 	GAsyncFontLazyLoad,
 	TEXT("Causes unloaded font faces that are lazily loaded, to be loaded asynchronously, until then the font won't measure correctly.  Once complete the UI will invalidate. Note that async font loading is disallowed in Editor/PIE, standalone and cooked builds are fine. "));
 
+static bool GNormalizeCachedFontRanges = false;
+FAutoConsoleVariableRef CVarNormalizeCachedFontRanges(
+	TEXT("Slate.Font.NormalizeCachedFontRanges"),
+	GNormalizeCachedFontRanges,
+	TEXT("Normalize the font ranges for subtypefaces in composite fonts when building out the cached font ranges. Given a list of font ranges F, normalization infolves removing all intersecting font ranges in F but ensuring that the union of all the font ranges are still the same."));
+
 namespace UE::Private::FontCacheCompositeFont
 {
-bool IsAsyncFontLoadingAllowed()
-{
-	// We only allow async font loading in standalone or cooked builds, not in Editor/PIE. We want everything to be synchronously loaded for the actual Editor
+	bool IsAsyncFontLoadingAllowed()
+	{
+		// We only allow async font loading in standalone or cooked builds, not in Editor/PIE. We want everything to be synchronously loaded for the actual Editor
 #if WITH_EDITOR
-	return GAsyncFontLazyLoad && !GIsEditor;
+		return GAsyncFontLazyLoad && !GIsEditor;
 #else
-	return GAsyncFontLazyLoad;
+		return GAsyncFontLazyLoad;
 #endif
-}
+	}
 }
 
 DECLARE_CYCLE_STAT(TEXT("Load Font"), STAT_SlateLoadFont, STATGROUP_Slate);
@@ -253,12 +260,27 @@ void FCachedCompositeFontData::RefreshFontRanges()
 			{
 				if (!Range.IsEmpty())
 				{
-					CachedFontRangesPtr->Emplace(Range, CachedTypeface);
+					// We cap the upper bound at INT32_MAX because of the following:
+					// 1. Unicode only has about 1.1 million values. A whole chunk of the range leading up to INT32_MAX will all be invalid so there really is no difference between INT32_MAX and INT32_MAX - 1
+					// 2. Capping the upper bound to INT32_MAX -1 greatly simplifies the line sweep algorithm we use for normalizing font ranges. See FCachedCompositeFontData::NormalizeFontRanges
+					FInt32Range ClampedRange = Range;
+					ClampedRange.SetUpperBoundValue(FMath::Min(Range.GetUpperBoundValue(), TNumericLimits<int32>::Max() - 1));
+					CachedFontRangesPtr->Emplace(MoveTemp(ClampedRange), CachedTypeface);
 				}
 			}
 		}
 	}
 
+	if (GNormalizeCachedFontRanges)
+	{
+		// normalization would already leave the array sorted and the contiguous font ranges using the same font face merged 
+		TArray<FCachedFontRange> NormalizedCachedPriorityFontRanges = FCachedCompositeFontData::NormalizeFontRanges(CachedPriorityFontRanges);
+		CachedPriorityFontRanges = MoveTemp(NormalizedCachedPriorityFontRanges);
+		TArray<FCachedFontRange> NormalizedCachedFontRanges = FCachedCompositeFontData::NormalizeFontRanges(CachedFontRanges);
+		CachedFontRanges = MoveTemp(NormalizedCachedFontRanges);
+		return;
+	}
+
 	// Sort the font ranges into ascending order so we can binary search them later
 	CachedPriorityFontRanges.Sort(&FCachedFontRange::SortPredicate);
 	CachedFontRanges.Sort(&FCachedFontRange::SortPredicate);
@@ -288,6 +310,176 @@ void FCachedCompositeFontData::RefreshFontRanges()
 	MergeContiguousRanges(CachedFontRanges);
 }
 
+TArray<FCachedCompositeFontData::FCachedFontRange> FCachedCompositeFontData::NormalizeFontRanges(const TArray<FCachedCompositeFontData::FCachedFontRange>& InFontRanges)
+{
+	TArray< FCachedFontRange> OutFontRanges;
+	if (InFontRanges.IsEmpty())
+	{
+		return OutFontRanges;
+	}
+	// We will at minimum have as many unicode ranges as our input cached font ranges 
+	OutFontRanges.Reserve(InFontRanges.Num());
+
+	// We will be using a line sweep algorithm to normalize the font ranges and get rid of any overlapped ranges 
+	// We split each unicode font range into 2 events: a start event and end event
+	// Say we have the font range (1, 1000), then we will have a start event at position 1 and an end event at position 1001 (the offset will be explained later)
+	// We will then sort all of the events based on its position so we basically have a bunch of points on a number line
+	// We sweep through all of the events, process all the events at the same position, create a new font range based on the event positions and decide which font face will be associated with the new font range
+	// The end result is an array of font ranges that do not overlap, but their union is the same as the union of the input font ranges
+
+	struct FSweepEvent
+	{
+		FSweepEvent() = default;
+		FSweepEvent(int32 InPosition, int32 InFontRangeIndex, bool bInStart)
+			: Position(InPosition)
+			, FontRangeIndex(InFontRangeIndex)
+			, bStart(bInStart)
+		{
+
+		}
+		static bool SortPredicate(const FSweepEvent& A, const FSweepEvent& B)
+		{
+			if (A.Position != B.Position)
+			{
+				return A.Position < B.Position;
+			}
+			else if (A.bStart != B.bStart)
+			{
+				// We will prioritize end events over start events
+				return static_cast<int32>(A.bStart) < static_cast<int32>(B.bStart);
+			}
+			// We use the order we processed the font ranges be the tie breaker if everything else is equal 
+			return A.FontRangeIndex < B.FontRangeIndex;
+		}
+
+		bool IsValid() const
+		{
+			return FontRangeIndex != INDEX_NONE && Position >= 0;
+		}
+		/** The position of this event on the number line*/
+		int32 Position = INDEX_NONE;
+		/** Index of the font range this event was decomposed from*/
+		int32 FontRangeIndex = INDEX_NONE;
+		/** If this event at the position of the number line is the start of a range. True if it is the start and false if it is the end of the range*/
+		bool bStart = false;
+	};
+	TArray<FSweepEvent> SweepEvents;
+	// We decompose each unicode interval into a start and end event, thus we double the size of the cached font ranges
+	SweepEvents.Reserve(InFontRanges.Num() * 2);
+	for (int32 Index = 0; Index < InFontRanges.Num(); ++Index)
+	{
+		const FCachedFontRange& FontRange = InFontRanges[Index];
+		// All empty font ranges should have already been filtered out prior to calling this function 
+		check(!FontRange.Range.IsEmpty());
+		SweepEvents.Emplace(FontRange.Range.GetLowerBoundValue(), Index, true);
+		// Unicode does not go that high and having this requirement simplifies the algorithm
+		// Anything with an upper bound of TNumericLimits<int32>::Max() should have been filtered out prior to this function 
+		check(FontRange.Range.GetUpperBoundValue() != TNumericLimits<int32>::Max());
+		// We always off set the upper bound by +1 because there would be an ambiguity otherwise whether the upper bound is inclusive or exclusive
+		// Say we have the font ranges {(1, 1000), FontA}, {(200, 299), FontB}
+		// We decompose into events without the offset: 1S, 200S, 299E, 1000E(S and E for start and end)
+		// At 299E, we have the ambiguity of whether 299 should be covered by font B and we'll have to order processing of events and querying carefully
+		// If we add the offset, we don't need to worry about any of that
+		SweepEvents.Emplace(FontRange.Range.GetUpperBoundValue() + 1, Index, false);
+	}
+	// We should have exactly double the number of entries as the input font range, otherwise somehow some input font ranges were not decomposed and added properly 
+	check(SweepEvents.Num() == InFontRanges.Num() * 2);
+	SweepEvents.Sort(&FSweepEvent::SortPredicate);
+
+	// As we sweep through the events and come across overlaps, we need to be able to figure out which font face applies to a particular range
+	// This priority queue keeps track of the indices of all active font ranges at a particular event and sorts them by the sizes of the font ranges
+	// We want to give precedence to smaller font ranges and use the font face associated with that range instead of a larger range if there is overlap. Example in the comparison lambda below 
+	TArray<int32> ActiveFontRangeIndexHeap;
+	ActiveFontRangeIndexHeap.Reserve(InFontRanges.Num());
+
+	auto FontRangeComparator = [&InFontRanges](const int32 A, const int32 B)
+	{
+		const int32 SizeA = InFontRanges[A].Range.Size<int32>();
+		const int32 SizeB = InFontRanges[B].Range.Size<int32>();
+		// We want to prioritize font ranges that are smaller and more specific
+		// E.g We have a range (1, 1000) associated with Font A and (50, 100) associated with Font B. B has the smaller range
+		// so we assume that it takes precedence over FontA 
+		if (SizeA != SizeB)
+		{
+			return SizeA < SizeB;
+		}
+		// If the size of the unicode ranges are identical, then we return the range that came first in the user declaration 
+		return A < B;
+	};
+	// Allows us to keep track of which font ranges are active at any point during the sweep
+	TBitArray<> FontRangeActiveStatus;
+	FontRangeActiveStatus.Init(false, InFontRanges.Num());
+
+	for (int32 Index = 0; Index < SweepEvents.Num(); )
+	{
+		const int32 CurrentPosition = SweepEvents[Index].Position;
+		// We apply all of the events that occur at the same position first 
+		while (Index < SweepEvents.Num() && SweepEvents[Index].Position == CurrentPosition)
+		{
+			const FSweepEvent& CurrentEvent = SweepEvents[Index];
+			check(CurrentEvent.IsValid());
+			if (CurrentEvent.bStart)
+			{
+				FontRangeActiveStatus[CurrentEvent.FontRangeIndex] = true;
+				ActiveFontRangeIndexHeap.HeapPush(CurrentEvent.FontRangeIndex, FontRangeComparator);
+			}
+			else
+			{
+				FontRangeActiveStatus[CurrentEvent.FontRangeIndex] = false;
+				// We do not remove the element from the heap here as that can be expensive
+				// Instead we keep popping elements off the heap afterwards and ignore the inactive elements
+			}
+			++Index;
+		}
+
+		// Now where we are, we get the top most priority range index that is active
+		// This represents the index of the font range (and consequently font face) that will be used for this new font range we are constructing 
+		int32 ActiveFontRangeIndex = INDEX_NONE;
+		while (!ActiveFontRangeIndexHeap.IsEmpty())
+		{
+			int32 TopIndex = ActiveFontRangeIndexHeap.HeapTop();
+			if (FontRangeActiveStatus[TopIndex])
+			{
+				ActiveFontRangeIndex = TopIndex;
+				break;
+			}
+			// The top index is associated with a range that is already inactive and thus should be removed 
+			ActiveFontRangeIndexHeap.HeapPopDiscard(FontRangeComparator, EAllowShrinking::No);
+		}
+
+		// There are no active font ranges/font faces at this event, we do not create a new font range 
+		if (ActiveFontRangeIndex == INDEX_NONE)
+		{
+			continue;
+		}
+		// The only time we would be able to get here with Index == SweepEvents.Num() is if there is a unicode range that has no end. I.e it is somehow unbounded 
+		check(Index != SweepEvents.Num());
+
+		// Now we create a new range
+		const int32 LowerBound = CurrentPosition;
+		// We need to remove the +1 bias from when we created the sweep events 
+		const int32 UpperBound = SweepEvents[Index].Position - 1;
+		// If this is violated, the events were somehow not sorted properly 
+		check(UpperBound >= LowerBound);
+		const TSharedPtr<FCachedTypefaceData>& ActiveFontFace = InFontRanges[ActiveFontRangeIndex].CachedTypeface;
+
+		FInt32Range NewRange(FInt32Range::BoundsType::Inclusive(LowerBound), FInt32Range::BoundsType::Inclusive(UpperBound));
+		// if this new range is using the same font face as the previous range, we can merge the ranges
+		if (!OutFontRanges.IsEmpty())
+		{
+			FCachedFontRange& PreviousFontRange = OutFontRanges.Last();
+			// if the previous font range is using the same font face and is adjoined to the new range, the 2 ranges are contiguous and we can merge them 
+			if (PreviousFontRange.CachedTypeface == ActiveFontFace && PreviousFontRange.Range.Adjoins(NewRange))
+			{
+				PreviousFontRange.Range.SetUpperBoundValue(UpperBound);
+				continue;
+			}
+		}
+
+		OutFontRanges.Emplace(MoveTemp(NewRange), ActiveFontFace);
+	}
+	return OutFontRanges;
+}
 
 FCompositeFontCache::FCompositeFontCache(const FFreeTypeLibrary* InFTLibrary)
 	: FTLibrary(InFTLibrary)
diff --git a/Engine/Source/Runtime/SlateCore/Private/Fonts/FontCacheCompositeFont.h b/Engine/Source/Runtime/SlateCore/Private/Fonts/FontCacheCompositeFont.h
--- a/Engine/Source/Runtime/SlateCore/Private/Fonts/FontCacheCompositeFont.h
+++ b/Engine/Source/Runtime/SlateCore/Private/Fonts/FontCacheCompositeFont.h
@@ -177,6 +177,7 @@ private:
 		TSharedPtr<FCachedTypefaceData> CachedTypeface;
 	};
 
+	static TArray<FCachedCompositeFontData::FCachedFontRange> NormalizeFontRanges(const TArray<FCachedCompositeFontData::FCachedFontRange>& InFontRanges);
 	/** Constants for indexing into CachedTypefaces */
 	static const int32 CachedTypeface_DefaultIndex = 0;
 	static const int32 CachedTypeface_FallbackIndex = 1;
