Support for `TSlateAttribute` in Widgets using Multiple Inheritance

On most platforms, a `TSlateAttribute` defined in a widget that uses multiple inheritance cannot be reliably located at runtime unless the widget's primary base class is the one deriving from `SWidget`. For example:
```
// TSlateAttributes defined in this class WOULD NOT be found at runtime on Windows
class STableRow : public ITableRow, public SBorder

// TSlateAttributes defined in this class WOULD be found at runtime on Windows
class STableRow : public SBorder, public ITableRow
```
This happens because multiple inheritance introduces multiple vtables, which breaks the offset calculation used to locate attributes as it was computed relative to a base class. Some platforms also do not guarantee deterministic vtable ordering.

This change updates the attribute offset computation so that all offsets are derived relative to the widget's `SWidget` base class. As long as the widget is first cast to `SWidget`, its attributes can now be consistently located regardless of inheritance order or platform???specific vtable layout.

#Synced-CL 50145007

diff --git a/Engine/Source/Runtime/Core/Public/Misc/Attribute.h b/Engine/Source/Runtime/Core/Public/Misc/Attribute.h
--- a/Engine/Source/Runtime/Core/Public/Misc/Attribute.h
+++ b/Engine/Source/Runtime/Core/Public/Misc/Attribute.h
@@ -448,7 +448,14 @@ public:
 	{
 		checkf(IsSet(), TEXT("It is an error to call Steal() on an unset TAttribute. Check IsSet() before calling Steal()."));
 		bIsSet = false;
-		return IsBound() ? TVariant<ObjectType, FGetter>(TInPlaceType<FGetter>(), MoveTemp(Getter)) : TVariant<ObjectType, FGetter>(TInPlaceType<ObjectType>(), MoveTemp(Value));
+		if (IsBound())
+		{
+			return TVariant<ObjectType, FGetter>(TInPlaceType<FGetter>(), MoveTemp(Getter));
+		}
+		else
+		{
+			return TVariant<ObjectType, FGetter>(TInPlaceType<ObjectType>(), MoveTemp(Value));
+		}
 	}
 
 	/**
diff --git a/Engine/Source/Runtime/SlateCore/Private/Test/SlateAttributeTest.cpp b/Engine/Source/Runtime/SlateCore/Private/Test/SlateAttributeTest.cpp
--- a/Engine/Source/Runtime/SlateCore/Private/Test/SlateAttributeTest.cpp
+++ b/Engine/Source/Runtime/SlateCore/Private/Test/SlateAttributeTest.cpp
@@ -409,6 +409,290 @@ void SContainerAttribute::PrivateRegisterAttributes(FSlateAttributeInitializer&
 	SLATE_ADD_MEMBER_ATTRIBUTE_DEFINITION(AttributeInitializer, AttributeC, EInvalidateWidgetReason::Volatility);
 }
 
+
+// ------------------------------------------------
+// -------------- Multi Inheritance ---------------
+// ------------------------------------------------
+
+class IMultiInheritanceWidget_Int_Base
+{
+public:
+	virtual void SomeBaseFunction_Int() {};
+};
+class IMultiInheritanceWidget_UInt_Base
+{
+public:
+	virtual void SomeBaseFunction_Int() {};
+};
+
+class SMultiInheritanceWidget_Base
+	: public IMultiInheritanceWidget_Int_Base, public SLeafWidget, public IMultiInheritanceWidget_UInt_Base
+{
+	SLATE_DECLARE_WIDGET(SMultiInheritanceWidget_Base, SLeafWidget)
+public:
+	SLATE_BEGIN_ARGS(SMultiInheritanceWidget_Base) {}
+	SLATE_END_ARGS()
+
+	SMultiInheritanceWidget_Base()
+		: BaseIntAttribute(*this, -9999)
+		, BaseUIntAttribute(*this, 9999)
+	{
+	}
+
+	void Construct(const FArguments& InArgs){}
+	virtual FVector2D ComputeDesiredSize(float) const override { return FVector2D{ 100, 100 }; }
+	virtual int32 OnPaint(const FPaintArgs& Args, const FGeometry& AllottedGeometry, const FSlateRect& MyCullingRect, FSlateWindowElementList& OutDrawElements, int32 LayerId, const FWidgetStyle& InWidgetStyle, bool bParentEnabled) const override
+	{
+		return LayerId;
+	}
+
+	TSlateAttribute<int32> BaseIntAttribute;
+	TSlateAttribute<uint32> BaseUIntAttribute;
+	int32 CallbackIntValue = 0;
+	uint32 CallbackUIntValue = 0;
+};
+
+SLATE_IMPLEMENT_WIDGET(SMultiInheritanceWidget_Base)
+void SMultiInheritanceWidget_Base::PrivateRegisterAttributes(FSlateAttributeInitializer& AttributeInitializer)
+{
+	SLATE_ADD_MEMBER_ATTRIBUTE_DEFINITION(AttributeInitializer, BaseIntAttribute, EInvalidateWidgetReason::Paint)
+		.OnValueChanged(FSlateAttributeDescriptor::FAttributeValueChangedDelegate::CreateLambda([](SWidget& Widget)
+		{
+			SMultiInheritanceWidget_Base& Instance = static_cast<SMultiInheritanceWidget_Base&>(Widget);
+			Instance.CallbackIntValue += Instance.BaseIntAttribute.Get();
+		}));
+	SLATE_ADD_MEMBER_ATTRIBUTE_DEFINITION(AttributeInitializer, BaseUIntAttribute, EInvalidateWidgetReason::Layout)
+		.OnValueChanged(FSlateAttributeDescriptor::FAttributeValueChangedDelegate::CreateLambda([](SWidget& Widget)
+		{
+			SMultiInheritanceWidget_Base& Instance = static_cast<SMultiInheritanceWidget_Base&>(Widget);
+			Instance.CallbackUIntValue += Instance.BaseUIntAttribute.Get();
+		}));
+}
+
+
+class SMultiInheritanceWidget_Parent : public SMultiInheritanceWidget_Base
+{
+	SLATE_DECLARE_WIDGET(SMultiInheritanceWidget_Parent, SMultiInheritanceWidget_Base)
+public:
+	SLATE_BEGIN_ARGS(SMultiInheritanceWidget_Parent) {}
+	SLATE_END_ARGS()
+
+
+	SMultiInheritanceWidget_Parent()
+		: ParentIntAttribute(*this, -8888)
+		, ParentUIntAttribute(*this, 8888)
+	{
+	}
+
+	void Construct(const FArguments& InArgs){}
+	virtual FVector2D ComputeDesiredSize(float) const override { return FVector2D{ 100, 100 }; }
+	virtual int32 OnPaint(const FPaintArgs& Args, const FGeometry& AllottedGeometry, const FSlateRect& MyCullingRect, FSlateWindowElementList& OutDrawElements, int32 LayerId, const FWidgetStyle& InWidgetStyle, bool bParentEnabled) const override
+	{
+		return LayerId;
+	}
+
+
+	TSlateAttribute<int32> ParentIntAttribute;
+	TSlateAttribute<uint32> ParentUIntAttribute;
+};
+
+SLATE_IMPLEMENT_WIDGET(SMultiInheritanceWidget_Parent)
+void SMultiInheritanceWidget_Parent::PrivateRegisterAttributes(FSlateAttributeInitializer& AttributeInitializer)
+{
+	SLATE_ADD_MEMBER_ATTRIBUTE_DEFINITION(AttributeInitializer, ParentIntAttribute, EInvalidateWidgetReason::Paint)
+		.OnValueChanged(FSlateAttributeDescriptor::FAttributeValueChangedDelegate::CreateLambda([](SWidget& Widget)
+		{
+			SMultiInheritanceWidget_Parent& Instance = static_cast<SMultiInheritanceWidget_Parent&>(Widget);
+			Instance.CallbackIntValue += Instance.ParentIntAttribute.Get();
+		}));
+	SLATE_ADD_MEMBER_ATTRIBUTE_DEFINITION(AttributeInitializer, ParentUIntAttribute, EInvalidateWidgetReason::Layout)
+		.OnValueChanged(FSlateAttributeDescriptor::FAttributeValueChangedDelegate::CreateLambda([](SWidget& Widget)
+		{
+			SMultiInheritanceWidget_Parent& Instance = static_cast<SMultiInheritanceWidget_Parent&>(Widget);
+			Instance.CallbackUIntValue += Instance.ParentUIntAttribute.Get();
+		}));
+}
+
+
+class IMultiInheritanceWidget_Int_Parent
+{
+public:
+	virtual void SomeParentFunction_Int() {};
+};
+class IMultiInheritanceWidget_UInt_Parent
+{
+public:
+	virtual void SomeParentFunction_UInt() {};
+};
+
+
+// Widget with multiple inheritance declared in SII order (SWidget, Interface, Interface)
+class SMultiInheritanceWidget_Child_SII
+	: public SMultiInheritanceWidget_Parent, public IMultiInheritanceWidget_Int_Parent, public IMultiInheritanceWidget_UInt_Parent
+{
+	SLATE_DECLARE_WIDGET(SMultiInheritanceWidget_Child_SII, SMultiInheritanceWidget_Parent)
+public:
+	SLATE_BEGIN_ARGS(SMultiInheritanceWidget_Child_SII) {}
+	SLATE_END_ARGS()
+
+	SMultiInheritanceWidget_Child_SII()
+		: ChildIntAttribute(*this, -511)
+		, ChildUIntAttribute(*this, 511)
+	{
+	}
+
+	void Construct(const FArguments& InArgs){}
+
+	TSlateAttribute<int32> ChildIntAttribute;
+	TSlateAttribute<uint32> ChildUIntAttribute;
+};
+SLATE_IMPLEMENT_WIDGET(SMultiInheritanceWidget_Child_SII)
+void SMultiInheritanceWidget_Child_SII::PrivateRegisterAttributes(FSlateAttributeInitializer& AttributeInitializer)
+{
+	SLATE_ADD_MEMBER_ATTRIBUTE_DEFINITION(AttributeInitializer, ChildIntAttribute, EInvalidateWidgetReason::Paint)
+		.OnValueChanged(FSlateAttributeDescriptor::FAttributeValueChangedDelegate::CreateLambda([](SWidget& Widget)
+		{
+			SMultiInheritanceWidget_Child_SII& Instance = static_cast<SMultiInheritanceWidget_Child_SII&>(Widget);
+			Instance.CallbackIntValue += Instance.ChildIntAttribute.Get();
+		}));
+	SLATE_ADD_MEMBER_ATTRIBUTE_DEFINITION(AttributeInitializer, ChildUIntAttribute, EInvalidateWidgetReason::Layout)
+		.OnValueChanged(FSlateAttributeDescriptor::FAttributeValueChangedDelegate::CreateLambda([](SWidget& Widget)
+		{
+			SMultiInheritanceWidget_Child_SII& Instance = static_cast<SMultiInheritanceWidget_Child_SII&>(Widget);
+			Instance.CallbackUIntValue += Instance.ChildUIntAttribute.Get();
+		}));
+}
+
+// Widget with multiple inheritance declared in ISI order (Interface, SWidget, Interface)
+class SMultiInheritanceWidget_Child_ISI
+	: public IMultiInheritanceWidget_Int_Parent, public SMultiInheritanceWidget_Parent, public IMultiInheritanceWidget_UInt_Parent
+{
+	SLATE_DECLARE_WIDGET(SMultiInheritanceWidget_Child_ISI, SMultiInheritanceWidget_Parent)
+public:
+	SLATE_BEGIN_ARGS(SMultiInheritanceWidget_Child_ISI) {}
+	SLATE_END_ARGS()
+
+	SMultiInheritanceWidget_Child_ISI()
+		: ChildIntAttribute(*this, -151)
+		, ChildUIntAttribute(*this, 151)
+	{
+	}
+
+	void Construct(const FArguments& InArgs){}
+
+	TSlateAttribute<int32> ChildIntAttribute;
+	TSlateAttribute<uint32> ChildUIntAttribute;
+};
+SLATE_IMPLEMENT_WIDGET(SMultiInheritanceWidget_Child_ISI)
+void SMultiInheritanceWidget_Child_ISI::PrivateRegisterAttributes(FSlateAttributeInitializer& AttributeInitializer)
+{
+	SLATE_ADD_MEMBER_ATTRIBUTE_DEFINITION(AttributeInitializer, ChildIntAttribute, EInvalidateWidgetReason::Paint)
+		.OnValueChanged(FSlateAttributeDescriptor::FAttributeValueChangedDelegate::CreateLambda([](SWidget& Widget)
+		{
+			SMultiInheritanceWidget_Child_ISI& Instance = static_cast<SMultiInheritanceWidget_Child_ISI&>(Widget);
+			Instance.CallbackIntValue += Instance.ChildIntAttribute.Get();
+		}));
+	SLATE_ADD_MEMBER_ATTRIBUTE_DEFINITION(AttributeInitializer, ChildUIntAttribute, EInvalidateWidgetReason::Layout)
+		.OnValueChanged(FSlateAttributeDescriptor::FAttributeValueChangedDelegate::CreateLambda([](SWidget& Widget)
+		{
+			SMultiInheritanceWidget_Child_ISI& Instance = static_cast<SMultiInheritanceWidget_Child_ISI&>(Widget);
+			Instance.CallbackUIntValue += Instance.ChildUIntAttribute.Get();
+		}));
+}
+
+// Widget with multiple inheritance declared in IIS order (Interface, Interface, SWidget)
+class SMultiInheritanceWidget_Child_IIS
+	: public IMultiInheritanceWidget_Int_Parent, public IMultiInheritanceWidget_UInt_Parent, public SMultiInheritanceWidget_Parent
+{
+	SLATE_DECLARE_WIDGET(SMultiInheritanceWidget_Child_IIS, SMultiInheritanceWidget_Parent)
+public:
+	SLATE_BEGIN_ARGS(SMultiInheritanceWidget_Child_IIS) {}
+	SLATE_END_ARGS()
+
+	SMultiInheritanceWidget_Child_IIS()
+		: ChildIntAttribute(*this, -151)
+		, ChildUIntAttribute(*this, 151)
+	{
+	}
+
+	void Construct(const FArguments& InArgs){}
+
+	TSlateAttribute<int32> ChildIntAttribute;
+	TSlateAttribute<uint32> ChildUIntAttribute;
+};
+SLATE_IMPLEMENT_WIDGET(SMultiInheritanceWidget_Child_IIS)
+void SMultiInheritanceWidget_Child_IIS::PrivateRegisterAttributes(FSlateAttributeInitializer& AttributeInitializer)
+{
+	SLATE_ADD_MEMBER_ATTRIBUTE_DEFINITION(AttributeInitializer, ChildIntAttribute, EInvalidateWidgetReason::Paint)
+		.OnValueChanged(FSlateAttributeDescriptor::FAttributeValueChangedDelegate::CreateLambda([](SWidget& Widget)
+		{
+			SMultiInheritanceWidget_Child_IIS& Instance = static_cast<SMultiInheritanceWidget_Child_IIS&>(Widget);
+			Instance.CallbackIntValue += Instance.ChildIntAttribute.Get();
+		}));
+	SLATE_ADD_MEMBER_ATTRIBUTE_DEFINITION(AttributeInitializer, ChildUIntAttribute, EInvalidateWidgetReason::Layout)
+		.OnValueChanged(FSlateAttributeDescriptor::FAttributeValueChangedDelegate::CreateLambda([](SWidget& Widget)
+		{
+			SMultiInheritanceWidget_Child_IIS& Instance = static_cast<SMultiInheritanceWidget_Child_IIS&>(Widget);
+			Instance.CallbackUIntValue += Instance.ChildUIntAttribute.Get();
+		}));
+}
+
+
+class IMultiInheritanceWidget_Int_GrandChild
+{
+public:
+	virtual void SomeGrandChildFunction_Int() {};
+};
+class IMultiInheritanceWidget_UInt_GrandChild
+{
+public:
+	virtual void SomeGrandChildFunction_UInt() {};
+};
+
+
+class SMultiInheritanceWidget_GrandChild_SII_SII
+	: public SMultiInheritanceWidget_Child_SII, public IMultiInheritanceWidget_Int_GrandChild, public IMultiInheritanceWidget_UInt_GrandChild
+{
+};
+
+class SMultiInheritanceWidget_GrandChild_SII_ISI
+	: public IMultiInheritanceWidget_Int_GrandChild, public SMultiInheritanceWidget_Child_SII, public IMultiInheritanceWidget_UInt_GrandChild
+{
+};
+class SMultiInheritanceWidget_GrandChild_SII_IIS
+	: public IMultiInheritanceWidget_Int_GrandChild, public IMultiInheritanceWidget_UInt_GrandChild, public SMultiInheritanceWidget_Child_SII
+{
+};
+
+
+class SMultiInheritanceWidget_GrandChild_ISI_SII
+	: public SMultiInheritanceWidget_Child_ISI, public IMultiInheritanceWidget_Int_GrandChild, public IMultiInheritanceWidget_UInt_GrandChild
+{
+};
+class SMultiInheritanceWidget_GrandChild_ISI_ISI
+	: public IMultiInheritanceWidget_Int_GrandChild, public SMultiInheritanceWidget_Child_ISI, public IMultiInheritanceWidget_UInt_GrandChild
+{
+};
+class SMultiInheritanceWidget_GrandChild_ISI_IIS
+	: public IMultiInheritanceWidget_Int_GrandChild, public IMultiInheritanceWidget_UInt_GrandChild, public SMultiInheritanceWidget_Child_ISI
+{
+};
+
+
+class SMultiInheritanceWidget_GrandChild_IIS_SII
+	: public SMultiInheritanceWidget_Child_IIS, public IMultiInheritanceWidget_Int_GrandChild, public IMultiInheritanceWidget_UInt_GrandChild
+{
+};
+class SMultiInheritanceWidget_GrandChild_IIS_ISI
+	: public IMultiInheritanceWidget_Int_GrandChild, public SMultiInheritanceWidget_Child_IIS, public IMultiInheritanceWidget_UInt_GrandChild
+{
+};
+class SMultiInheritanceWidget_GrandChild_IIS_IIS
+	: public IMultiInheritanceWidget_Int_GrandChild, public IMultiInheritanceWidget_UInt_GrandChild, public SMultiInheritanceWidget_Child_IIS
+{
+};
+
+
 } //namespace
 
 bool FSlateAttributeTest::RunTest(const FString& Parameters)
@@ -671,7 +955,14 @@ bool FSlateAttributeTest::RunTest(const FString& Parameters)
 					&& FLocalConstructionCounter::MoveConstructionCounter == Move
 					&& FLocalConstructionCounter::CopyOperatorCounter == CopyAssign
 					&& FLocalConstructionCounter::MoveOperatorCounter == MoveAssign;
-				AddErrorIfFalse(bSuccess, Message);
+				AddErrorIfFalse(bSuccess, FString(Message) +
+					FString::Printf(TEXT(" [Construct: %d [exp: %d], Copy: %d [exp: %d], Move: %d [exp: %d], CopyAssign: %d [exp: %d], MoveAssign: %d [exp: %d]]"),
+					FLocalConstructionCounter::DefaultConstructionCounter, Construct,
+					FLocalConstructionCounter::CopyConstructionCounter, Copy,
+					FLocalConstructionCounter::MoveConstructionCounter, Move,
+					FLocalConstructionCounter::CopyOperatorCounter, CopyAssign,
+					FLocalConstructionCounter::MoveOperatorCounter, MoveAssign
+				));
 			};
 
 			{
@@ -786,7 +1077,7 @@ bool FSlateAttributeTest::RunTest(const FString& Parameters)
 				Attribute1.Set({ 2 });
 				FLocalConstructionCounter::ResetCounter();
 				Attribute.Assign(MoveTemp(Attribute1));
-				AddErrorIfCounterDoNotMatches(0, 0, 2, 0, 1, TEXT("Assign Move failed."));
+				AddErrorIfCounterDoNotMatches(0, 0, 1, 0, 1, TEXT("Assign Move failed."));
 			}
 			// Test unbinded Attribute
 			{
@@ -877,7 +1168,7 @@ bool FSlateAttributeTest::RunTest(const FString& Parameters)
 				TAttribute<FLocalConstructionCounter> Attribute1 = Counter;
 				FLocalConstructionCounter::ResetCounter();
 				Attribute.Assign(MoveTemp(Attribute1), Counter);
-				AddErrorIfCounterDoNotMatches(0, 0, 2, 0, 1, TEXT("Assign Set Move failed."));
+				AddErrorIfCounterDoNotMatches(0, 0, 1, 0, 1, TEXT("Assign Set Move failed."));
 			}
 			{
 				SManagedAttribute::ManagedSlateAttributeType Attribute = SManagedAttribute::ManagedSlateAttributeType(Widget.ToSharedRef());
@@ -885,7 +1176,7 @@ bool FSlateAttributeTest::RunTest(const FString& Parameters)
 				TAttribute<FLocalConstructionCounter> Attribute1 = Counter;
 				FLocalConstructionCounter::ResetCounter();
 				Attribute.Assign(MoveTemp(Attribute1), MoveTemp(Counter));
-				AddErrorIfCounterDoNotMatches(0, 0, 2, 0, 1, TEXT("Assign Set Move failed."));
+				AddErrorIfCounterDoNotMatches(0, 0, 1, 0, 1, TEXT("Assign Set Move failed."));
 			}
 		}
 
@@ -1095,6 +1386,65 @@ bool FSlateAttributeTest::RunTest(const FString& Parameters)
 		}
 	}
 
+	// Test Multi Inheritance
+	// Worth noting that the vtable order is not deterministic by the different compilers (especially on Switch)
+	{
+		auto MultiInheritanceWidget_Child_Test = [this]<typename WidgetType>(TSharedRef<WidgetType> Widget)
+		{
+			const FSlateAttributeDescriptor& AttributeDescriptor = Widget->GetWidgetClass().GetAttributeDescriptor();
+			const int32 IndexBI  = AttributeDescriptor.IndexOfAttribute("BaseIntAttribute");
+			const int32 IndexBUI = AttributeDescriptor.IndexOfAttribute("BaseUIntAttribute");
+			const int32 IndexPI  = AttributeDescriptor.IndexOfAttribute("ParentIntAttribute");
+			const int32 IndexPUI = AttributeDescriptor.IndexOfAttribute("ParentUIntAttribute");
+			const int32 IndexI   = AttributeDescriptor.IndexOfAttribute("ChildIntAttribute");
+			const int32 IndexUI  = AttributeDescriptor.IndexOfAttribute("ChildUIntAttribute");
+
+			AddErrorIfFalse(IndexBI  != INDEX_NONE, FString::Printf(TEXT("[%s] Could not find the BaseIntAttribute"), *Widget->GetTypeAsString()));
+			AddErrorIfFalse(IndexBUI != INDEX_NONE, FString::Printf(TEXT("[%s] Could not find the BaseUIntAttribute"), *Widget->GetTypeAsString()));
+			AddErrorIfFalse(IndexPI  != INDEX_NONE, FString::Printf(TEXT("[%s] Could not find the ParentIntAttribute"), *Widget->GetTypeAsString()));
+			AddErrorIfFalse(IndexPUI != INDEX_NONE, FString::Printf(TEXT("[%s] Could not find the ParentUIntAttribute"), *Widget->GetTypeAsString()));
+			AddErrorIfFalse(IndexI   != INDEX_NONE, FString::Printf(TEXT("[%s] Could not find the IntAttribute"), *Widget->GetTypeAsString()));
+			AddErrorIfFalse(IndexUI  != INDEX_NONE, FString::Printf(TEXT("[%s] Could not find the UIntAttribute"), *Widget->GetTypeAsString()));
+
+			AddErrorIfFalse(&AttributeDescriptor.GetAttributeAtIndex(IndexBI)  == AttributeDescriptor.FindAttribute("BaseIntAttribute"), FString::Printf( TEXT("[%s] Index and Attribute should return the same value."), *Widget->GetTypeAsString()));
+			AddErrorIfFalse(&AttributeDescriptor.GetAttributeAtIndex(IndexBUI) == AttributeDescriptor.FindAttribute("BaseUIntAttribute"), FString::Printf( TEXT("[%s] Index and Attribute should return the same value."), *Widget->GetTypeAsString()));
+			AddErrorIfFalse(&AttributeDescriptor.GetAttributeAtIndex(IndexPI)  == AttributeDescriptor.FindAttribute("ParentIntAttribute"), FString::Printf( TEXT("[%s] Index and Attribute should return the same value."), *Widget->GetTypeAsString()));
+			AddErrorIfFalse(&AttributeDescriptor.GetAttributeAtIndex(IndexPUI) == AttributeDescriptor.FindAttribute("ParentUIntAttribute"), FString::Printf( TEXT("[%s] Index and Attribute should return the same value."), *Widget->GetTypeAsString()));
+			AddErrorIfFalse(&AttributeDescriptor.GetAttributeAtIndex(IndexI)   == AttributeDescriptor.FindAttribute("ChildIntAttribute"), FString::Printf( TEXT("[%s] Index and Attribute should return the same value."), *Widget->GetTypeAsString()));
+			AddErrorIfFalse(&AttributeDescriptor.GetAttributeAtIndex(IndexUI)  == AttributeDescriptor.FindAttribute("ChildUIntAttribute"), FString::Printf( TEXT("[%s] Index and Attribute should return the same value."), *Widget->GetTypeAsString()));
+
+			AddErrorIfFalse(Widget->CallbackIntValue == 0, FString::Printf( TEXT("[%s] The Int callback was triggered."), *Widget->GetTypeAsString()));
+			AddErrorIfFalse(Widget->CallbackUIntValue == 0, FString::Printf( TEXT("[%s] The UInt callback was triggered."), *Widget->GetTypeAsString()));
+
+			Widget->BaseIntAttribute.Assign(*Widget, -8);
+			AddErrorIfFalse(Widget->CallbackIntValue == -8, FString::Printf( TEXT("[%s] CallbackIntValue is not the expected value after setting 'BaseIntAttribute': Expected: %d  Value: %d."), *Widget->GetTypeAsString(), -8, Widget->CallbackIntValue));
+			Widget->ParentIntAttribute.Assign(*Widget, -80);
+			AddErrorIfFalse(Widget->CallbackIntValue == -88, FString::Printf( TEXT("[%s] CallbackIntValue is not the expected value after setting 'ParentIntAttribute': Expected: %d  Value: %d."), *Widget->GetTypeAsString(), -88, Widget->CallbackIntValue));
+			Widget->ChildIntAttribute.Assign(*Widget, -800);
+			AddErrorIfFalse(Widget->CallbackIntValue == -888, FString::Printf( TEXT("[%s] CallbackIntValue is not the expected value after setting 'ChildIntAttribute': Expected: %d  Value: %d."), *Widget->GetTypeAsString(), -888, Widget->CallbackIntValue));
+
+			Widget->ParentUIntAttribute.Assign(*Widget, 700);
+			AddErrorIfFalse(Widget->CallbackUIntValue == 700, FString::Printf( TEXT("[%s] CallbackUIntValue is not the expected value after setting 'ParentUIntAttribute': Expected: %d  Value: %d."), *Widget->GetTypeAsString(), 700, Widget->CallbackUIntValue));
+			Widget->ChildUIntAttribute.Assign(*Widget, 70);
+			AddErrorIfFalse(Widget->CallbackUIntValue == 770, FString::Printf( TEXT("[%s] CallbackUIntValue is not the expected value after setting 'ChildUIntAttribute': Expected: %d  Value: %d."), *Widget->GetTypeAsString(), 770, Widget->CallbackUIntValue));
+			Widget->BaseUIntAttribute.Assign(*Widget, 7);
+			AddErrorIfFalse(Widget->CallbackUIntValue == 777, FString::Printf( TEXT("[%s] CallbackUIntValue is not the expected value after setting 'BaseUIntAttribute': Expected: %d  Value: %d."), *Widget->GetTypeAsString(), 777, Widget->CallbackUIntValue));
+		};
+
+		MultiInheritanceWidget_Child_Test(SNew(UE::Slate::Private::SMultiInheritanceWidget_Child_SII));
+		MultiInheritanceWidget_Child_Test(SNew(UE::Slate::Private::SMultiInheritanceWidget_Child_ISI));
+		MultiInheritanceWidget_Child_Test(SNew(UE::Slate::Private::SMultiInheritanceWidget_Child_IIS));
+
+		MultiInheritanceWidget_Child_Test(SNew(UE::Slate::Private::SMultiInheritanceWidget_GrandChild_SII_SII));
+		MultiInheritanceWidget_Child_Test(SNew(UE::Slate::Private::SMultiInheritanceWidget_GrandChild_SII_ISI));
+		MultiInheritanceWidget_Child_Test(SNew(UE::Slate::Private::SMultiInheritanceWidget_GrandChild_SII_IIS));
+		MultiInheritanceWidget_Child_Test(SNew(UE::Slate::Private::SMultiInheritanceWidget_GrandChild_ISI_SII));
+		MultiInheritanceWidget_Child_Test(SNew(UE::Slate::Private::SMultiInheritanceWidget_GrandChild_ISI_ISI));
+		MultiInheritanceWidget_Child_Test(SNew(UE::Slate::Private::SMultiInheritanceWidget_GrandChild_ISI_IIS));
+		MultiInheritanceWidget_Child_Test(SNew(UE::Slate::Private::SMultiInheritanceWidget_GrandChild_IIS_SII));
+		MultiInheritanceWidget_Child_Test(SNew(UE::Slate::Private::SMultiInheritanceWidget_GrandChild_IIS_ISI));
+		MultiInheritanceWidget_Child_Test(SNew(UE::Slate::Private::SMultiInheritanceWidget_GrandChild_IIS_IIS));
+	}
 	return true;
 }
 
diff --git a/Engine/Source/Runtime/SlateCore/Public/Types/SlateAttributeDescriptor.h b/Engine/Source/Runtime/SlateCore/Public/Types/SlateAttributeDescriptor.h
--- a/Engine/Source/Runtime/SlateCore/Public/Types/SlateAttributeDescriptor.h
+++ b/Engine/Source/Runtime/SlateCore/Public/Types/SlateAttributeDescriptor.h
@@ -36,6 +36,19 @@ namespace SlateAttributePrivate
 		return reinterpret_cast<std::size_t>(&(Object->*Member))
 			- reinterpret_cast<std::size_t>(Object);;
 	}
+
+	/**
+	 * Find the offset of a base class from a derived class. Useful for multi inheritance
+	 */
+	template <typename InDerivedType, typename InBaseType>
+	std::ptrdiff_t GetBaseClassOffset()
+	{
+		static_assert(std::is_base_of_v<InBaseType, InDerivedType>, "InDerivedType needs to derive from InBaseType");
+
+		const InDerivedType* Object = reinterpret_cast<const InDerivedType*>(alignof(InDerivedType));
+		return reinterpret_cast<std::ptrdiff_t>(static_cast<const InBaseType*>(Object))
+			- reinterpret_cast<std::ptrdiff_t>(Object);
+	}
 }
 
 
@@ -629,7 +642,10 @@ private:
 		static_assert(decltype(_Property)::AttributeType == SlateAttributePrivate::ESlateAttributeType::Member, "The SlateProperty is not a TSlateAttribute. Do not use SLATE_ADD_MEMBER_ATTRIBUTE_DEFINITION"); \
 		static_assert(!decltype(_Property)::HasDefinedInvalidationReason, "When implementing the SLATE_DECLARE_WIDGET pattern, use TSlateAttribute without the invalidation reason."); \
 		static_assert(!std::is_same<decltype(_Reason), EInvalidateWidgetReason>::value || FSlateAttributeBase::IsInvalidateWidgetReasonSupported(_Reason), "The invalidation is not supported by the SlateAttribute system."); \
-		_Initializer.AddMemberAttribute(_Name, STRUCT_OFFSET(PrivateThisType, _Property), FSlateAttributeDescriptor::FInvalidateWidgetReasonAttribute{_Reason})
+		const std::ptrdiff_t UE_JOIN(__OffsetFromSWidget_, __LINE__) = STRUCT_OFFSET(PrivateThisType, _Property) - SlateAttributePrivate::GetBaseClassOffset<PrivateThisType, SWidget>(); \
+		checkf(UE_JOIN(__OffsetFromSWidget_, __LINE__) >= std::numeric_limits<FSlateAttributeDescriptor::OffsetType>::lowest(), TEXT("The offset from SWidget for attribute " #_Property " is too small")); \
+		checkf(UE_JOIN(__OffsetFromSWidget_, __LINE__) <= std::numeric_limits<FSlateAttributeDescriptor::OffsetType>::max(), TEXT("The offset from SWidget for attribute " #_Property " is too big")); \
+		_Initializer.AddMemberAttribute(_Name, (FSlateAttributeDescriptor::OffsetType)UE_JOIN(__OffsetFromSWidget_, __LINE__), FSlateAttributeDescriptor::FInvalidateWidgetReasonAttribute{_Reason})
 
 #define SLATE_ADD_MEMBER_ATTRIBUTE_DEFINITION(_Initializer, _Property, _Reason) \
 	SLATE_ADD_MEMBER_ATTRIBUTE_DEFINITION_WITH_NAME(_Initializer, GET_MEMBER_NAME_CHECKED(PrivateThisType, _Property), _Property, _Reason)
