Slate: Support non-axis aligned hittesting for SRetainer widgets

#su

#Synced-CL 50496765

diff --git a/Engine/Source/Runtime/SlateCore/Private/Input/HittestGrid.cpp b/Engine/Source/Runtime/SlateCore/Private/Input/HittestGrid.cpp
--- a/Engine/Source/Runtime/SlateCore/Private/Input/HittestGrid.cpp
+++ b/Engine/Source/Runtime/SlateCore/Private/Input/HittestGrid.cpp
@@ -171,6 +171,7 @@ FHittestGrid::FHittestGrid()
 	, GridOrigin(0, 0)
 	, GridSize(0, 0)
 	, CurrentUserIndex(INDEX_NONE)
+	, PaintToHitTestTransform() // Default identity
 #if WITH_SLATE_DEBUGGING
 	, LastNavigationMethod()
 #endif
@@ -242,7 +243,7 @@ TArray<FWidgetAndPointer> FHittestGrid::GetBubblePath(UE::Slate::FDeprecateVecto
 						Path.Append(MoveTemp(BubblePathExtension));
 					}
 				}
-	
+
 				return Path;
 			}
 		}
@@ -260,7 +261,7 @@ bool FHittestGrid::SetHittestArea(const UE::Slate::FDeprecateVector2DParameter&
 	{
 		GridSize = HittestDimensions;
 		NumCells = FIntPoint(FMath::CeilToInt(GridSize.X / CellSize.X), FMath::CeilToInt(GridSize.Y / CellSize.Y));
-		
+
 		const int32 NewTotalCells = NumCells.X * NumCells.Y;
 		if (NewTotalCells < 0)
 		{
@@ -344,9 +345,9 @@ namespace HittestGridDebuggingText
 	static FText ParentDisabled = LOCTEXT("StateParentNotEnabled", "ParentDisabled"); //~ A parent of the widget is disabled
 	static FText DoesNotSuportKeyboardFocus = LOCTEXT("StateDoesNotSuportKeyboardFocus", "Keyboard focus unsupported"); //~ THe widget does not support keyboard focus
 }
-	#define AddToNextFocusableWidgetCondidateDebugResults(Candidate, Result) { if (IntermediateResultsPtr) { IntermediateResultsPtr->Emplace((Candidate), (Result)); } }
+#define AddToNextFocusableWidgetCondidateDebugResults(Candidate, Result) { if (IntermediateResultsPtr) { IntermediateResultsPtr->Emplace((Candidate), (Result)); } }
 #else
-	#define AddToNextFocusableWidgetCondidateDebugResults(Candidate, Result) CA_ASSUME(Candidate)
+#define AddToNextFocusableWidgetCondidateDebugResults(Candidate, Result) CA_ASSUME(Candidate)
 #endif
 
 template<typename TCompareFunc, typename TSourceSideFunc, typename TDestSideFunc>
@@ -441,7 +442,7 @@ TSharedPtr<SWidget> FHittestGrid::FindFocusableWidget(FSlateRect WidgetRect, con
 					AddToNextFocusableWidgetCondidateDebugResults(TestWidget, HittestGridDebuggingText::Disabled);
 					continue;
 				}
-				
+
 				if (!TestWidget->SupportsKeyboardFocus())
 				{
 					AddToNextFocusableWidgetCondidateDebugResults(TestWidget, HittestGridDebuggingText::DoesNotSuportKeyboardFocus);
@@ -494,7 +495,7 @@ TSharedPtr<SWidget> FHittestGrid::FindFocusableWidget(FSlateRect WidgetRect, con
 			// Make sure all parents of the chosen widget are enabled before returning.
 			// Note that IsParentsEnabled is a costly function. We call it here as the last step to minimize the number of calls to it.
 			if (!UE::Slate::Private::IsParentsEnabled(BestWidget.Get()))
-			{	
+			{
 				// Find the next best widget because this one has disabled parents.
 				DisabledDestinations->Add(BestWidget);
 				return FindFocusableWidget(WidgetRect, SweptRect, AxisIndex, Increment, Direction, NavigationReply, CompareFunc, SourceSideFunc, DestSideFunc, UserIndex, IntermediateResultsPtr, DisabledDestinations);
@@ -617,7 +618,7 @@ TSharedPtr<SWidget> FHittestGrid::FindNextFocusableWidgetDefault(const FArranged
 #if WITH_SLATE_DEBUGGING
 	if (IntermediateResultsPtr)
 	{
-		FDebuggingFindNextFocusableWidgetArgs Args {StartingWidget, Direction, NavigationReply, RuleWidget, UserIndex, Widget , MoveTemp(*IntermediateResultsPtr)};
+		FDebuggingFindNextFocusableWidgetArgs Args{ StartingWidget, Direction, NavigationReply, RuleWidget, UserIndex, Widget , MoveTemp(*IntermediateResultsPtr) };
 		OnFindNextFocusableWidgetExecuted.Broadcast(this, Args);
 	}
 #endif
@@ -660,7 +661,7 @@ TSharedPtr<SWidget> FHittestGrid::FindNextFocusableWidgetExperimental(const FArr
 #else
 	TArray<FDebuggingFindNextFocusableWidgetArgs::FWidgetResult>* IntermediateResultsPtr = nullptr;
 #endif
-	
+
 	TSharedPtr<SWidget> NextFocusedWidget;
 	TInstancedStruct<FNavigationMethod> NavigationMethod = GetNavigationMethod(StartingWidget);
 	if (NavigationMethod.IsValid())
@@ -682,10 +683,10 @@ TSharedPtr<SWidget> FHittestGrid::FindNextFocusableWidgetExperimental(const FArr
 #if WITH_SLATE_DEBUGGING
 	if (IntermediateResultsPtr)
 	{
-		FDebuggingFindNextFocusableWidgetArgs Args {StartingWidget, Direction, NavigationReply, RuleWidget, UserIndex, NextFocusedWidget, MoveTemp(*IntermediateResultsPtr)};
+		FDebuggingFindNextFocusableWidgetArgs Args{ StartingWidget, Direction, NavigationReply, RuleWidget, UserIndex, NextFocusedWidget, MoveTemp(*IntermediateResultsPtr) };
 		OnFindNextFocusableWidgetExecuted.Broadcast(this, Args);
 	}
-	
+
 	if (GSlateNavigationDebugging && NavigationMethod.IsValid())
 	{
 		// Duplicate, as this pointer may outlive the navigation method associated with raw instanced struct pointer
@@ -730,19 +731,19 @@ bool FHittestGrid::IsValidCellCoord(const int32 XCoord, const int32 YCoord) cons
 void FHittestGrid::AddGrid(const TSharedRef<const FHittestGrid>& OtherGrid)
 {
 	auto GetCollapsedHittestGrid_NoTests = [](const FHittestGrid* HittestGrid, FCollapsedHittestGridArray& OutResult)
-	{
-		OutResult.Add(HittestGrid);
-		for (int32 Index = 0; Index < OutResult.Num(); ++Index)
 		{
-			for (const FAppendedGridData& AppendedGridData : OutResult[Index]->AppendedGridArray)
+			OutResult.Add(HittestGrid);
+			for (int32 Index = 0; Index < OutResult.Num(); ++Index)
 			{
-				if (const TSharedPtr<const FHittestGrid> Grid = AppendedGridData.Grid.Pin())
+				for (const FAppendedGridData& AppendedGridData : OutResult[Index]->AppendedGridArray)
 				{
-					OutResult.Add(Grid.Get());
+					if (const TSharedPtr<const FHittestGrid> Grid = AppendedGridData.Grid.Pin())
+					{
+						OutResult.Add(Grid.Get());
+					}
 				}
 			}
-		}
-	};
+		};
 
 	const bool bIsContains = AppendedGridArray.ContainsByPredicate([OtherGrid](const FAppendedGridData& Other) { return Other.Grid == OtherGrid; });
 	if (ensure(CanBeAppended(&OtherGrid.Get())))
@@ -801,6 +802,15 @@ void FHittestGrid::RemoveGrid(const SWidget* OtherGridOwner)
 	}
 }
 
+void FHittestGrid::SetPaintToHitTestTransform(const FSlateRenderTransform& InTransform)
+{
+	PaintToHitTestTransform = InTransform;
+	if (PaintToHitTestTransform.IsSet() && PaintToHitTestTransform.GetValue().IsIdentity())
+	{
+		PaintToHitTestTransform = TOptional<FSlateRenderTransform>();
+	}
+}
+
 bool FHittestGrid::CanBeAppended(const FHittestGrid* OtherGrid) const
 {
 	return OtherGrid && OtherGrid->Owner && this != OtherGrid && SameSize(OtherGrid);
@@ -828,7 +838,19 @@ void FHittestGrid::AddWidget(const SWidget* InWidget, int32 InBatchPriorityGroup
 	GridSpaceGeometry.AppendTransform(FSlateLayoutTransform(-GridWindowOrigin));
 
 	// Currently using grid offset because the grid covers all desktop space.
-	const FSlateRect BoundingRect = GridSpaceGeometry.GetRenderBoundingRect();
+	FSlateRect BoundingRect = GridSpaceGeometry.GetRenderBoundingRect();
+
+	if (PaintToHitTestTransform.IsSet())
+	{
+		// Ensure the cells check correspond to the hittest to paint remapping
+		const FSlateRotatedRect WindowOrientedClipRect = TransformRect(
+			Concatenate(
+				Inverse(GridSpaceGeometry.GetAccumulatedLayoutTransform()),
+				Concatenate(GridSpaceGeometry.GetAccumulatedRenderTransform(), PaintToHitTestTransform.GetValue())),
+			FSlateRotatedRect(GridSpaceGeometry.GetLayoutBoundingRect())
+		);
+		BoundingRect = WindowOrientedClipRect.ToBoundingRect();
+	}
 
 	// Starting and ending cells covered by this widget.	
 	const FIntPoint UpperLeftCell = GetCellCoordinate(BoundingRect.GetTopLeft());
@@ -991,12 +1013,27 @@ FHittestGrid::FIndexAndDistance FHittestGrid::GetHitIndexFromCellIndex(const FGr
 				if (bPointInsideClipMasks)
 				{
 					// Compute the render space clipping rect (FGeometry exposes a layout space clipping rect).
-					const FSlateRotatedRect WindowOrientedClipRect = TransformRect(
-						Concatenate(
-							Inverse(TestGeometry.GetAccumulatedLayoutTransform()),
-							TestGeometry.GetAccumulatedRenderTransform()),
-						FSlateRotatedRect(TestGeometry.GetLayoutBoundingRect())
-					);
+					FSlateRotatedRect WindowOrientedClipRect;
+					if (WidgetIndexes[i].GetGrid()->PaintToHitTestTransform.IsSet())
+					{
+						// If the hittestgrid contains an additional hittest rotation, apply this.
+						const FTransform2D GridPaintToHitTestTransform = WidgetIndexes[i].GetGrid()->PaintToHitTestTransform.GetValue();
+						WindowOrientedClipRect = TransformRect(
+							Concatenate(
+								Inverse(TestGeometry.GetAccumulatedLayoutTransform()),
+								Concatenate(TestGeometry.GetAccumulatedRenderTransform(), GridPaintToHitTestTransform)),
+							FSlateRotatedRect(TestGeometry.GetLayoutBoundingRect())
+						);
+					}
+					else
+					{
+						WindowOrientedClipRect = TransformRect(
+							Concatenate(
+								Inverse(TestGeometry.GetAccumulatedLayoutTransform()),
+								TestGeometry.GetAccumulatedRenderTransform()),
+							FSlateRotatedRect(TestGeometry.GetLayoutBoundingRect())
+						);
+					}
 
 					if (IsOverlappingSlateRotatedRect(WindowSpaceCoordinate, Params.Radius, WindowOrientedClipRect))
 					{
diff --git a/Engine/Source/Runtime/SlateCore/Public/Input/HittestGrid.h b/Engine/Source/Runtime/SlateCore/Public/Input/HittestGrid.h
--- a/Engine/Source/Runtime/SlateCore/Public/Input/HittestGrid.h
+++ b/Engine/Source/Runtime/SlateCore/Public/Input/HittestGrid.h
@@ -117,6 +117,8 @@ public:
 	/** Remove a grid that was appended. */
 	SLATECORE_API void RemoveGrid(const SWidget* OtherGridOwner);
 
+	SLATECORE_API void SetPaintToHitTestTransform(const FSlateRenderTransform& InTransform);
+
 	struct FDebuggingFindNextFocusableWidgetArgs
 	{
 		struct FWidgetResult
@@ -344,6 +346,11 @@ private:
 	/** The current slate user index that should be associated with any added widgets */
 	int32 CurrentUserIndex;
 
+	/** Optional delta transform applied during hit-testing.
+		Used by RetainerWidget to reconcile hittest space with the child’s paint space
+		when the retainer is non-grid aligned (i.e rotation) */
+	TOptional<FSlateRenderTransform> PaintToHitTestTransform;
+
 	/** The size in pixels of a cell */
 	static const FVector2f CellSize;
 
diff --git a/Engine/Source/Runtime/UMG/Private/Slate/SRetainerWidget.cpp b/Engine/Source/Runtime/UMG/Private/Slate/SRetainerWidget.cpp
--- a/Engine/Source/Runtime/UMG/Private/Slate/SRetainerWidget.cpp
+++ b/Engine/Source/Runtime/UMG/Private/Slate/SRetainerWidget.cpp
@@ -41,14 +41,14 @@ FAutoConsoleVariableRef EnableRetainedRendering(
 	TEXT("Whether to attempt to render things in SRetainerWidgets to render targets first.") 
 );
 
-/** True if we allow the transform to be modify to render in local space. */
-bool GSlateEnableRenderWithLocalTransform = true;
-FAutoConsoleVariableRef CVarGlateEnableRenderWithLocalTransform(
-	TEXT("Slate.EnableRetainedRenderingWithLocalTransform"),
-	GSlateEnableRenderWithLocalTransform,
-	TEXT("This console variable is no longer useful, as the SRetainerWidget render bugs workedaround by it have been fixed. It will be removed.")
+bool GSlateEnableRetainedRenderingHittestRemap = true;
+FAutoConsoleVariableRef CVarGlateEnableRetainedRenderingHittestRemap(
+	TEXT("Slate.EnableRetainedRenderingHittestRemap"),
+	GSlateEnableRetainedRenderingHittestRemap,
+	TEXT("If enables, the rotation delta of the SRetainer widget will be applied to the hittest grid. Used to resolve an issue with layout and paint space deviations in retained widgets")
 );
 
+
 static bool IsRetainedRenderingEnabled()
 {
 	return GEnableRetainedRendering != 0;
@@ -752,5 +752,21 @@ int32 SRetainerWidget::PaintSlowPath(const FSlateInvalidationContext& Context)
 	FTransform2D ScalePosTransform(OriginalPaintSpaceGeometry.Scale, OriginalPaintSpaceGeometry.AbsolutePosition);
 	SimplifiedRenderTransform = Concatenate(SimplifiedRenderTransform, Inverse(ScalePosTransform));
 	const FGeometry NewPaintSpaceGeometry = FGeometry::MakeRoot(OriginalPaintSpaceGeometry.GetLocalSize(), FSlateLayoutTransform(OriginalPaintSpaceGeometry.Scale, OriginalPaintSpaceGeometry.AbsolutePosition)).MakeChild(SimplifiedRenderTransform, FVector2D::ZeroVector);
+
+	if (GSlateEnableRetainedRenderingHittestRemap)
+	{
+		// Get the delta from simplified render transform to the original (should contain just the missing rotation about retainer pivot)
+		const FSlateRenderTransform NewToOriginalPaintSpace =
+			Concatenate(
+				Inverse(NewPaintSpaceGeometry.GetAccumulatedRenderTransform()),
+				OriginalPaintSpaceGeometry.GetAccumulatedRenderTransform());
+
+		HittestGrid->SetPaintToHitTestTransform(NewToOriginalPaintSpace);
+	}
+	else
+	{
+		HittestGrid->SetPaintToHitTestTransform(FSlateRenderTransform());
+	}
+
 	return SCompoundWidget::OnPaint(*Context.PaintArgs, NewPaintSpaceGeometry, Context.CullingRect, *Context.WindowElementList, Context.IncomingLayerId, Context.WidgetStyle, Context.bParentEnabled);
-}
\ No newline at end of file
+}
