Slate Dynamic Invalidation - Avoid switch to Slow Path if unecessary

Currently, to be able to catch widgets with invalid parents and get out of fast path, we revert to Slow Path as soon as a ChildOrder invalidation happens. Reverting to Slow Path while Fast Path is being processed will lead to some widgets being painted twice.
To avoid doing this unless necessary, we ask for a switch to Slow Path only if invalidation becomes not supported.

#Synced-CL 50145007

diff --git a/Engine/Source/Runtime/SlateCore/Private/Debugging/ConsoleSlateDebuggerPaint.cpp b/Engine/Source/Runtime/SlateCore/Private/Debugging/ConsoleSlateDebuggerPaint.cpp
--- a/Engine/Source/Runtime/SlateCore/Private/Debugging/ConsoleSlateDebuggerPaint.cpp
+++ b/Engine/Source/Runtime/SlateCore/Private/Debugging/ConsoleSlateDebuggerPaint.cpp
@@ -122,7 +122,7 @@ void FConsoleSlateDebuggerPaint::AddUpdatedWidget(const SWidget& Widget, const F
 
 	if (bLogWarningIfWidgetIsPaintedMoreThanOnce && WidgetInfo.UpdateCount != 0)
 	{
-		UE_LOG(LogSlateDebugger, Warning, TEXT("'%s' got painted more than once."), *(WidgetInfo.WidgetName));
+		UE_LOG(LogSlateDebugger, Warning, TEXT("'%s' got painted more than once on frame %u."), *(WidgetInfo.WidgetName), GFrameNumber);
 	}
 
 	if (bIncrementUpdateCount)
diff --git a/Engine/Source/Runtime/SlateCore/Private/Widgets/SWidget.cpp b/Engine/Source/Runtime/SlateCore/Private/Widgets/SWidget.cpp
--- a/Engine/Source/Runtime/SlateCore/Private/Widgets/SWidget.cpp
+++ b/Engine/Source/Runtime/SlateCore/Private/Widgets/SWidget.cpp
@@ -2110,13 +2110,18 @@ void SWidget::InvalidateCachedSupportsInvalidation()
 			return;
 		}
 
-		Widget->CachedSupportsInvalidationRecursive = ECachedInvalidationResult::NeedsRefresh;
 		if (GSlateIsOnFastUpdatePath && Widget->Advanced_IsInvalidationRoot())
 		{
-			// If we are currently in fast path, we want to stop and do a slow path as we might have invalid widgets that have been added
-			const_cast<FSlateInvalidationRoot*>(Widget->Advanced_AsInvalidationRoot())->InvalidateRootLayout(this);
+			if (!Widget->SupportsInvalidationRecursive(EInvalidationStrategy::RecomputeSelf))
+			{
+				// If we are currently in fast path, we want to stop and do a slow path as we might have invalid widgets that have been added
+				const_cast<FSlateInvalidationRoot*>(Widget->Advanced_AsInvalidationRoot())->InvalidateRootLayout(this);
+			}
+			// We do not need to refresh the parents of this widget as `SupportsInvalidationRecursive` would have taken care of it if necessary.
+			return;
 		}
 
+		Widget->CachedSupportsInvalidationRecursive = ECachedInvalidationResult::NeedsRefresh;
 		Widget = Widget->GetParentWidget();
 	}
 }
