Fonts:
Set Slate.Font.NormalizeCachedFontRanges to true by default in the Engine. This means that all composite fonts will no longer allow intersecting unicode ranges for subtypefaces even if they are provided in the Font Editor. The font ranges will be normalized and made disjoint for the composite font cache.

diff --git a/Engine/Source/Runtime/SlateCore/Private/Fonts/FontCacheCompositeFont.cpp b/Engine/Source/Runtime/SlateCore/Private/Fonts/FontCacheCompositeFont.cpp
--- a/Engine/Source/Runtime/SlateCore/Private/Fonts/FontCacheCompositeFont.cpp
+++ b/Engine/Source/Runtime/SlateCore/Private/Fonts/FontCacheCompositeFont.cpp
@@ -21,7 +21,7 @@ FAutoConsoleVariableRef CVarAsyncLazyLoad(
 	GAsyncFontLazyLoad,
 	TEXT("Causes unloaded font faces that are lazily loaded, to be loaded asynchronously, until then the font won't measure correctly.  Once complete the UI will invalidate. Note that async font loading is disallowed in Editor/PIE, standalone and cooked builds are fine. "));
 
-static bool GNormalizeCachedFontRanges = false;
+static bool GNormalizeCachedFontRanges = true;
 FAutoConsoleVariableRef CVarNormalizeCachedFontRanges(
 	TEXT("Slate.Font.NormalizeCachedFontRanges"),
 	GNormalizeCachedFontRanges,
