[Fix] Small fix for `SWidget::InvalidateCachedSupportsInvalidation` to early exit if a shared instance is not available (called from the destructor for example) instead of crashing

#Synced-CL 50897806

diff --git a/Engine/Source/Runtime/SlateCore/Private/Widgets/SWidget.cpp b/Engine/Source/Runtime/SlateCore/Private/Widgets/SWidget.cpp
--- a/Engine/Source/Runtime/SlateCore/Private/Widgets/SWidget.cpp
+++ b/Engine/Source/Runtime/SlateCore/Private/Widgets/SWidget.cpp
@@ -2107,6 +2107,11 @@ UE_TRACE_EVENT_END()
 
 void SWidget::InvalidateCachedSupportsInvalidation()
 {
+	if (!DoesSharedInstanceExist())
+	{
+		return;
+	}
+
 #if UE_SLATE_TRACE_ENABLED
 	UE_TRACE_LOG_SCOPED_T(Cpu, SWidget_InvalidateCachedSupportsInvalidation, SlateChannel)
 		<< SWidget_InvalidateCachedSupportsInvalidation.CreatedInLocation(*GetCreatedInLocation().ToString());
