Fonts:
Improved logging for the composite font cache to log synchronous and asynchronous font loading paths for easier debugging.

diff --git a/Engine/Source/Runtime/SlateCore/Private/Fonts/FontCacheCompositeFont.cpp b/Engine/Source/Runtime/SlateCore/Private/Fonts/FontCacheCompositeFont.cpp
--- a/Engine/Source/Runtime/SlateCore/Private/Fonts/FontCacheCompositeFont.cpp
+++ b/Engine/Source/Runtime/SlateCore/Private/Fonts/FontCacheCompositeFont.cpp
@@ -12,9 +12,9 @@
 #include "Internationalization/Internationalization.h"
 #include "Misc/FileHelper.h"
 #include "Serialization/MemoryReader.h"
+#include "Textures/TextureAtlas.h"
 #include "Trace/SlateMemoryTags.h"
 
-
 static bool GAsyncFontLazyLoad = false;
 FAutoConsoleVariableRef CVarAsyncLazyLoad(
 	TEXT("Slate.Font.AsyncLazyLoad"),
@@ -57,11 +57,11 @@ public:
 	{
 		if (!FFileHelper::LoadFileToArray(FontFaceData, *FontData.GetFontFilename()))
 		{
-			UE_LOG(LogSlate, Warning, TEXT("FAsyncLoadFontFaceData failed to load or process '%s' with face index %d"), *FontData.GetFontFilename(), FontData.GetSubFaceIndex());
+			UE_LOGF(LogSlate, Warning, "FAsyncLoadFontFaceData failed to load or process '%ls' with face index %d", *FontData.GetFontFilename(), FontData.GetSubFaceIndex());
 		}
 		else
 		{
-			UE_LOG(LogSlate, Display, TEXT("FAsyncLoadFontFaceData successfully loaded '%s' with face index %d"), *FontData.GetFontFilename(), FontData.GetSubFaceIndex());
+			UE_LOGF(LogSlate, Display, "FAsyncLoadFontFaceData successfully loaded '%ls' with face index %d", *FontData.GetFontFilename(), FontData.GetSubFaceIndex());
 		}
 	}
 
@@ -675,7 +675,7 @@ const FFontData& FCompositeFontCache::GetFontDataForCodepoint(const FSlateFontIn
 		if (const FFontData* FoundFontData = GetFontDataForCharacterInTypeface(CachedFallbackTypefaceData, CachedDefaultTypefaceData, /*InSkipCharacterCheck*/true))
 		{
 			OutScalingFactor = CachedFallbackTypefaceData->GetScalingFactor();
-			UE_LOG(LogSlate, Verbose, TEXT("Using fallback font %s for codepoint U+%x, as it was not found in default nor sub typefaces of SlateFontInfo %s."), *FoundFontData->GetFontFilename(), InCodepoint, *InFontInfo.TypefaceFontName.ToString());
+			UE_LOGF(LogSlate, Verbose, "Using fallback font %ls for codepoint U+%x, as it was not found in default nor sub typefaces of SlateFontInfo %ls.", *FoundFontData->GetFontFilename(), InCodepoint, *InFontInfo.TypefaceFontName.ToString());
 			return *FoundFontData;
 		}
 
@@ -701,8 +701,8 @@ TSharedPtr<FFreeTypeFace> FCompositeFontCache::GetFontFace(const FFontData& InFo
 		SCOPE_CYCLE_COUNTER(STAT_SlateLoadFont);
 
 		// IMPORTANT: Do not log from this function until the new font has been added to the FontFaceMap, as it may be the Output Log font being loaded, which would cause an infinite recursion!
-		FString LoadLogMessage;
-
+		TArray<FString> LogMessages;
+		const FString& FontFilename = InFontData.GetFontFilename();
 		{
 			// If this font data is referencing an asset, we just need to load it from memory
 			FFontFaceDataConstPtr FontFaceData = InFontData.GetFontFaceData();
@@ -721,6 +721,7 @@ TSharedPtr<FFreeTypeFace> FCompositeFontCache::GetFontFace(const FFontData& InFo
 				{
 					if (UE::Private::FontCacheCompositeFont::IsAsyncFontLoadingAllowed())
 					{
+						LogMessages.Add(FString::Printf(TEXT("Requested to async lazy load '%s' font file with subface index %d in %s thread."), *FontFilename, InFontData.GetSubFaceIndex(), GetCurrentSlateTextureAtlasThreadIdAsString()));
 						FAsyncTask<FAsyncLoadFontFaceData>* AsyncFontLoadTaskPtr = new FAsyncTask<FAsyncLoadFontFaceData>(InFontData);
 						LoadFontFaceTasks.Add(AsyncFontLoadTaskPtr);
 						AsyncFontLoadTaskPtr->StartBackgroundTask();
@@ -729,18 +730,19 @@ TSharedPtr<FFreeTypeFace> FCompositeFontCache::GetFontFace(const FFontData& InFo
 					}
 					else
 					{
+						LogMessages.Add(FString::Printf(TEXT("Synchronously lazy loading font file '%s' with subface index %d in %s thread."), *FontFilename, InFontData.GetSubFaceIndex(), GetCurrentSlateTextureAtlasThreadIdAsString()));
 						const double FontDataLoadStartTime = FPlatformTime::Seconds();
 
 						TArray<uint8> FontFaceDataArray;
-						const FString& FontFilename = InFontData.GetFontFilename();
 						if (FFileHelper::LoadFileToArray(FontFaceDataArray, *FontFilename))
 						{
 							const int32 FontDataSizeKB = (FontFaceDataArray.Num() + 1023) / 1024;
-							LoadLogMessage = FString::Printf(TEXT("Took %f seconds to synchronously load lazily loaded font '%s' (%dK)"), float(FPlatformTime::Seconds() - FontDataLoadStartTime), *InFontData.GetFontFilename(), FontDataSizeKB);
+							LogMessages.Add(FString::Printf(TEXT("Took %f seconds to synchronously load lazily loaded font '%s' (%dK) on %s thread."), float(FPlatformTime::Seconds() - FontDataLoadStartTime), *FontFilename, FontDataSizeKB, GetCurrentSlateTextureAtlasThreadIdAsString()));
 
 							if (Private::IsFontFileName(FontFilename) || Private::IsFontFileData(FontFaceDataArray))
 							{
 								FaceAndMemory = MakeShared<FFreeTypeFace>(FTLibrary, FFontFaceData::MakeFontFaceData(MoveTemp(FontFaceDataArray)), InFontData.GetSubFaceIndex(), InFontData.GetLayoutMethod());
+								LogMessages.Add(FString::Printf(TEXT("Freetype font face in memory successfully created with synchronously loaded raw font file '%s' subface %d in %s thread."), *FontFilename, InFontData.GetSubFaceIndex(), GetCurrentSlateTextureAtlasThreadIdAsString()));
 							}
 							else
 							{
@@ -749,6 +751,7 @@ TSharedPtr<FFreeTypeFace> FCompositeFontCache::GetFontFace(const FFontData& InFo
 								FontFaceData->Serialize(Ar);
 
 								FaceAndMemory = MakeShared<FFreeTypeFace>(FTLibrary, FontFaceData, InFontData.GetSubFaceIndex(), InFontData.GetLayoutMethod());
+								LogMessages.Add(FString::Printf(TEXT("Freetype font face in memory successfully created with synchronously loaded cooked font file '%s' subface %d in %s thread."), *FontFilename, InFontData.GetSubFaceIndex(), GetCurrentSlateTextureAtlasThreadIdAsString()));
 							}
 						}
 					}
@@ -756,6 +759,7 @@ TSharedPtr<FFreeTypeFace> FCompositeFontCache::GetFontFace(const FFontData& InFo
 				break;
 			case EFontLoadingPolicy::Stream:
 				{
+					LogMessages.Add(FString::Printf(TEXT("Requested to stream font file '%s' with subfontface index %d in %s thread."), *FontFilename, InFontData.GetSubFaceIndex(), GetCurrentSlateTextureAtlasThreadIdAsString()));
 					FaceAndMemory = MakeShared<FFreeTypeFace>(FTLibrary, InFontData.GetFontFilename(), InFontData.GetSubFaceIndex(), InFontData.GetLayoutMethod());
 				}
 				break;
@@ -771,12 +775,11 @@ TSharedPtr<FFreeTypeFace> FCompositeFontCache::GetFontFace(const FFontData& InFo
 			FaceAndMemory->OverrideDescent(InFontData.IsDescendOverridden(), InFontData.GetDescendOverriddenValue());
 
 			FontFaceMap.Add(InFontData, FaceAndMemory);
-
-			if (!LoadLogMessage.IsEmpty())
+			LogMessages.Add(FString::Printf(TEXT("Font face '%s' subface index %d successfully added to font face map in %s thread."), *FontFilename, InFontData.GetSubFaceIndex(), GetCurrentSlateTextureAtlasThreadIdAsString()));
+			for (const FString& LogMessage : LogMessages)
 			{
-				UE_LOG(LogSlate, Log, TEXT("%s"), *LoadLogMessage);
+				UE_LOGF(LogSlate, Display, "%ls", *LogMessage);
 			}
-
 #if WITH_EDITOR //Triggers texts layout regeneration only in editor, don't want to slow down things in engine.
 			GSlateLayoutGeneration++;
 #endif
@@ -784,7 +787,7 @@ TSharedPtr<FFreeTypeFace> FCompositeFontCache::GetFontFace(const FFontData& InFo
 		else
 		{
 			FaceAndMemory.Reset();
-			UE_LOG(LogSlate, Warning, TEXT("GetFontFace failed to load or process '%s' with face index %d"), *InFontData.GetFontFilename(), InFontData.GetSubFaceIndex());
+			UE_LOGF(LogSlate, Warning, "GetFontFace failed to load or process '%ls' with subface index %d", *FontFilename, InFontData.GetSubFaceIndex());
 		}
 	}
 	return FaceAndMemory;
@@ -805,6 +808,7 @@ void FCompositeFontCache::FlushCompositeFont(const FCompositeFont& InCompositeFo
 
 void FCompositeFontCache::FlushCache()
 {
+	UE_LOGF(LogSlate, Display, "Flushing composite font cache on %ls thread. Clearing %d entries in composite font map and %d entries in font face map.", GetCurrentSlateTextureAtlasThreadIdAsString(), CompositeFontToCachedDataMap.Num(), FontFaceMap.Num());
 	CompositeFontToCachedDataMap.Empty();
 	FontFaceMap.Empty();
 }
@@ -814,7 +818,10 @@ void FCompositeFontCache::Update()
 	if (LoadFontFaceTasks.Num() > 0)
 	{
 		bool bRefreshNeeded = false;
-
+		const int32 NumLoadFontFaceTasks = LoadFontFaceTasks.Num();
+		int32 NumCompletedFontFaceTasks = 0;
+		const TCHAR* CurrentThreadString = GetCurrentSlateTextureAtlasThreadIdAsString();
+		UE_LOGF(LogSlate, Verbose, "Updating composite font cache on %ls thread. %d font loading tasks need updating.", CurrentThreadString, NumLoadFontFaceTasks);
 		for (int32 LoadTaskIndex = 0; LoadTaskIndex < LoadFontFaceTasks.Num(); LoadTaskIndex++)
 		{
 			FAsyncTask<FAsyncLoadFontFaceData>* LoadFontFaceTask = LoadFontFaceTasks[LoadTaskIndex];
@@ -832,7 +839,9 @@ void FCompositeFontCache::Update()
 						if (Private::IsFontFileName(FontFilename) || Private::IsFontFileData(FontFaceDataArray))
 						{
 							FaceAndMemory->CompleteAsyncLoad(FTLibrary, FFontFaceData::MakeFontFaceData(MoveTemp(FontFaceDataArray)), FontData.GetSubFaceIndex());
+							UE_LOGF(LogSlate, Display, "Freetype font face in memory successfully updated with async loaded raw font file '%ls' subface %d in %ls thread.", *FontFilename, FontData.GetSubFaceIndex(), CurrentThreadString);
 						}
+						// Handles cooked .ufont files 
 						else
 						{
 							FFontFaceDataRef FontFaceData = FFontFaceData::MakeFontFaceData();
@@ -840,11 +849,13 @@ void FCompositeFontCache::Update()
 							FontFaceData->Serialize(Ar);
 
 							FaceAndMemory->CompleteAsyncLoad(FTLibrary, FontFaceData, FontData.GetSubFaceIndex());
+							UE_LOGF(LogSlate, Display, "Freetype font face in memory successfully updated with async loaded cooked font file '%ls' subface %d in %ls thread.", *FontFilename, FontData.GetSubFaceIndex(), CurrentThreadString);
 						}
 					}
 					else
 					{
 						FaceAndMemory->FailAsyncLoad();
+						UE_LOGF(LogSlate, Warning, "Async loaded font data for '%ls' font file with subface index %d on %ls thread is empty. The font face will not be available.", *FontFilename, FontData.GetSubFaceIndex(), CurrentThreadString);
 					}
 
 					bRefreshNeeded = true;
@@ -853,14 +864,16 @@ void FCompositeFontCache::Update()
 				delete LoadFontFaceTask;
 				LoadFontFaceTasks.RemoveAt(LoadTaskIndex);
 				LoadTaskIndex--;
+				++NumCompletedFontFaceTasks;
 			}
 		}
 
 		if (bRefreshNeeded)
 		{
-			AsyncTask(ENamedThreads::GameThread, []() {
+			AsyncTask(ENamedThreads::GameThread, [NumLoadFontFaceTasks, NumCompletedFontFaceTasks, CurrentThreadString]() {
 				if (FSlateApplicationBase::IsInitialized())
 				{
+					UE_LOGF(LogSlate, Display, "Invalidating all widgets from %d/%d async loading font faces completing loading on %ls thread.", NumCompletedFontFaceTasks, NumLoadFontFaceTasks, CurrentThreadString);
 					FSlateApplicationBase::Get().InvalidateAllWidgets(false);
 					GSlateLayoutGeneration++;
 				}
diff --git a/Engine/Source/Runtime/SlateCore/Private/Textures/TextureAtlas.cpp b/Engine/Source/Runtime/SlateCore/Private/Textures/TextureAtlas.cpp
--- a/Engine/Source/Runtime/SlateCore/Private/Textures/TextureAtlas.cpp
+++ b/Engine/Source/Runtime/SlateCore/Private/Textures/TextureAtlas.cpp
@@ -15,6 +15,23 @@ DEFINE_STAT(STAT_SlateTextureDataMemory);
 DECLARE_MEMORY_STAT(TEXT("Texture Atlas Memory (CPU)"), STAT_SlateTextureAtlasMemory, STATGROUP_SlateMemory);
 
 
+const TCHAR* LexToString(ESlateTextureAtlasThreadId ThreadId)
+{
+	switch (ThreadId)
+	{
+		case ESlateTextureAtlasThreadId::Unknown: 
+			return TEXT("Unknown");
+		case ESlateTextureAtlasThreadId::Game: 
+			return TEXT("Game");
+		case ESlateTextureAtlasThreadId::Render: 
+			return TEXT("Render");
+		default:
+			checkNoEntry();
+			break;
+	}
+	return TEXT("Unsupported thread ID");
+}
+
 ESlateTextureAtlasThreadId GetCurrentSlateTextureAtlasThreadId()
 {
 	// Note: For Game thread ownership, there is a point at which multiple worker threads operate on text simultaneously while the game thread is blocked
@@ -32,6 +49,11 @@ ESlateTextureAtlasThreadId GetCurrentSlateTextureAtlasThreadId()
 	return ESlateTextureAtlasThreadId::Unknown;
 }
 
+const TCHAR* GetCurrentSlateTextureAtlasThreadIdAsString()
+{
+	return LexToString(GetCurrentSlateTextureAtlasThreadId());
+}
+
 uint32 GetSlateFontAtlasContentBytesPerPixel(ESlateFontAtlasContentType InContentType)
 {
 	switch (InContentType)
diff --git a/Engine/Source/Runtime/SlateCore/Public/Textures/TextureAtlas.h b/Engine/Source/Runtime/SlateCore/Public/Textures/TextureAtlas.h
--- a/Engine/Source/Runtime/SlateCore/Public/Textures/TextureAtlas.h
+++ b/Engine/Source/Runtime/SlateCore/Public/Textures/TextureAtlas.h
@@ -39,6 +39,7 @@ enum ESlateTextureAtlasPaddingStyle
 
 /** 
  * The type of thread that owns a texture atlas - this is the only thread that can safely update it 
+ * If you need to update this enum, keep it in sync with the LexToString overload below.
  */
 enum class ESlateTextureAtlasThreadId
 {
@@ -50,11 +51,18 @@ enum class ESlateTextureAtlasThreadId
 	Render = 1,
 };
 
+SLATECORE_API const TCHAR* LexToString(ESlateTextureAtlasThreadId ThreadId);
+
 /** 
  * Get the correct atlas thread ID based on the thread we're currently in 
  */
 SLATECORE_API ESlateTextureAtlasThreadId GetCurrentSlateTextureAtlasThreadId();
 
+/**
+ * Convenience function to get the correct atlas thread ID based on the thread we're currently in as a string.
+ */
+SLATECORE_API const TCHAR* GetCurrentSlateTextureAtlasThreadIdAsString();
+
 /**
  * Returns the byte size of a single pixel of a font atlas with the specified content type
  */
