Fixed Slate.HitTestGridDebugging, which was rendering incorrectly due to changes in the UE5 style's FocusRectangle brush. This brush now uses a rounded box brush with an outline instead of a border brush, so we'll apply the player-specific coloring as a mostly transparent overlay in lieu of coloring the outline

#Synced-CL 51374034

diff --git a/Engine/Source/Runtime/SlateCore/Private/Input/HittestGrid.cpp b/Engine/Source/Runtime/SlateCore/Private/Input/HittestGrid.cpp
--- a/Engine/Source/Runtime/SlateCore/Private/Input/HittestGrid.cpp
+++ b/Engine/Source/Runtime/SlateCore/Private/Input/HittestGrid.cpp
@@ -1217,7 +1217,7 @@ constexpr FLinearColor GetUserColorFromIndex(int32 Index)
 	case 3:
 		return FLinearColor::Yellow;
 	default:
-		return FLinearColor::White;
+		return FLinearColor::Transparent;
 	}
 };
 
@@ -1252,7 +1252,7 @@ void FHittestGrid::DisplayGrid(int32 InLayer, const FGeometry& AllottedGeometry,
 					CachedWidget->GetPaintSpaceGeometry().ToPaintGeometry(),
 					Brush,
 					ESlateDrawEffect::None,
-					GetUserColorFromIndex(CurWidgetData.UserIndex)
+					GetUserColorFromIndex(CurWidgetData.UserIndex).CopyWithNewOpacity(.005f)
 				);
 			}
 		}
