Slate: Console Debugger: filter for game/app content more specifically than whole window.

Widgets can now explicitly declare themselves and/or their children as "project content" during prepass. The CVars that control whether debuggers omit non-content widgets have been renamed to be agnostic to game vs non-game projects.

#rn The Paint, Prepass, and Update visualizers in the Slate Console Debugger now work in editor without requiring the use of a new PIE window, and omit non-content widgets (e.g. title bar, exit button).
#Synced-CL 51476348

diff --git a/Engine/Source/Runtime/Engine/Private/Slate/DebugCanvas.cpp b/Engine/Source/Runtime/Engine/Private/Slate/DebugCanvas.cpp
--- a/Engine/Source/Runtime/Engine/Private/Slate/DebugCanvas.cpp
+++ b/Engine/Source/Runtime/Engine/Private/Slate/DebugCanvas.cpp
@@ -371,6 +371,13 @@ SDebugCanvas::SDebugCanvas()
 void SDebugCanvas::Construct(const FArguments& InArgs)
 {
 	SceneViewport = InArgs._SceneViewport;
+	bHasCustomPrepass = true;
+}
+
+bool SDebugCanvas::CustomPrepass(float LayoutScaleMultiplier)
+{
+	SetIsProjectContent(false);
+	return false;
 }
 
 int32 SDebugCanvas::OnPaint(const FPaintArgs& Args, const FGeometry& AllottedGeometry, const FSlateRect& MyCullingRect, FSlateWindowElementList& OutDrawElements, int32 LayerId, const FWidgetStyle& InWidgetStyle, bool bParentEnabled) const
diff --git a/Engine/Source/Runtime/Engine/Private/Slate/SGameLayerManager.cpp b/Engine/Source/Runtime/Engine/Private/Slate/SGameLayerManager.cpp
--- a/Engine/Source/Runtime/Engine/Private/Slate/SGameLayerManager.cpp
+++ b/Engine/Source/Runtime/Engine/Private/Slate/SGameLayerManager.cpp
@@ -101,6 +101,7 @@ SGameLayerManager::SGameLayerManager()
 
 void SGameLayerManager::Construct(const SGameLayerManager::FArguments& InArgs)
 {
+	bHasCustomPrepass = true;
 	SceneViewport = InArgs._SceneViewport;
 
 	// In PIE we should default to per user hit testing being off because developers will need the mouse and keyboard to work for all players
@@ -436,14 +437,12 @@ int32 SGameLayerManager::OnPaint(const FPaintArgs& Args, const FGeometry& Allott
 	}
 #endif
 
-	OutDrawElements.SetIsInGameLayer(true);
 	int32 ResultLayer = SCompoundWidget::OnPaint(Args, AllottedGeometry, MyCullingRect, OutDrawElements, LayerId, InWidgetStyle, bParentEnabled);
 
 #if WITH_SLATE_DEBUGGING
 	ResultLayer = FSlateNavigationDebugger::Get().Paint(OutDrawElements, ResultLayer);
 #endif
 
-	OutDrawElements.SetIsInGameLayer(false);
 	return ResultLayer;
 }
 
@@ -754,6 +753,12 @@ void SGameLayerManager::SetWindowTitleBarVisibility(bool bIsVisible)
 	UpdateWindowTitleBarVisibility();
 }
 
+bool SGameLayerManager::CustomPrepass(float LayoutScaleMultiplier)
+{
+	SetIsProjectContentParent(true);
+	return true;
+}
+
 void SGameLayerManager::UpdateWindowTitleBar()
 {
 	if (WindowTitleBarState.ContentWidget.IsValid())
diff --git a/Engine/Source/Runtime/Engine/Public/Slate/DebugCanvas.h b/Engine/Source/Runtime/Engine/Public/Slate/DebugCanvas.h
--- a/Engine/Source/Runtime/Engine/Public/Slate/DebugCanvas.h
+++ b/Engine/Source/Runtime/Engine/Public/Slate/DebugCanvas.h
@@ -34,6 +34,7 @@ public:
 	ENGINE_API void Construct(const FArguments& InArgs);
 
 	/** SWidget interface */
+	ENGINE_API virtual bool CustomPrepass(float LayoutScaleMultiplier) override;
 	ENGINE_API int32 OnPaint(const FPaintArgs& Args, const FGeometry& AllottedGeometry, const FSlateRect& MyCullingRect, FSlateWindowElementList& OutDrawElements, int32 LayerId, const FWidgetStyle& InWidgetStyle, bool bParentEnabled) const override;
 	ENGINE_API virtual FVector2D ComputeDesiredSize(float LayoutScaleMultiplier) const override;
 
diff --git a/Engine/Source/Runtime/Engine/Public/Slate/SGameLayerManager.h b/Engine/Source/Runtime/Engine/Public/Slate/SGameLayerManager.h
--- a/Engine/Source/Runtime/Engine/Public/Slate/SGameLayerManager.h
+++ b/Engine/Source/Runtime/Engine/Public/Slate/SGameLayerManager.h
@@ -152,14 +152,15 @@ public:
 	ENGINE_API virtual void SetWindowTitleBarVisibility(bool bIsVisible);
 	// End IGameLayerManager
 
-public:
-
 	//~ Begin SWidget overrides
+protected:
+	virtual bool CustomPrepass(float LayoutScaleMultiplier) override;
 	ENGINE_API virtual void Tick(const FGeometry& AllottedGeometry, const double InCurrentTime, const float InDeltaTime) override;
 	ENGINE_API virtual int32 OnPaint(const FPaintArgs& Args, const FGeometry& AllottedGeometry, const FSlateRect& MyCullingRect, FSlateWindowElementList& OutDrawElements, int32 LayerId, const FWidgetStyle& InWidgetStyle, bool bParentEnabled) const override;
 	ENGINE_API virtual bool OnVisualizeTooltip(const TSharedPtr<SWidget>& TooltipContent) override;
 	//~ End SWidget overrides
 
+public:
 	/**
 	 * Function will instruct internal DPI computations to use a provided reference viewport size instead of the actual viewport size.
 	 * After the DPI will be retrieved it will be scaled down with the ratio between the actual viewport size and the provided one.
diff --git a/Engine/Source/Runtime/Slate/Private/Framework/Application/SlateApplication.cpp b/Engine/Source/Runtime/Slate/Private/Framework/Application/SlateApplication.cpp
--- a/Engine/Source/Runtime/Slate/Private/Framework/Application/SlateApplication.cpp
+++ b/Engine/Source/Runtime/Slate/Private/Framework/Application/SlateApplication.cpp
@@ -3613,9 +3613,9 @@ void FSlateApplication::ProcessReply( const FWidgetPath& CurrentEventPath, const
 
 		if (NavigationSource.IsValid())
 		{
-			if (!GSlateEnableGamepadEditorNavigation && TheReply.GetNavigationGenesis() == ENavigationGenesis::Controller && !NavigationSource.GetLastWidget()->GetPersistentState().bIsInGameLayer)
+			if (!GSlateEnableGamepadEditorNavigation && TheReply.GetNavigationGenesis() == ENavigationGenesis::Controller && !NavigationSource.GetLastWidget()->GetPersistentState().IsProjectContent())
 			{
-				// Gamepad navigation while not in a game layer, do nothing as specified by GSlateEnableGamepadEditorNavigation
+				// Gamepad navigation while not in project content, do nothing as specified by GSlateEnableGamepadEditorNavigation
 			}
 			else if (TheReply.GetNavigationDestination().IsValid())
 			{
diff --git a/Engine/Source/Runtime/Slate/Public/Widgets/Input/SNumericEntryBox.h b/Engine/Source/Runtime/Slate/Public/Widgets/Input/SNumericEntryBox.h
--- a/Engine/Source/Runtime/Slate/Public/Widgets/Input/SNumericEntryBox.h
+++ b/Engine/Source/Runtime/Slate/Public/Widgets/Input/SNumericEntryBox.h
@@ -750,7 +750,7 @@ private:
 #if WITH_EDITOR
 		if constexpr (std::is_floating_point_v<NumericType>)
 		{
-			if (!GetPersistentState().bIsInGameLayer && !Interface->GetIndicateNearlyInteger())
+			if (!GetPersistentState().IsProjectContent() && !Interface->GetIndicateNearlyInteger())
 			{
 				Interface->SetIndicateNearlyInteger(true);
 				bCachedValueStringDirty = true;
diff --git a/Engine/Source/Runtime/SlateCore/Private/Debugging/ConsoleSlateDebuggerPaint.cpp b/Engine/Source/Runtime/SlateCore/Private/Debugging/ConsoleSlateDebuggerPaint.cpp
--- a/Engine/Source/Runtime/SlateCore/Private/Debugging/ConsoleSlateDebuggerPaint.cpp
+++ b/Engine/Source/Runtime/SlateCore/Private/Debugging/ConsoleSlateDebuggerPaint.cpp
@@ -51,11 +51,16 @@ FConsoleSlateDebuggerPaint::FConsoleSlateDebuggerPaint()
 		, bLogWarningIfWidgetIsPaintedMoreThanOnce
 		, TEXT("Option to log a warning if a widget is painted more than once in a single frame.")
 		, FConsoleVariableDelegate::CreateRaw(this, &FConsoleSlateDebuggerPaint::SaveConfigOnVariableChanged))
-	, OnlyGameWindowRefCVar(
-		TEXT("SlateDebugger.Paint.OnlyGameWindow")
-		, bDebugGameWindowOnly
-		, TEXT("Option to only debug the game window")
+	, OnlyProjectContentRefCVar(
+		TEXT("SlateDebugger.Paint.OnlyProjectContent")
+		, bDebugProjectContentOnly
+		, TEXT("Only debug the project's contents (exclude editor, chrome, diagnostic output)")
 		, FConsoleVariableDelegate::CreateRaw(this, &FConsoleSlateDebuggerPaint::RebuildWidgetListOnVariableChanged))
+	, OnlyGameWindowDepCVar(
+		TEXT("SlateDebugger.Paint.OnlyGameWindow"),
+		TEXT("SlateDebugger.Paint.OnlyProjectContent"),
+		TEXT("5.8"),
+		EShadowCVarBehavior::Warn)
 	, OnlyInvalidationRootRefCVar(
 		TEXT("SlateDebugger.Paint.OnlyInvalidationRoot")
 		, bInvalidationRootOnly
@@ -87,7 +92,7 @@ void FConsoleSlateDebuggerPaint::LoadConfig()
 	const FString Section = *GetConfigSection();
 	GConfig->GetBool(*Section, TEXT("bLogWarningIfWidgetIsPaintedMoreThanOnce"), bLogWarningIfWidgetIsPaintedMoreThanOnce, *GEditorPerProjectIni);
 
-	OnlyGameWindowRefCVar->Set(bDebugGameWindowOnly, ECVF_SetByProjectSetting);
+	OnlyProjectContentRefCVar->Set(bDebugProjectContentOnly, ECVF_SetByProjectSetting);
 	OnlyInvalidationRootRefCVar->Set(bInvalidationRootOnly, ECVF_SetByProjectSetting);
 	EnableWidgetsNameListRefCVar->Set(bDisplayWidgetsNameList, ECVF_SetByProjectSetting);
 	DrawFillEnabledRefCVar->Set(bDrawBox, ECVF_SetByProjectSetting);
@@ -133,9 +138,7 @@ void FConsoleSlateDebuggerPaint::AddUpdatedWidget(const SWidget& Widget, const F
 
 void FConsoleSlateDebuggerPaint::HandleEndWidgetPaint(const SWidget* Widget, const FSlateWindowElementList& OutDrawElements, int32 LayerId)
 {
-	// Exclude all windows but the game window
-	const SWindow* WindowToDrawIn = OutDrawElements.GetPaintWindow();
-	if (bDebugGameWindowOnly && (WindowToDrawIn->GetType() != EWindowType::GameWindow && WindowToDrawIn->GetTag() != PIEWindowTag))
+	if (bDebugProjectContentOnly && !Widget->IsProjectContent())
 	{
 		return;
 	}
diff --git a/Engine/Source/Runtime/SlateCore/Private/Debugging/ConsoleSlateDebuggerPaint.h b/Engine/Source/Runtime/SlateCore/Private/Debugging/ConsoleSlateDebuggerPaint.h
--- a/Engine/Source/Runtime/SlateCore/Private/Debugging/ConsoleSlateDebuggerPaint.h
+++ b/Engine/Source/Runtime/SlateCore/Private/Debugging/ConsoleSlateDebuggerPaint.h
@@ -56,7 +56,8 @@ private:
 	FAutoConsoleCommand ToggleWidgetsNameListCommand;
 	FAutoConsoleVariableRef MaxNumberOfWidgetInListRefCVar;
 	FAutoConsoleVariableRef LogWarningIfWidgetIsPaintedMoreThanOnceRefCVar;
-	FAutoConsoleVariableRef OnlyGameWindowRefCVar;
+	FAutoConsoleVariableRef OnlyProjectContentRefCVar;
+	FAutoConsoleVariableDeprecated OnlyGameWindowDepCVar;
 	FAutoConsoleVariableRef OnlyInvalidationRootRefCVar;
 	FAutoConsoleVariableRef DrawBorderEnabledRefCVar;
 	FAutoConsoleVariableRef DrawFillEnabledRefCVar;
diff --git a/Engine/Source/Runtime/SlateCore/Private/Debugging/ConsoleSlateDebuggerPassBase.cpp b/Engine/Source/Runtime/SlateCore/Private/Debugging/ConsoleSlateDebuggerPassBase.cpp
--- a/Engine/Source/Runtime/SlateCore/Private/Debugging/ConsoleSlateDebuggerPassBase.cpp
+++ b/Engine/Source/Runtime/SlateCore/Private/Debugging/ConsoleSlateDebuggerPassBase.cpp
@@ -22,14 +22,13 @@ FConsoleSlateDebuggerPassBase::FConsoleSlateDebuggerPassBase()
 	, bDrawBorder(true)
 	, bLogWidgetName(false)
 	, bLogWidgetNameOnce(false)
-	, bDebugGameWindowOnly(true)
+	, bDebugProjectContentOnly(true)
 	, bInvalidationRootOnly(false)
 	, MostRecentColor(0.75f, 0.0f, 0.0f, 0.02f)
 	, LeastRecentColor(0.0f, 0.75f, 0.0f, 0.5f)
 	, DrawWidgetNameColor(FColorList::SpicyPink)
 	, MaxNumberOfWidgetInList(20)
 	, FadeDuration(2.0f)
-	, PIEWindowTag("PIEWindow")
 {
 }
 
@@ -37,10 +36,24 @@ FConsoleSlateDebuggerPassBase::~FConsoleSlateDebuggerPassBase()
 {
 }
 
+const FConsoleSlateDebuggerPassBase::TUpdatedWidgetMap& FConsoleSlateDebuggerPassBase::GetWidgets()
+{
+	return Widgets;
+}
+
 void FConsoleSlateDebuggerPassBase::LoadConfig()
 {
 	const FString Section = *GetConfigSection();
-	GConfig->GetBool(*Section, TEXT("bDebugGameWindowOnly"), bDebugGameWindowOnly, *GEditorPerProjectIni);
+
+	// UE_DEPRECATED(5.8)
+	FString Text;
+	if (GConfig->GetString(*Section, TEXT("bDebugGameWindowOnly"), Text, *GEditorPerProjectIni))
+	{
+		UE_LOGF(LogSlateDebugger, Warning, "[%S] bDebugGameWindowOnly is deprecated, use bDebugProjectContentOnly instead", *Section);
+		bDebugProjectContentOnly = FCString::ToBool(*Text);
+	}
+
+	GConfig->GetBool(*Section, TEXT("bDebugProjectContentOnly"), bDebugProjectContentOnly, *GEditorPerProjectIni);
 	GConfig->GetBool(*Section, TEXT("bInvalidationRootOnly"), bInvalidationRootOnly, *GEditorPerProjectIni);
 	GConfig->GetBool(*Section, TEXT("bDisplayWidgetsNameList"), bDisplayWidgetsNameList, *GEditorPerProjectIni);
 	GConfig->GetBool(*Section, TEXT("bUseWidgetPathAsName"), bUseWidgetPathAsName, *GEditorPerProjectIni);
@@ -67,7 +80,8 @@ void FConsoleSlateDebuggerPassBase::LoadConfig()
 void FConsoleSlateDebuggerPassBase::SaveConfig()
 {
 	const FString Section = *GetConfigSection();
-	GConfig->SetBool(*Section, TEXT("bDebugGameWindowOnly"), bDebugGameWindowOnly, *GEditorPerProjectIni);
+	GConfig->SetBool(*Section, TEXT("bDebugProjectContentOnly"), bDebugProjectContentOnly, *GEditorPerProjectIni);
+	GConfig->RemoveKeyFromSection(*Section, TEXT("bDebugGameWindowOnly"), *GEditorPerProjectIni); // UE_DEPRECATED(5.8)
 	GConfig->SetBool(*Section, TEXT("bInvalidationRootOnly"), bInvalidationRootOnly, *GEditorPerProjectIni);
 	GConfig->SetBool(*Section, TEXT("bDisplayWidgetsNameList"), bDisplayWidgetsNameList, *GEditorPerProjectIni);
 	GConfig->SetBool(*Section, TEXT("bUseWidgetPathAsName"), bUseWidgetPathAsName, *GEditorPerProjectIni);
@@ -179,11 +193,6 @@ void FConsoleSlateDebuggerPassBase::HandleEndFrame()
 
 void FConsoleSlateDebuggerPassBase::HandlePaintDebugInfo(const FPaintArgs& InArgs, const FGeometry& InAllottedGeometry, FSlateWindowElementList& InOutDrawElements, int32& InOutLayerId)
 {
-	if (bDebugGameWindowOnly && (InOutDrawElements.GetPaintWindow()->GetType() != EWindowType::GameWindow && InOutDrawElements.GetPaintWindow()->GetTag() != PIEWindowTag))
-	{
-		return;
-	}
-
 	++InOutLayerId;
 
 	FConsoleSlateDebuggerUtility::TSWindowId PaintWindow = FConsoleSlateDebuggerUtility::GetId(InOutDrawElements.GetPaintWindow());
@@ -313,13 +322,10 @@ void FConsoleSlateDebuggerPassBase::BuildInitialWidgetList()
 	for (TArray<TSharedRef<SWindow>>::TIterator It(SlateWindows); It; ++It)
 	{
 		const TSharedRef<SWindow>& Window = *It;
-		
-		if (!(bDebugGameWindowOnly && Window->GetType() != EWindowType::GameWindow && Window->GetTag() != PIEWindowTag))
-		{
-			const FConsoleSlateDebuggerUtility::TSWindowId WindowId = FConsoleSlateDebuggerUtility::GetId(*Window);
-			AddInitialVisibleWidget(*StaticCastSharedRef<SWidget>(Window), WindowId);
-		}
-		
+
+		const FConsoleSlateDebuggerUtility::TSWindowId WindowId = FConsoleSlateDebuggerUtility::GetId(*Window);
+		AddInitialVisibleWidget(*StaticCastSharedRef<SWidget>(Window), WindowId);
+
 		SlateWindows.Append(Window->GetChildWindows());
 	}
 }
@@ -331,8 +337,11 @@ void FConsoleSlateDebuggerPassBase::AddInitialVisibleWidget(const SWidget& Widge
 		return;
 	}
 
-	const bool bIncrementUpdateCount = false; // This is called from BuildInitialWidgetList and we therefore do not want to record an update event
-	AddUpdatedWidget(Widget, WindowId, bIncrementUpdateCount);
+	if (!(bDebugProjectContentOnly && !Widget.IsProjectContent()))
+	{
+		constexpr bool bIncrementUpdateCount = false; // This is called from BuildInitialWidgetList, we therefore do not want to record an update event
+		AddUpdatedWidget(Widget, WindowId, bIncrementUpdateCount);
+	}
 
 	const FChildren* Children = const_cast<SWidget&>(Widget).Debug_GetChildrenForReflector();
 	if (ensure(Children))
diff --git a/Engine/Source/Runtime/SlateCore/Private/Debugging/ConsoleSlateDebuggerPassBase.h b/Engine/Source/Runtime/SlateCore/Private/Debugging/ConsoleSlateDebuggerPassBase.h
--- a/Engine/Source/Runtime/SlateCore/Private/Debugging/ConsoleSlateDebuggerPassBase.h
+++ b/Engine/Source/Runtime/SlateCore/Private/Debugging/ConsoleSlateDebuggerPassBase.h
@@ -16,6 +16,22 @@
 class FConsoleSlateDebuggerPassBase
 {
 public:
+	struct FWidgetInfo 
+	{
+		FConsoleSlateDebuggerUtility::TSWindowId Window = FConsoleSlateDebuggerUtility::InvalidWindowId;
+		TWeakPtr<const SWidget> Widget; // used to check if the widget has been destroyed
+		FVector2f PaintLocation;
+		FVector2f PaintSize;
+		FString WidgetName;
+
+		uint32 LastUpdatedFrame = 0;
+		double LastUpdatedTime = 0.0;
+		int32 UpdateCount = 0;
+		bool bInInvalidationRoot = false;
+	};
+
+	using TUpdatedWidgetMap = TMap<FConsoleSlateDebuggerUtility::TSWidgetId, FWidgetInfo>;
+
 	FConsoleSlateDebuggerPassBase();
 	virtual ~FConsoleSlateDebuggerPassBase();
 
@@ -26,6 +42,8 @@ public:
 	
 	bool IsEnabled() const { return bEnabled; }
 
+	const TUpdatedWidgetMap& GetWidgets();
+
 	virtual void LoadConfig();
 	virtual void SaveConfig();
 
@@ -48,20 +66,6 @@ protected:
 	 * @param bIncrementUpdateCount If true, the UpdateCount of the widget should be increased. This can ve false when building the initial widget list
 	 */
 	virtual void AddUpdatedWidget(const SWidget& Widget, const FConsoleSlateDebuggerUtility::TSWindowId WindowId, bool bIncrementUpdateCount = true) = 0;
-	
-	struct FWidgetInfo 
-	{
-		FConsoleSlateDebuggerUtility::TSWindowId Window = FConsoleSlateDebuggerUtility::InvalidWindowId;
-		TWeakPtr<const SWidget> Widget; // used to check if the widget has been destroyed
-		FVector2f PaintLocation;
-		FVector2f PaintSize;
-		FString WidgetName;
-
-		bool bInInvalidationRoot = false;
-		uint32 LastUpdatedFrame = 0;
-		double LastUpdatedTime = 0.0;
-		int32 UpdateCount = 0;
-	};
 
 	/**
 	 * The internal function that should be called from AddUpdatedWidget to update the internal widget list.
@@ -103,18 +107,16 @@ protected:
 	bool bDrawBorder;
 	bool bLogWidgetName;
 	bool bLogWidgetNameOnce;
-	bool bDebugGameWindowOnly;
+	bool bDebugProjectContentOnly;
 	bool bInvalidationRootOnly;
 	FLinearColor MostRecentColor;
 	FLinearColor LeastRecentColor;
 	FLinearColor DrawWidgetNameColor;
 	int32 MaxNumberOfWidgetInList;
 	float FadeDuration;
-	const FName PIEWindowTag;
 
 	FVector2f WidgetLogLocation{10.f, 10.f};
 
-	using TUpdatedWidgetMap = TMap<FConsoleSlateDebuggerUtility::TSWidgetId, FWidgetInfo>;
 	TUpdatedWidgetMap Widgets;
 };
 
diff --git a/Engine/Source/Runtime/SlateCore/Private/Debugging/ConsoleSlateDebuggerPrepass.cpp b/Engine/Source/Runtime/SlateCore/Private/Debugging/ConsoleSlateDebuggerPrepass.cpp
--- a/Engine/Source/Runtime/SlateCore/Private/Debugging/ConsoleSlateDebuggerPrepass.cpp
+++ b/Engine/Source/Runtime/SlateCore/Private/Debugging/ConsoleSlateDebuggerPrepass.cpp
@@ -46,11 +46,16 @@ FConsoleSlateDebuggerPrepass::FConsoleSlateDebuggerPrepass()
 		, MaxNumberOfWidgetInList
 		, TEXT("The max number of widgets that will be displayed when DisplayWidgetNameList is active.")
 		, FConsoleVariableDelegate::CreateRaw(this, &FConsoleSlateDebuggerPrepass::SaveConfigOnVariableChanged))
-	, OnlyGameWindowRefCVar(
-		TEXT("SlateDebugger.Prepass.OnlyGameWindow")
-		, bDebugGameWindowOnly
-		, TEXT("Option to only debug the game window")
+	, OnlyProjectContentRefCVar(
+		TEXT("SlateDebugger.Prepass.OnlyProjectContent")
+		, bDebugProjectContentOnly
+		, TEXT("Only debug the project's contents (exclude editor, chrome, diagnostic output)")
 		, FConsoleVariableDelegate::CreateRaw(this, &FConsoleSlateDebuggerPrepass::RebuildWidgetListOnVariableChanged))
+	, OnlyGameWindowDepCVar(
+		TEXT("SlateDebugger.Prepass.OnlyGameWindow"),
+		TEXT("SlateDebugger.Prepass.OnlyProjectContent"),
+		TEXT("5.8"),
+		EShadowCVarBehavior::Warn)
 	, OnlyInvalidationRootRefCVar(
 		TEXT("SlateDebugger.Prepass.OnlyInvalidationRoot")
 		, bInvalidationRootOnly
@@ -80,7 +85,7 @@ void FConsoleSlateDebuggerPrepass::LoadConfig()
 {
 	FConsoleSlateDebuggerPassBase::LoadConfig();
 	
-	OnlyGameWindowRefCVar->Set(bDebugGameWindowOnly, ECVF_SetByProjectSetting);
+	OnlyProjectContentRefCVar->Set(bDebugProjectContentOnly, ECVF_SetByProjectSetting);
 	OnlyInvalidationRootRefCVar->Set(bInvalidationRootOnly, ECVF_SetByProjectSetting);
 	EnableWidgetsNameListRefCVar->Set(bDisplayWidgetsNameList, ECVF_SetByProjectSetting);
 	DrawFillEnabledRefCVar->Set(bDrawBox, ECVF_SetByProjectSetting);
@@ -113,9 +118,9 @@ void FConsoleSlateDebuggerPrepass::AddUpdatedWidget(const SWidget& Widget, const
 
 void FConsoleSlateDebuggerPrepass::HandleEndWidgetPrepass(const SWidget* Widget)
 {
-	// Exclude all windows but the game window
 	const TSharedPtr<SWindow> WindowToDrawIn = FSlateApplicationBase::Get().FindWidgetWindow(Widget->AsShared());
-	if (!WindowToDrawIn || (bDebugGameWindowOnly && (WindowToDrawIn->GetType() != EWindowType::GameWindow && WindowToDrawIn->GetTag() != PIEWindowTag)))
+
+	if (!WindowToDrawIn || (bDebugProjectContentOnly && !Widget->IsProjectContent()))
 	{
 		return;
 	}
diff --git a/Engine/Source/Runtime/SlateCore/Private/Debugging/ConsoleSlateDebuggerPrepass.h b/Engine/Source/Runtime/SlateCore/Private/Debugging/ConsoleSlateDebuggerPrepass.h
--- a/Engine/Source/Runtime/SlateCore/Private/Debugging/ConsoleSlateDebuggerPrepass.h
+++ b/Engine/Source/Runtime/SlateCore/Private/Debugging/ConsoleSlateDebuggerPrepass.h
@@ -52,7 +52,8 @@ private:
 	FAutoConsoleVariableRef EnableWidgetsNameListRefCVar;
 	FAutoConsoleCommand ToggleWidgetsNameListCommand;
 	FAutoConsoleVariableRef MaxNumberOfWidgetInListRefCVar;
-	FAutoConsoleVariableRef OnlyGameWindowRefCVar;
+	FAutoConsoleVariableRef OnlyProjectContentRefCVar;
+	FAutoConsoleVariableDeprecated OnlyGameWindowDepCVar;
 	FAutoConsoleVariableRef OnlyInvalidationRootRefCVar;
 	FAutoConsoleVariableRef DrawBorderEnabledRefCVar;
 	FAutoConsoleVariableRef DrawFillEnabledRefCVar;
diff --git a/Engine/Source/Runtime/SlateCore/Private/Debugging/ConsoleSlateDebuggerUpdate.cpp b/Engine/Source/Runtime/SlateCore/Private/Debugging/ConsoleSlateDebuggerUpdate.cpp
--- a/Engine/Source/Runtime/SlateCore/Private/Debugging/ConsoleSlateDebuggerUpdate.cpp
+++ b/Engine/Source/Runtime/SlateCore/Private/Debugging/ConsoleSlateDebuggerUpdate.cpp
@@ -26,7 +26,7 @@ FConsoleSlateDebuggerUpdate::FConsoleSlateDebuggerUpdate()
 	, bDisplayUpdateFromPaint(false)
 	, bShowLegend(false)
 	, bShowQuad(false)
-	, bDebugGameWindowOnly(true)
+	, bDebugProjectContentOnly(true)
 	, WidgetUpdateFlagsFilter(EWidgetUpdateFlags::AnyUpdate)
 	, DrawVolatilePaintColor(FColorList::Red)
 	, DrawRepaintColor(FColorList::Yellow)
@@ -36,7 +36,6 @@ FConsoleSlateDebuggerUpdate::FConsoleSlateDebuggerUpdate()
 	, MaxNumberOfWidgetInList(20)
 	, InvalidationRootIdFilter(-1)
 	, CacheDuration(2.0f)
-	, PIEWindowTag("PIEWindow")
 	, StartCommand(
 		TEXT("SlateDebugger.Update.Start"),
 		TEXT("Start the update widget debug tool. It shows when widgets are updated."),
@@ -70,10 +69,15 @@ FConsoleSlateDebuggerUpdate::FConsoleSlateDebuggerUpdate()
 		TEXT("SlateDebugger.Update.SetInvalidationRootIdFilter"),
 		InvalidationRootIdFilter,
 		TEXT("Option to show only the widgets that are part of an invalidation root."))
-	, OnlyGameWindow(
+	, OnlyProjectContentRefCVar(
+		TEXT("SlateDebugger.Update.OnlyProjectContent"),
+		bDebugProjectContentOnly,
+		TEXT("Only debug content widgets (exclude editor, chrome, diagnostic output)"))
+	, OnlyGameWindowDepCVar(
 		TEXT("SlateDebugger.Update.OnlyGameWindow"),
-		bDebugGameWindowOnly,
-		TEXT("Option to only the debug the game window"))
+		TEXT("SlateDebugger.Update.OnlyProjectContent"),
+		TEXT("5.8"),
+		EShadowCVarBehavior::Warn)
 {
 	GConfig->GetBool(TEXT("SlateDebugger.Update"), TEXT("bDisplayWidgetsNameList"), bDisplayWidgetsNameList, *GEditorPerProjectIni);
 	GConfig->GetBool(TEXT("SlateDebugger.Update"), TEXT("bUseWidgetPathAsName"), bUseWidgetPathAsName, *GEditorPerProjectIni);
@@ -305,6 +309,7 @@ FConsoleSlateDebuggerUpdate::FWidgetInfo::FWidgetInfo(const SWidget* Widget, EWi
 	: WindowId(0)
 	, UpdateFlags(InUpdateFlags)
 	, LastInvalidationTime(FSlateApplicationBase::Get().GetCurrentTime())
+	, bIsProjectContent(Widget->IsProjectContent())
 {
 	PaintLocation = Widget->GetPersistentState().AllottedGeometry.GetAbsolutePosition();
 	PaintSize = Widget->GetPersistentState().AllottedGeometry.GetAbsoluteSize();
@@ -361,13 +366,6 @@ void FConsoleSlateDebuggerUpdate::HandlePaintDebugInfo(const FPaintArgs& InArgs,
 
 	++InOutLayerId;
 
-	// Exclude all windows but the game window
-	SWindow* WindowToDrawIn = InOutDrawElements.GetPaintWindow();
-	if (bDebugGameWindowOnly && (WindowToDrawIn->GetType() != EWindowType::GameWindow && WindowToDrawIn->GetTag() != PIEWindowTag))
-	{
-		return;
-	}
-
 	const FConsoleSlateDebuggerUtility::TSWindowId PaintWindow = FConsoleSlateDebuggerUtility::GetId(InOutDrawElements.GetPaintWindow());
 	int32 NumberOfWidget = 0;
 	TArray<const FString*, TInlineAllocator<100>> NamesToDisplay;
@@ -404,6 +402,11 @@ void FConsoleSlateDebuggerUpdate::HandlePaintDebugInfo(const FPaintArgs& InArgs,
 	{
 		if (Itt.Value.WindowId == PaintWindow)
 		{
+			if (bDebugProjectContentOnly && !Itt.Value.bIsProjectContent)
+			{
+				continue;
+			}
+
 			const float LerpValue = FMath::Clamp((float)(SlateApplicationCurrentTime - Itt.Value.LastInvalidationTime) / CacheDuration, 0.f, 1.f);
 			const FGeometry Geometry = FGeometry::MakeRoot(Itt.Value.PaintSize, FSlateLayoutTransform(1.f, Itt.Value.PaintLocation));
 			const FPaintGeometry PaintGeometry = Geometry.ToPaintGeometry();
diff --git a/Engine/Source/Runtime/SlateCore/Private/Debugging/ConsoleSlateDebuggerUpdate.h b/Engine/Source/Runtime/SlateCore/Private/Debugging/ConsoleSlateDebuggerUpdate.h
--- a/Engine/Source/Runtime/SlateCore/Private/Debugging/ConsoleSlateDebuggerUpdate.h
+++ b/Engine/Source/Runtime/SlateCore/Private/Debugging/ConsoleSlateDebuggerUpdate.h
@@ -37,6 +37,7 @@ public:
 	bool IsEnabled() const { return bEnabled; }
 
 	void ToggleDisplayLegend();
+	bool IsLegendShown() const { return bShowLegend; }
 	void ToogleDisplayWidgetNameList();
 	void ToogleDisplayUpdateFromPaint();
 
@@ -60,7 +61,7 @@ private:
 	bool bDisplayUpdateFromPaint;
 	bool bShowLegend;
 	bool bShowQuad;
-	bool bDebugGameWindowOnly;
+	bool bDebugProjectContentOnly;
 	EWidgetUpdateFlags WidgetUpdateFlagsFilter;
 	FLinearColor DrawVolatilePaintColor;
 	FLinearColor DrawRepaintColor;
@@ -70,7 +71,6 @@ private:
 	int32 MaxNumberOfWidgetInList;
 	int32 InvalidationRootIdFilter;
 	float CacheDuration;
-	FName PIEWindowTag;
 
 	//~ Console objects
 	FAutoConsoleCommand StartCommand;
@@ -81,7 +81,8 @@ private:
 	FAutoConsoleCommand ToogleDisplayUpdateFromPaintCommand;
 	FAutoConsoleCommand SetWidgetUpdateFlagsFilterCommand;
 	FAutoConsoleVariableRef InvalidationRootFilterRefCVar;
-	FAutoConsoleVariableRef OnlyGameWindow;
+	FAutoConsoleVariableRef OnlyProjectContentRefCVar;
+	FAutoConsoleVariableDeprecated OnlyGameWindowDepCVar;
 
 	struct FWidgetInfo
 	{
@@ -94,6 +95,7 @@ private:
 		FString WidgetName;
 		EWidgetUpdateFlags UpdateFlags;
 		double LastInvalidationTime;
+		bool bIsProjectContent;
 	};
 
 	using TWidgetMap = TMap<FConsoleSlateDebuggerUtility::TSWidgetId, FWidgetInfo>;
diff --git a/Engine/Source/Runtime/SlateCore/Private/FastUpdate/WidgetProxy.cpp b/Engine/Source/Runtime/SlateCore/Private/FastUpdate/WidgetProxy.cpp
--- a/Engine/Source/Runtime/SlateCore/Private/FastUpdate/WidgetProxy.cpp
+++ b/Engine/Source/Runtime/SlateCore/Private/FastUpdate/WidgetProxy.cpp
@@ -297,7 +297,7 @@ FWidgetProxy::FUpdateResult FWidgetProxy::Repaint(const FPaintArgs& PaintArgs, F
 		OutDrawElements.PopPixelSnappingMethod();
 	}
 
-	OutDrawElements.SetIsInGameLayer(MyState.bIsInGameLayer);
+	OutDrawElements.SetIsInProjectContent(MyState.IsProjectContent());
 	return FUpdateResult{ PrevLayerId, NewLayerId };
 }
 
diff --git a/Engine/Source/Runtime/SlateCore/Private/Rendering/DrawElementTypes.cpp b/Engine/Source/Runtime/SlateCore/Private/Rendering/DrawElementTypes.cpp
--- a/Engine/Source/Runtime/SlateCore/Private/Rendering/DrawElementTypes.cpp
+++ b/Engine/Source/Runtime/SlateCore/Private/Rendering/DrawElementTypes.cpp
@@ -35,7 +35,7 @@ static void HandleApplyDisabledEffectOnWidgetsToggled(IConsoleVariable* CVar)
 static FAutoConsoleVariableRef CVarApplyDisabledEffectOnWidgets(
 	TEXT("Slate.ApplyDisabledEffectOnWidgets"),
 	bApplyDisabledEffectOnWidgets,
-	TEXT("If true, disabled game-layer widgets will have alpha multiplied by 0.45."),
+	TEXT("If true, disabled widgets will have alpha multiplied by 0.45 (applies to project content only, not editor)."),
 	FConsoleVariableDelegate::CreateStatic(&HandleApplyDisabledEffectOnWidgetsToggled)
 );
 
@@ -153,7 +153,7 @@ static bool ShouldCull(const FSlateWindowElementList& ElementList, const FPaintG
 
 void FSlateDrawElement::Init(FSlateWindowElementList& ElementList, EElementType InElementType, uint32 InLayer, const FPaintGeometry& PaintGeometry, ESlateDrawEffect InDrawEffects)
 {
-	if (ElementList.GetIsInGameLayer() && !bApplyDisabledEffectOnWidgets)
+	if (ElementList.IsInProjectContent() && !bApplyDisabledEffectOnWidgets)
 	{
 		InDrawEffects &= ~ESlateDrawEffect::DisabledEffect;
 	}
diff --git a/Engine/Source/Runtime/SlateCore/Private/Rendering/DrawElements.cpp b/Engine/Source/Runtime/SlateCore/Private/Rendering/DrawElements.cpp
--- a/Engine/Source/Runtime/SlateCore/Private/Rendering/DrawElements.cpp
+++ b/Engine/Source/Runtime/SlateCore/Private/Rendering/DrawElements.cpp
@@ -49,10 +49,10 @@ FSlateWindowElementList::FSlateWindowElementList(const TSharedPtr<SWindow>& InPa
 	, MemManagerAllocatedMemory(0)
 #endif
 	, RenderTargetWindow(nullptr)
+	, bIsProjectContent(false)
 	, bNeedsDeferredResolve(false)
 	, ResolveToDeferredIndex()
 	, WindowSize(FVector2f(0.0f, 0.0f))
-	, bIsInGameLayer(false)
 {
 	if (InPaintWindow.IsValid())
 	{
@@ -67,12 +67,17 @@ FSlateWindowElementList::~FSlateWindowElementList()
 
 void FSlateWindowElementList::SetIsInGameLayer(bool bInGameLayer)
 {
-	bIsInGameLayer = bInGameLayer;
+	bIsProjectContent = bInGameLayer;
 }
 
 bool FSlateWindowElementList::GetIsInGameLayer()
 {
-	return bIsInGameLayer;
+	return bIsProjectContent;
+}
+
+void FSlateWindowElementList::SetIsInProjectContent(bool bIsInProjectContent)
+{
+	bIsProjectContent = bIsInProjectContent;
 }
 
 void FSlateWindowElementList::StartMergeRenderBatches()
diff --git a/Engine/Source/Runtime/SlateCore/Private/SlateCoreModule.cpp b/Engine/Source/Runtime/SlateCore/Private/SlateCoreModule.cpp
--- a/Engine/Source/Runtime/SlateCore/Private/SlateCoreModule.cpp
+++ b/Engine/Source/Runtime/SlateCore/Private/SlateCoreModule.cpp
@@ -1,6 +1,6 @@
 // Copyright Epic Games, Inc. All Rights Reserved.
 
-#include "CoreMinimal.h"
+#include "SlateCoreModule.h"
 #include "Debugging/ConsoleSlateDebugger.h"
 #include "Debugging/ConsoleSlateDebuggerInvalidate.h"
 #include "Debugging/ConsoleSlateDebuggerInvalidationRoot.h"
@@ -8,7 +8,6 @@
 #include "Debugging/ConsoleSlateDebuggerPrepass.h"
 #include "Debugging/ConsoleSlateDebuggerUpdate.h"
 #include "Debugging/ConsoleSlateDebuggerBreak.h"
-#include "Modules/ModuleInterface.h"
 #include "Modules/ModuleManager.h"
 #include "SlateGlobals.h"
 #include "Types/SlateStructs.h"
@@ -18,15 +17,20 @@ DEFINE_LOG_CATEGORY(LogSlateStyles);
 
 const float FOptionalSize::Unspecified = -1.0f;
 
+namespace
+{
+	UE::FSlateCoreModule* Instance;
+}
 
-/**
- * Implements the SlateCore module.
- */
-class FSlateCoreModule
-	: public IModuleInterface
+namespace UE
 {
-public:
-	FSlateCoreModule()
+	FSlateCoreModule& FSlateCoreModule::Get()
+	{
+		check(Instance);
+		return *Instance;
+	}
+
+	FSlateCoreModule::FSlateCoreModule()
 	{
 #if WITH_SLATE_DEBUGGING
 		SlateDebuggerEvent = MakeUnique<FConsoleSlateDebugger>();
@@ -41,17 +45,33 @@ public:
 #endif
 	}
 
+	void FSlateCoreModule::StartupModule()
+	{
+		Instance = this;
+	}
+
+	void FSlateCoreModule::ShutdownModule()
+	{
+		Instance = nullptr;
+	}
+
 #if WITH_SLATE_DEBUGGING
-private:
-	TUniquePtr<FConsoleSlateDebugger> SlateDebuggerEvent;
-	TUniquePtr<FConsoleSlateDebuggerInvalidate> SlateDebuggerInvalidate;
-	TUniquePtr<FConsoleSlateDebuggerInvalidationRoot> SlateDebuggerInvalidationRoot;
-	TUniquePtr<FConsoleSlateDebuggerPaint> SlateDebuggerPaint;
-	TUniquePtr<FConsoleSlateDebuggerPrepass> SlateDebuggerPrepass;
-	TUniquePtr<FConsoleSlateDebuggerUpdate> SlateDebuggerUpdate;
-	TUniquePtr<FConsoleSlateDebuggerBreak> SlateDebuggerBreak;
+	FConsoleSlateDebuggerPaint& FSlateCoreModule::GetPaintDebugger() const
+	{
+		return *SlateDebuggerPaint;
+	}
+
+	FConsoleSlateDebuggerPrepass& FSlateCoreModule::GetPrepassDebugger() const
+	{
+		return *SlateDebuggerPrepass;
+	}
+
+	FConsoleSlateDebuggerUpdate& FSlateCoreModule::GetUpdateDebugger() const
+	{
+		return *SlateDebuggerUpdate;
+	}
 #endif
-};
+}
 
 
-IMPLEMENT_MODULE(FSlateCoreModule, SlateCore);
+IMPLEMENT_MODULE(UE::FSlateCoreModule, SlateCore);
diff --git a/Engine/Source/Runtime/SlateCore/Private/SlateCoreModule.h b/Engine/Source/Runtime/SlateCore/Private/SlateCoreModule.h
new file mode 100644
--- /dev/null
+++ b/Engine/Source/Runtime/SlateCore/Private/SlateCoreModule.h
@@ -0,0 +1,45 @@
+// Copyright Epic Games, Inc. All Rights Reserved.
+
+#pragma once
+
+#include "CoreMinimal.h"
+#include "Debugging/SlateDebugging.h"
+#include "Modules/ModuleInterface.h"
+
+#if WITH_SLATE_DEBUGGING
+class FConsoleSlateDebugger;
+class FConsoleSlateDebuggerInvalidate;
+class FConsoleSlateDebuggerInvalidationRoot;
+class FConsoleSlateDebuggerPaint;
+class FConsoleSlateDebuggerPrepass;
+class FConsoleSlateDebuggerUpdate;
+class FConsoleSlateDebuggerBreak;
+#endif
+
+namespace UE
+{
+	class FSlateCoreModule : public IModuleInterface
+	{
+	public:
+		static FSlateCoreModule& Get();
+
+		FSlateCoreModule();
+		virtual void StartupModule() override;
+		virtual void ShutdownModule() override;
+
+#if WITH_SLATE_DEBUGGING
+		FConsoleSlateDebuggerPaint& GetPaintDebugger() const;
+		FConsoleSlateDebuggerPrepass& GetPrepassDebugger() const;
+		FConsoleSlateDebuggerUpdate& GetUpdateDebugger() const;
+
+	private:
+		TUniquePtr<FConsoleSlateDebugger> SlateDebuggerEvent;
+		TUniquePtr<FConsoleSlateDebuggerInvalidate> SlateDebuggerInvalidate;
+		TUniquePtr<FConsoleSlateDebuggerInvalidationRoot> SlateDebuggerInvalidationRoot;
+		TUniquePtr<FConsoleSlateDebuggerPaint> SlateDebuggerPaint;
+		TUniquePtr<FConsoleSlateDebuggerPrepass> SlateDebuggerPrepass;
+		TUniquePtr<FConsoleSlateDebuggerUpdate> SlateDebuggerUpdate;
+		TUniquePtr<FConsoleSlateDebuggerBreak> SlateDebuggerBreak;
+#endif
+	};
+}
diff --git a/Engine/Source/Runtime/SlateCore/Private/Test/ConsoleSlateDebuggerWindowFilterTest.cpp b/Engine/Source/Runtime/SlateCore/Private/Test/ConsoleSlateDebuggerWindowFilterTest.cpp
new file mode 100644
--- /dev/null
+++ b/Engine/Source/Runtime/SlateCore/Private/Test/ConsoleSlateDebuggerWindowFilterTest.cpp
@@ -0,0 +1,238 @@
+// Copyright Epic Games, Inc. All Rights Reserved.
+
+#include "CoreMinimal.h"
+#include "SlateCoreModule.h"
+#include "Debugging/ConsoleSlateDebuggerPaint.h"
+#include "Debugging/ConsoleSlateDebuggerPrepass.h"
+#include "Debugging/ConsoleSlateDebuggerUpdate.h"
+#include "Misc/AutomationTest.h"
+#include "Widgets/SWindow.h"
+#include "Widgets/SCompoundWidget.h"
+#include "Rendering/DrawElements.h"
+#include "Debugging/SlateDebugging.h"
+#include "HAL/IConsoleManager.h"
+#include "Input/HittestGrid.h"
+#include "Types/PaintArgs.h"
+#include "Misc/App.h"
+
+#if WITH_AUTOMATION_WORKER && WITH_SLATE_DEBUGGING
+
+namespace
+{
+	/**
+	 * Test widget that exposes SetIsProjectContent for testing content-based filtering.
+	 * Simulates the behavior of widgets inside SGameLayerManager (IsProjectContent=true) vs widgets outside it (IsProjectContent=false).
+	 */
+	class SContentTestWidget : public SCompoundWidget
+	{
+	public:
+		SLATE_BEGIN_ARGS(SContentTestWidget) {}
+		SLATE_END_ARGS()
+		void Construct(const FArguments&) {}
+		using SWidget::SetIsProjectContent;
+
+		/** Set AllottedGeometry on the widget's persistent state for testing.
+		 *  Normally set by SWidget::Paint(), but test widgets are never painted by the framework.
+		 */
+		void SetTestAllottedGeometry(const FGeometry& InGeometry)
+		{
+			const_cast<FSlateWidgetPersistentState&>(GetPersistentState()).AllottedGeometry = InGeometry;
+		}
+	};
+
+	int32 CountDrawElements(const FSlateDrawElementMap& ElementMap)
+	{
+		int32 Total = 0;
+		VisitTupleElements([&](const auto& ElementType) { Total += ElementType.Num(); }, ElementMap);
+		return Total;
+	}
+
+	int32 CountPasses(FConsoleSlateDebuggerPassBase& PassDebugger, FConsoleSlateDebuggerUtility::TSWindowId Window)
+	{
+		int32 Total = 0;
+		for (auto& Widget : PassDebugger.GetWidgets())
+		{
+			if (Widget.Value.Window == Window)
+			{
+				++Total;
+			}
+		}
+		return Total;
+	}
+}
+
+namespace UE::Slate::Test
+{
+	constexpr EAutomationTestFlags TestFlags = EAutomationTestFlags::EditorContext | EAutomationTestFlags::ClientContext |
+		EAutomationTestFlags::ProgramContext | EAutomationTestFlags::LowPriority | EAutomationTestFlags::EngineFilter;
+
+	/**
+	 * Integration test for Console Slate Debugger bDebugProjectContentOnly content filtering.
+	 * Verifies that when the filter is enabled, only content widgets (those within SGameLayerManager) are tracked, while non-content widgets are
+	 * filtered out.
+	 */
+	IMPLEMENT_SIMPLE_AUTOMATION_TEST(FConsoleSlateDebuggerContentFilterTest, "Slate.Debugger.ContentFilter", TestFlags)
+
+	bool FConsoleSlateDebuggerContentFilterTest::RunTest(const FString& Parameters)
+	{
+		// Get console variables to control the debuggers
+		IConsoleVariable* PaintEnableCVar = IConsoleManager::Get().FindConsoleVariable(TEXT("SlateDebugger.Paint.Enable"));
+		IConsoleVariable* PaintOnlyContentCVar = IConsoleManager::Get().FindConsoleVariable(TEXT("SlateDebugger.Paint.OnlyProjectContent"));
+		IConsoleVariable* PrepassEnableCVar = IConsoleManager::Get().FindConsoleVariable(TEXT("SlateDebugger.Prepass.Enable"));
+		IConsoleVariable* PrepassOnlyContentCVar = IConsoleManager::Get().FindConsoleVariable(TEXT("SlateDebugger.Prepass.OnlyProjectContent"));
+		IConsoleVariable* UpdateEnableCVar = IConsoleManager::Get().FindConsoleVariable(TEXT("SlateDebugger.Update.Enable"));
+		IConsoleVariable* UpdateOnlyContentCVar = IConsoleManager::Get().FindConsoleVariable(TEXT("SlateDebugger.Update.OnlyProjectContent"));
+
+		if (!PaintEnableCVar || !PaintOnlyContentCVar || !PrepassEnableCVar || !PrepassOnlyContentCVar
+			|| !UpdateEnableCVar || !UpdateOnlyContentCVar)
+		{
+			AddError(TEXT("Could not find SlateDebugger console variables"));
+			return false;
+		}
+
+		const int32 OriginalUpdateEnabled = UpdateEnableCVar->GetInt();
+		const int32 OriginalUpdateOnlyContent = UpdateOnlyContentCVar->GetInt();
+		const int32 OriginalPaintEnabled = PaintEnableCVar->GetInt();
+		const int32 OriginalPaintOnlyContent = PaintOnlyContentCVar->GetInt();
+		const int32 OriginalPrepassEnabled = PrepassEnableCVar->GetInt();
+		const int32 OriginalPrepassOnlyContent = PrepassOnlyContentCVar->GetInt();
+
+		// Disable legend so Update debugger draw elements come only from tracked widgets,
+		// allowing us to observe the content filter.
+		FSlateCoreModule& Module = FSlateCoreModule::Get();
+		FConsoleSlateDebuggerUpdate& UpdateDebugger = Module.GetUpdateDebugger();
+		const bool bLegendWasEnabled = UpdateDebugger.IsLegendShown();
+		if (bLegendWasEnabled)
+		{
+			UpdateDebugger.ToggleDisplayLegend();
+		}
+
+		auto TestContentFilter = [&](bool bWidgetIsProjectContent, bool bOnlyContent, bool bExpectRecorded, const TCHAR* Description)
+		{
+			TSharedPtr<SContentTestWidget> TestWidget;
+			TSharedRef<SWindow> TestWindow = SNew(SWindow)
+				.ClientSize(FVector2f(100, 100))
+				.CreateTitleBar(false)
+				.SizingRule(ESizingRule::FixedSize)
+				[
+					SAssignNew(TestWidget, SContentTestWidget)
+				];
+
+			if (bWidgetIsProjectContent)
+			{
+				TestWidget->SetIsProjectContent(true);
+			}
+
+			// Give the widget non-zero geometry so that FSlateDrawElement::MakeBox does not cull it.
+			TestWidget->SetTestAllottedGeometry(FGeometry::MakeRoot(FVector2f(50, 50), FSlateLayoutTransform()));
+
+			FSlateWindowElementList ElementList(TestWindow);
+			const FConsoleSlateDebuggerUtility::TSWindowId WindowId = FConsoleSlateDebuggerUtility::GetId(TestWindow.Get());
+
+			// Paint
+			PaintOnlyContentCVar->Set(bOnlyContent ? 1 : 0, ECVF_SetByTemp);
+			PaintEnableCVar->Set(1);
+
+			FConsoleSlateDebuggerPaint& PaintDebugger = Module.GetPaintDebugger();
+
+			const int32 PaintsBefore = CountPasses(PaintDebugger, WindowId);
+			FSlateDebugging::EndWidgetPaint.Broadcast(TestWidget.Get(), ElementList, 0);
+			const int32 PaintsAfter = CountPasses(PaintDebugger, WindowId);
+			const bool bPaintsAdded = PaintsAfter > PaintsBefore;
+
+			if (bExpectRecorded && !bPaintsAdded)
+			{
+				AddError(FString::Printf(TEXT("%s: Expected debug paints to be added (widget should pass filter), but none were"), Description));
+			}
+			else if (!bExpectRecorded && bPaintsAdded)
+			{
+				AddError(FString::Printf(TEXT("%s: Expected widget to be filtered (no paints), but %d paints were added"),
+					Description, PaintsAfter - PaintsBefore));
+			}
+
+			PaintEnableCVar->Set(0);
+			PaintOnlyContentCVar->Unset(ECVF_SetByTemp);
+
+			// Prepass
+			PrepassOnlyContentCVar->Set(bOnlyContent ? 1 : 0, ECVF_SetByTemp);
+			PrepassEnableCVar->Set(1);
+
+			FConsoleSlateDebuggerPrepass& PrepassDebugger = Module.GetPrepassDebugger();
+
+			const int32 PassesBefore = CountPasses(PrepassDebugger, WindowId);
+			FSlateDebugging::EndWidgetPrepass.Broadcast(TestWidget.Get());
+			const int32 PassesAfter = CountPasses(PrepassDebugger, WindowId);
+			const bool bPassesAdded = PassesAfter > PassesBefore;
+
+			if (bExpectRecorded && !bPassesAdded)
+			{
+				AddError(FString::Printf(TEXT("%s: Expected debug passes to be added (widget should pass filter), but none were"), Description));
+			}
+			else if (!bExpectRecorded && bPassesAdded)
+			{
+				AddError(FString::Printf(TEXT("%s: Expected widget to be filtered (no passes), but %d passes were added"),
+					Description, PassesAfter - PassesBefore));
+			}
+
+			PrepassEnableCVar->Set(0);
+			PrepassOnlyContentCVar->Unset(ECVF_SetByTemp);
+
+			// Update
+			UpdateOnlyContentCVar->Set(bOnlyContent ? 1 : 0, ECVF_SetByTemp);
+			UpdateEnableCVar->Set(1);
+
+			const int32 ElementsBefore = CountDrawElements(ElementList.GetUncachedDrawElements());
+			FSlateDebugging::BroadcastWidgetUpdated(TestWidget.Get(), EWidgetUpdateFlags::NeedsRepaint);
+			FHittestGrid HittestGrid;
+			FPaintArgs PaintArgs(&TestWindow.Get(), HittestGrid, FVector2f::ZeroVector, FApp::GetCurrentTime(), FApp::GetDeltaTime());
+			FGeometry Geometry = FGeometry::MakeRoot(FVector2f(100, 100), FSlateLayoutTransform());
+			int32 LayerId = 0;
+			FSlateDebugging::PaintDebugElements.Broadcast(PaintArgs, Geometry, ElementList, LayerId);
+			const int32 ElementsAfter = CountDrawElements(ElementList.GetUncachedDrawElements());
+			const bool bElementsAdded = ElementsAfter > ElementsBefore;
+
+			if (bExpectRecorded && !bElementsAdded)
+			{
+				AddError(FString::Printf(TEXT("%s: Expected debug elements to be added (widget should pass filter), but none were"), Description));
+			}
+			else if (!bExpectRecorded && bElementsAdded)
+			{
+				AddError(FString::Printf(TEXT("%s: Expected widget to be filtered (no elements), but %d elements were added"),
+					Description, ElementsAfter - ElementsBefore));
+			}
+
+			// Disable debugger between tests to reset state
+			UpdateEnableCVar->Set(0);
+			UpdateOnlyContentCVar->Unset(ECVF_SetByTemp);
+		};
+
+		// Content widgets (simulating inside SGameLayerManager) should pass through filter
+		TestContentFilter(/* bWidgetIsContent */ true, /* bOnlyContent */ true, /* bExpect */ true,
+			TEXT("Content widget with filter=true"));
+		// Non-content widgets (simulating outside SGameLayerManager) should be filtered out
+		TestContentFilter(/* bWidgetIsContent */ false, /* bOnlyContent */ true, /* bExpect */ false,
+			TEXT("Non-content widget with filter=true"));
+
+		// With filter disabled, all widgets should pass through regardless of content status
+		TestContentFilter(/* bWidgetIsContent */ true, /* bOnlyContent */ false, /* bExpect */ true,
+			TEXT("Content widget with filter=false"));
+		TestContentFilter(/* bWidgetIsContent */ false, /* bOnlyContent */ false, /* bExpect */ true,
+			TEXT("Non-content widget with filter=false"));
+
+		// Restore original settings
+		if (bLegendWasEnabled)
+		{
+			UpdateDebugger.ToggleDisplayLegend();
+		}
+		UpdateEnableCVar->Set(OriginalUpdateEnabled);
+		UpdateOnlyContentCVar->Set(OriginalUpdateOnlyContent);
+		PaintEnableCVar->Set(OriginalPaintEnabled);
+		PaintOnlyContentCVar->Set(OriginalPaintOnlyContent);
+		PrepassEnableCVar->Set(OriginalPrepassEnabled);
+		PrepassOnlyContentCVar->Set(OriginalPrepassOnlyContent);
+
+		return true;
+	}
+}
+
+#endif // WITH_AUTOMATION_WORKER && WITH_SLATE_DEBUGGING
diff --git a/Engine/Source/Runtime/SlateCore/Private/Widgets/SWidget.cpp b/Engine/Source/Runtime/SlateCore/Private/Widgets/SWidget.cpp
--- a/Engine/Source/Runtime/SlateCore/Private/Widgets/SWidget.cpp
+++ b/Engine/Source/Runtime/SlateCore/Private/Widgets/SWidget.cpp
@@ -228,6 +228,8 @@ SWidget::SWidget()
 	, bHasActiveTimers(false)
 	, bDesiredSizeSet(false)
 	, bPrepassLayoutScaleMultiplierSet(false)
+	, bIsProjectContent(false)
+	, bIsProjectContentParent(false)
 	, bHasCustomPrepass(false)
 	, bHasRelativeLayoutScale(false)
 	, bVolatilityAlwaysInvalidatesPrepass(false)
@@ -1519,7 +1521,7 @@ int32 SWidget::Paint(const FPaintArgs& Args, const FGeometry& AllottedGeometry,
 	PersistentState.IncomingUserIndex = (int8)IncomingUserIndex;
 
 	PersistentState.IncomingFlowDirection = GSlateFlowDirection;
-	PersistentState.bIsInGameLayer = OutDrawElements.GetIsInGameLayer();
+	PersistentState.SetIsProjectContent(IsProjectContent());
 
 	FPaintArgs UpdatedArgs = Args.WithNewParent(this);
 	UpdatedArgs.SetInheritedHittestability(bOutgoingHittestability);
@@ -1607,8 +1609,10 @@ int32 SWidget::Paint(const FPaintArgs& Args, const FGeometry& AllottedGeometry,
 		LLM_SCOPE_DYNAMIC_STAT_OBJECTPATH_FNAME(ClassName, ELLMTagSet::AssetClasses);
 		UE_TRACE_METADATA_SCOPE_ASSET_FNAME(AssetName, ClassName, PackageName);
 #endif
-	// Paint the geometry of this widget.
+		// Paint the geometry of this widget.
+		OutDrawElements.SetIsInProjectContent(IsProjectContent());
 		NewLayerId = OnPaint(UpdatedArgs, AllottedGeometry, CullingBounds, OutDrawElements, LayerId, ContentWidgetStyle, bParentEnabled);
+		OutDrawElements.SetIsInProjectContent(false);
 	}
 
 	// Just repainted
@@ -1830,6 +1834,9 @@ void SWidget::Prepass_ChildLoop(float InLayoutScaleMultiplier, FChildren* MyChil
 				? InLayoutScaleMultiplier * Self->GetRelativeLayoutScale(ChildIndex, InLayoutScaleMultiplier)
 				: InLayoutScaleMultiplier;
 
+			// Inherit project-content status by default.  Children can opt out (overwrite) in custom prepass.
+			Child.bIsProjectContent = Self->bIsProjectContent || Self->bIsProjectContentParent;
+
 			// Recur: Descend down the widget tree.
 			Child.Prepass_Internal(ChildLayoutScaleMultiplier);
 		}
diff --git a/Engine/Source/Runtime/SlateCore/Public/FastUpdate/WidgetProxy.h b/Engine/Source/Runtime/SlateCore/Public/FastUpdate/WidgetProxy.h
--- a/Engine/Source/Runtime/SlateCore/Public/FastUpdate/WidgetProxy.h
+++ b/Engine/Source/Runtime/SlateCore/Public/FastUpdate/WidgetProxy.h
@@ -249,7 +249,9 @@ struct FSlateWidgetPersistentState
 		, bParentEnabled(true)
 		, bInheritedHittestability(false)
 		, bDeferredPainting(false)
+PRAGMA_DISABLE_DEPRECATION_WARNINGS
 		, bIsInGameLayer(false)
+PRAGMA_ENABLE_DEPRECATION_WARNINGS
 	{}
 
 	TWeakPtr<SWidget> PaintParent;
@@ -268,8 +270,26 @@ struct FSlateWidgetPersistentState
 	uint8 bParentEnabled : 1;
 	uint8 bInheritedHittestability : 1;
 	uint8 bDeferredPainting : 1;
+
+	UE_DEPRECATED(5.8, "Use methods instead of direct field access")
 	uint8 bIsInGameLayer : 1;
 
+	/** The value of SWidget::IsContent() */
+	bool IsProjectContent() const
+	{
+PRAGMA_DISABLE_DEPRECATION_WARNINGS
+		return bIsInGameLayer;
+PRAGMA_ENABLE_DEPRECATION_WARNINGS
+	}
+
+	/** Only to be called by SWidget::Paint. */
+	void SetIsProjectContent(bool bInIsContent)
+	{
+PRAGMA_DISABLE_DEPRECATION_WARNINGS
+		bIsInGameLayer = bInIsContent;
+PRAGMA_ENABLE_DEPRECATION_WARNINGS
+	}
+
 	static const FSlateWidgetPersistentState NoState;
 };
 
diff --git a/Engine/Source/Runtime/SlateCore/Public/Rendering/DrawElements.h b/Engine/Source/Runtime/SlateCore/Public/Rendering/DrawElements.h
--- a/Engine/Source/Runtime/SlateCore/Public/Rendering/DrawElements.h
+++ b/Engine/Source/Runtime/SlateCore/Public/Rendering/DrawElements.h
@@ -405,9 +405,20 @@ public:
 
 	void AddReferencedObjects(FReferenceCollector& Collector);
 	
+	UE_DEPRECATED(5.8, "Use SetIsProjectContent from CustomPrepass instead")
 	SLATECORE_API void SetIsInGameLayer(bool bInGameLayer);
+
+	UE_DEPRECATED(5.8, "Renamed to IsInProjectContent")
 	SLATECORE_API bool GetIsInGameLayer();
 
+	void SetIsInProjectContent(bool bIsInProjectContent);
+
+	/** Whether the stack frame is drawing the project's content (as opposed to editor, chrome, diagnostic output). */
+	bool IsInProjectContent() const
+	{
+		return bIsProjectContent;
+	}
+
 	TArrayView<FSlateCachedElementList* const> GetCurrentCachedElementWithNewData() const
 	{ 
 		if (GetCurrentCachedElementData() != nullptr)
@@ -482,6 +493,9 @@ private:
 	 */
 	TArray< TSharedPtr<FDeferredPaint> > DeferredPaintList;
 
+	/** Caches whether the current widget constitutes "content" (as opposed to editor, chrome, diagnostic output). */
+	bool bIsProjectContent;
+
 	bool bNeedsDeferredResolve;
 	TArray<int32> ResolveToDeferredIndex;
 
@@ -512,7 +526,4 @@ private:
 
 	/** Store the size of the window being used to paint */
 	FVector2f WindowSize;
-
-	/** Store if currently drawing within Game Layer */
-	bool bIsInGameLayer;
 };
diff --git a/Engine/Source/Runtime/SlateCore/Public/Widgets/SWidget.h b/Engine/Source/Runtime/SlateCore/Public/Widgets/SWidget.h
--- a/Engine/Source/Runtime/SlateCore/Public/Widgets/SWidget.h
+++ b/Engine/Source/Runtime/SlateCore/Public/Widgets/SWidget.h
@@ -721,9 +721,35 @@ public:
 		return EnumHasAnyFlags(UpdateFlags, FlagsToCheck);
 	}
 
+	/** Whether the widget is a part of the project's content (as opposed to editor, chrome, diagnostic output).
+	 * Dynamic: updated during pre-pass, inherited from parent but may be overwritten in CustomPrepass.
+	 * Used to filter which widgets can be navigated to, which are considered by diagnostic tools.
+	 */
+	bool IsProjectContent() const
+	{
+		return bIsProjectContent;
+	}
+
 protected:
 	void SetVolatilePrepass(bool bVolatile) { bVolatile ? AddUpdateFlags(EWidgetUpdateFlags::NeedsVolatilePrepass) : RemoveUpdateFlags(EWidgetUpdateFlags::NeedsVolatilePrepass); }
 
+	/**
+	 * Call from CustomPrepass to change whether this widget and its children are a part of the project's content (as opposed to editor, chrome,
+	 * diagnostic output).
+	 */
+	void SetIsProjectContent(bool bInIsProjectContent)
+	{
+		bIsProjectContent = bInIsProjectContent;
+	}
+
+	/**
+	 * Call from CustomPrepass to change whether this widget's children are a part of the project's content.
+	 */
+	void SetIsProjectContentParent(bool bInIsProjectContent)
+	{
+		bIsProjectContentParent = bInIsProjectContent;
+	}
+
 	virtual bool CustomPrepass(float LayoutScaleMultiplier) { return false; }
 
 	/**
@@ -1869,6 +1895,9 @@ private:
 	uint8 bSupportsInvalidation : 1 = false;
 #endif
 
+	uint8 bIsProjectContent : 1;
+	uint8 bIsProjectContentParent : 1;
+
 protected:
 	uint8 bHasCustomPrepass : 1;
 
