Undo //UE5/Main/Engine/Source/Runtime/... changelist 51486290

#rnx
#Synced-CL 51476348

diff --git a/Engine/Source/Runtime/Engine/Private/Slate/DebugCanvas.cpp b/Engine/Source/Runtime/Engine/Private/Slate/DebugCanvas.cpp
--- a/Engine/Source/Runtime/Engine/Private/Slate/DebugCanvas.cpp
+++ b/Engine/Source/Runtime/Engine/Private/Slate/DebugCanvas.cpp
@@ -371,13 +371,6 @@ SDebugCanvas::SDebugCanvas()
 void SDebugCanvas::Construct(const FArguments& InArgs)
 {
 	SceneViewport = InArgs._SceneViewport;
-	bHasCustomPrepass = true;
-}
-
-bool SDebugCanvas::CustomPrepass(float LayoutScaleMultiplier)
-{
-	SetIsProjectContent(false);
-	return false;
 }
 
 int32 SDebugCanvas::OnPaint(const FPaintArgs& Args, const FGeometry& AllottedGeometry, const FSlateRect& MyCullingRect, FSlateWindowElementList& OutDrawElements, int32 LayerId, const FWidgetStyle& InWidgetStyle, bool bParentEnabled) const
diff --git a/Engine/Source/Runtime/Engine/Private/Slate/SGameLayerManager.cpp b/Engine/Source/Runtime/Engine/Private/Slate/SGameLayerManager.cpp
--- a/Engine/Source/Runtime/Engine/Private/Slate/SGameLayerManager.cpp
+++ b/Engine/Source/Runtime/Engine/Private/Slate/SGameLayerManager.cpp
@@ -101,7 +101,6 @@ SGameLayerManager::SGameLayerManager()
 
 void SGameLayerManager::Construct(const SGameLayerManager::FArguments& InArgs)
 {
-	bHasCustomPrepass = true;
 	SceneViewport = InArgs._SceneViewport;
 
 	// In PIE we should default to per user hit testing being off because developers will need the mouse and keyboard to work for all players
@@ -437,12 +436,14 @@ int32 SGameLayerManager::OnPaint(const FPaintArgs& Args, const FGeometry& Allott
 	}
 #endif
 
+	OutDrawElements.SetIsInGameLayer(true);
 	int32 ResultLayer = SCompoundWidget::OnPaint(Args, AllottedGeometry, MyCullingRect, OutDrawElements, LayerId, InWidgetStyle, bParentEnabled);
 
 #if WITH_SLATE_DEBUGGING
 	ResultLayer = FSlateNavigationDebugger::Get().Paint(OutDrawElements, ResultLayer);
 #endif
 
+	OutDrawElements.SetIsInGameLayer(false);
 	return ResultLayer;
 }
 
@@ -753,12 +754,6 @@ void SGameLayerManager::SetWindowTitleBarVisibility(bool bIsVisible)
 	UpdateWindowTitleBarVisibility();
 }
 
-bool SGameLayerManager::CustomPrepass(float LayoutScaleMultiplier)
-{
-	SetIsProjectContentParent(true);
-	return true;
-}
-
 void SGameLayerManager::UpdateWindowTitleBar()
 {
 	if (WindowTitleBarState.ContentWidget.IsValid())
diff --git a/Engine/Source/Runtime/Engine/Public/Slate/DebugCanvas.h b/Engine/Source/Runtime/Engine/Public/Slate/DebugCanvas.h
--- a/Engine/Source/Runtime/Engine/Public/Slate/DebugCanvas.h
+++ b/Engine/Source/Runtime/Engine/Public/Slate/DebugCanvas.h
@@ -34,7 +34,6 @@ public:
 	ENGINE_API void Construct(const FArguments& InArgs);
 
 	/** SWidget interface */
-	ENGINE_API virtual bool CustomPrepass(float LayoutScaleMultiplier) override;
 	ENGINE_API int32 OnPaint(const FPaintArgs& Args, const FGeometry& AllottedGeometry, const FSlateRect& MyCullingRect, FSlateWindowElementList& OutDrawElements, int32 LayerId, const FWidgetStyle& InWidgetStyle, bool bParentEnabled) const override;
 	ENGINE_API virtual FVector2D ComputeDesiredSize(float LayoutScaleMultiplier) const override;
 
diff --git a/Engine/Source/Runtime/Engine/Public/Slate/SGameLayerManager.h b/Engine/Source/Runtime/Engine/Public/Slate/SGameLayerManager.h
--- a/Engine/Source/Runtime/Engine/Public/Slate/SGameLayerManager.h
+++ b/Engine/Source/Runtime/Engine/Public/Slate/SGameLayerManager.h
@@ -152,15 +152,14 @@ public:
 	ENGINE_API virtual void SetWindowTitleBarVisibility(bool bIsVisible);
 	// End IGameLayerManager
 
+public:
+
 	//~ Begin SWidget overrides
-protected:
-	virtual bool CustomPrepass(float LayoutScaleMultiplier) override;
 	ENGINE_API virtual void Tick(const FGeometry& AllottedGeometry, const double InCurrentTime, const float InDeltaTime) override;
 	ENGINE_API virtual int32 OnPaint(const FPaintArgs& Args, const FGeometry& AllottedGeometry, const FSlateRect& MyCullingRect, FSlateWindowElementList& OutDrawElements, int32 LayerId, const FWidgetStyle& InWidgetStyle, bool bParentEnabled) const override;
 	ENGINE_API virtual bool OnVisualizeTooltip(const TSharedPtr<SWidget>& TooltipContent) override;
 	//~ End SWidget overrides
 
-public:
 	/**
 	 * Function will instruct internal DPI computations to use a provided reference viewport size instead of the actual viewport size.
 	 * After the DPI will be retrieved it will be scaled down with the ratio between the actual viewport size and the provided one.
diff --git a/Engine/Source/Runtime/Slate/Private/Framework/Application/SlateApplication.cpp b/Engine/Source/Runtime/Slate/Private/Framework/Application/SlateApplication.cpp
--- a/Engine/Source/Runtime/Slate/Private/Framework/Application/SlateApplication.cpp
+++ b/Engine/Source/Runtime/Slate/Private/Framework/Application/SlateApplication.cpp
@@ -3613,9 +3613,9 @@ void FSlateApplication::ProcessReply( const FWidgetPath& CurrentEventPath, const
 
 		if (NavigationSource.IsValid())
 		{
-			if (!GSlateEnableGamepadEditorNavigation && TheReply.GetNavigationGenesis() == ENavigationGenesis::Controller && !NavigationSource.GetLastWidget()->GetPersistentState().IsProjectContent())
+			if (!GSlateEnableGamepadEditorNavigation && TheReply.GetNavigationGenesis() == ENavigationGenesis::Controller && !NavigationSource.GetLastWidget()->GetPersistentState().bIsInGameLayer)
 			{
-				// Gamepad navigation while not in project content, do nothing as specified by GSlateEnableGamepadEditorNavigation
+				// Gamepad navigation while not in a game layer, do nothing as specified by GSlateEnableGamepadEditorNavigation
 			}
 			else if (TheReply.GetNavigationDestination().IsValid())
 			{
diff --git a/Engine/Source/Runtime/Slate/Public/Widgets/Input/SNumericEntryBox.h b/Engine/Source/Runtime/Slate/Public/Widgets/Input/SNumericEntryBox.h
--- a/Engine/Source/Runtime/Slate/Public/Widgets/Input/SNumericEntryBox.h
+++ b/Engine/Source/Runtime/Slate/Public/Widgets/Input/SNumericEntryBox.h
@@ -750,7 +750,7 @@ private:
 #if WITH_EDITOR
 		if constexpr (std::is_floating_point_v<NumericType>)
 		{
-			if (!GetPersistentState().IsProjectContent() && !Interface->GetIndicateNearlyInteger())
+			if (!GetPersistentState().bIsInGameLayer && !Interface->GetIndicateNearlyInteger())
 			{
 				Interface->SetIndicateNearlyInteger(true);
 				bCachedValueStringDirty = true;
diff --git a/Engine/Source/Runtime/SlateCore/Private/Debugging/ConsoleSlateDebuggerPaint.cpp b/Engine/Source/Runtime/SlateCore/Private/Debugging/ConsoleSlateDebuggerPaint.cpp
--- a/Engine/Source/Runtime/SlateCore/Private/Debugging/ConsoleSlateDebuggerPaint.cpp
+++ b/Engine/Source/Runtime/SlateCore/Private/Debugging/ConsoleSlateDebuggerPaint.cpp
@@ -51,16 +51,11 @@ FConsoleSlateDebuggerPaint::FConsoleSlateDebuggerPaint()
 		, bLogWarningIfWidgetIsPaintedMoreThanOnce
 		, TEXT("Option to log a warning if a widget is painted more than once in a single frame.")
 		, FConsoleVariableDelegate::CreateRaw(this, &FConsoleSlateDebuggerPaint::SaveConfigOnVariableChanged))
-	, OnlyProjectContentRefCVar(
-		TEXT("SlateDebugger.Paint.OnlyProjectContent")
-		, bDebugProjectContentOnly
-		, TEXT("Only debug the project's contents (exclude editor, chrome, diagnostic output)")
+	, OnlyGameWindowRefCVar(
+		TEXT("SlateDebugger.Paint.OnlyGameWindow")
+		, bDebugGameWindowOnly
+		, TEXT("Option to only debug the game window")
 		, FConsoleVariableDelegate::CreateRaw(this, &FConsoleSlateDebuggerPaint::RebuildWidgetListOnVariableChanged))
-	, OnlyGameWindowDepCVar(
-		TEXT("SlateDebugger.Paint.OnlyGameWindow"),
-		TEXT("SlateDebugger.Paint.OnlyProjectContent"),
-		TEXT("5.8"),
-		EShadowCVarBehavior::Warn)
 	, OnlyInvalidationRootRefCVar(
 		TEXT("SlateDebugger.Paint.OnlyInvalidationRoot")
 		, bInvalidationRootOnly
@@ -92,7 +87,7 @@ void FConsoleSlateDebuggerPaint::LoadConfig()
 	const FString Section = *GetConfigSection();
 	GConfig->GetBool(*Section, TEXT("bLogWarningIfWidgetIsPaintedMoreThanOnce"), bLogWarningIfWidgetIsPaintedMoreThanOnce, *GEditorPerProjectIni);
 
-	OnlyProjectContentRefCVar->Set(bDebugProjectContentOnly, ECVF_SetByProjectSetting);
+	OnlyGameWindowRefCVar->Set(bDebugGameWindowOnly, ECVF_SetByProjectSetting);
 	OnlyInvalidationRootRefCVar->Set(bInvalidationRootOnly, ECVF_SetByProjectSetting);
 	EnableWidgetsNameListRefCVar->Set(bDisplayWidgetsNameList, ECVF_SetByProjectSetting);
 	DrawFillEnabledRefCVar->Set(bDrawBox, ECVF_SetByProjectSetting);
@@ -138,7 +133,9 @@ void FConsoleSlateDebuggerPaint::AddUpdatedWidget(const SWidget& Widget, const F
 
 void FConsoleSlateDebuggerPaint::HandleEndWidgetPaint(const SWidget* Widget, const FSlateWindowElementList& OutDrawElements, int32 LayerId)
 {
-	if (bDebugProjectContentOnly && !Widget->IsProjectContent())
+	// Exclude all windows but the game window
+	const SWindow* WindowToDrawIn = OutDrawElements.GetPaintWindow();
+	if (bDebugGameWindowOnly && (WindowToDrawIn->GetType() != EWindowType::GameWindow && WindowToDrawIn->GetTag() != PIEWindowTag))
 	{
 		return;
 	}
diff --git a/Engine/Source/Runtime/SlateCore/Private/Debugging/ConsoleSlateDebuggerPaint.h b/Engine/Source/Runtime/SlateCore/Private/Debugging/ConsoleSlateDebuggerPaint.h
--- a/Engine/Source/Runtime/SlateCore/Private/Debugging/ConsoleSlateDebuggerPaint.h
+++ b/Engine/Source/Runtime/SlateCore/Private/Debugging/ConsoleSlateDebuggerPaint.h
@@ -56,8 +56,7 @@ private:
 	FAutoConsoleCommand ToggleWidgetsNameListCommand;
 	FAutoConsoleVariableRef MaxNumberOfWidgetInListRefCVar;
 	FAutoConsoleVariableRef LogWarningIfWidgetIsPaintedMoreThanOnceRefCVar;
-	FAutoConsoleVariableRef OnlyProjectContentRefCVar;
-	FAutoConsoleVariableDeprecated OnlyGameWindowDepCVar;
+	FAutoConsoleVariableRef OnlyGameWindowRefCVar;
 	FAutoConsoleVariableRef OnlyInvalidationRootRefCVar;
 	FAutoConsoleVariableRef DrawBorderEnabledRefCVar;
 	FAutoConsoleVariableRef DrawFillEnabledRefCVar;
diff --git a/Engine/Source/Runtime/SlateCore/Private/Debugging/ConsoleSlateDebuggerPassBase.cpp b/Engine/Source/Runtime/SlateCore/Private/Debugging/ConsoleSlateDebuggerPassBase.cpp
--- a/Engine/Source/Runtime/SlateCore/Private/Debugging/ConsoleSlateDebuggerPassBase.cpp
+++ b/Engine/Source/Runtime/SlateCore/Private/Debugging/ConsoleSlateDebuggerPassBase.cpp
@@ -22,13 +22,14 @@ FConsoleSlateDebuggerPassBase::FConsoleSlateDebuggerPassBase()
 	, bDrawBorder(true)
 	, bLogWidgetName(false)
 	, bLogWidgetNameOnce(false)
-	, bDebugProjectContentOnly(true)
+	, bDebugGameWindowOnly(true)
 	, bInvalidationRootOnly(false)
 	, MostRecentColor(0.75f, 0.0f, 0.0f, 0.02f)
 	, LeastRecentColor(0.0f, 0.75f, 0.0f, 0.5f)
 	, DrawWidgetNameColor(FColorList::SpicyPink)
 	, MaxNumberOfWidgetInList(20)
 	, FadeDuration(2.0f)
+	, PIEWindowTag("PIEWindow")
 {
 }
 
@@ -36,24 +37,10 @@ FConsoleSlateDebuggerPassBase::~FConsoleSlateDebuggerPassBase()
 {
 }
 
-const FConsoleSlateDebuggerPassBase::TUpdatedWidgetMap& FConsoleSlateDebuggerPassBase::GetWidgets()
-{
-	return Widgets;
-}
-
 void FConsoleSlateDebuggerPassBase::LoadConfig()
 {
 	const FString Section = *GetConfigSection();
-
-	// UE_DEPRECATED(5.8)
-	FString Text;
-	if (GConfig->GetString(*Section, TEXT("bDebugGameWindowOnly"), Text, *GEditorPerProjectIni))
-	{
-		UE_LOGF(LogSlateDebugger, Warning, "[%S] bDebugGameWindowOnly is deprecated, use bDebugProjectContentOnly instead", *Section);
-		bDebugProjectContentOnly = FCString::ToBool(*Text);
-	}
-
-	GConfig->GetBool(*Section, TEXT("bDebugProjectContentOnly"), bDebugProjectContentOnly, *GEditorPerProjectIni);
+	GConfig->GetBool(*Section, TEXT("bDebugGameWindowOnly"), bDebugGameWindowOnly, *GEditorPerProjectIni);
 	GConfig->GetBool(*Section, TEXT("bInvalidationRootOnly"), bInvalidationRootOnly, *GEditorPerProjectIni);
 	GConfig->GetBool(*Section, TEXT("bDisplayWidgetsNameList"), bDisplayWidgetsNameList, *GEditorPerProjectIni);
 	GConfig->GetBool(*Section, TEXT("bUseWidgetPathAsName"), bUseWidgetPathAsName, *GEditorPerProjectIni);
@@ -80,8 +67,7 @@ void FConsoleSlateDebuggerPassBase::LoadConfig()
 void FConsoleSlateDebuggerPassBase::SaveConfig()
 {
 	const FString Section = *GetConfigSection();
-	GConfig->SetBool(*Section, TEXT("bDebugProjectContentOnly"), bDebugProjectContentOnly, *GEditorPerProjectIni);
-	GConfig->RemoveKeyFromSection(*Section, TEXT("bDebugGameWindowOnly"), *GEditorPerProjectIni); // UE_DEPRECATED(5.8)
+	GConfig->SetBool(*Section, TEXT("bDebugGameWindowOnly"), bDebugGameWindowOnly, *GEditorPerProjectIni);
 	GConfig->SetBool(*Section, TEXT("bInvalidationRootOnly"), bInvalidationRootOnly, *GEditorPerProjectIni);
 	GConfig->SetBool(*Section, TEXT("bDisplayWidgetsNameList"), bDisplayWidgetsNameList, *GEditorPerProjectIni);
 	GConfig->SetBool(*Section, TEXT("bUseWidgetPathAsName"), bUseWidgetPathAsName, *GEditorPerProjectIni);
@@ -193,6 +179,11 @@ void FConsoleSlateDebuggerPassBase::HandleEndFrame()
 
 void FConsoleSlateDebuggerPassBase::HandlePaintDebugInfo(const FPaintArgs& InArgs, const FGeometry& InAllottedGeometry, FSlateWindowElementList& InOutDrawElements, int32& InOutLayerId)
 {
+	if (bDebugGameWindowOnly && (InOutDrawElements.GetPaintWindow()->GetType() != EWindowType::GameWindow && InOutDrawElements.GetPaintWindow()->GetTag() != PIEWindowTag))
+	{
+		return;
+	}
+
 	++InOutLayerId;
 
 	FConsoleSlateDebuggerUtility::TSWindowId PaintWindow = FConsoleSlateDebuggerUtility::GetId(InOutDrawElements.GetPaintWindow());
@@ -322,10 +313,13 @@ void FConsoleSlateDebuggerPassBase::BuildInitialWidgetList()
 	for (TArray<TSharedRef<SWindow>>::TIterator It(SlateWindows); It; ++It)
 	{
 		const TSharedRef<SWindow>& Window = *It;
-
-		const FConsoleSlateDebuggerUtility::TSWindowId WindowId = FConsoleSlateDebuggerUtility::GetId(*Window);
-		AddInitialVisibleWidget(*StaticCastSharedRef<SWidget>(Window), WindowId);
-
+		
+		if (!(bDebugGameWindowOnly && Window->GetType() != EWindowType::GameWindow && Window->GetTag() != PIEWindowTag))
+		{
+			const FConsoleSlateDebuggerUtility::TSWindowId WindowId = FConsoleSlateDebuggerUtility::GetId(*Window);
+			AddInitialVisibleWidget(*StaticCastSharedRef<SWidget>(Window), WindowId);
+		}
+		
 		SlateWindows.Append(Window->GetChildWindows());
 	}
 }
@@ -337,11 +331,8 @@ void FConsoleSlateDebuggerPassBase::AddInitialVisibleWidget(const SWidget& Widge
 		return;
 	}
 
-	if (!(bDebugProjectContentOnly && !Widget.IsProjectContent()))
-	{
-		constexpr bool bIncrementUpdateCount = false; // This is called from BuildInitialWidgetList, we therefore do not want to record an update event
-		AddUpdatedWidget(Widget, WindowId, bIncrementUpdateCount);
-	}
+	const bool bIncrementUpdateCount = false; // This is called from BuildInitialWidgetList and we therefore do not want to record an update event
+	AddUpdatedWidget(Widget, WindowId, bIncrementUpdateCount);
 
 	const FChildren* Children = const_cast<SWidget&>(Widget).Debug_GetChildrenForReflector();
 	if (ensure(Children))
diff --git a/Engine/Source/Runtime/SlateCore/Private/Debugging/ConsoleSlateDebuggerPassBase.h b/Engine/Source/Runtime/SlateCore/Private/Debugging/ConsoleSlateDebuggerPassBase.h
--- a/Engine/Source/Runtime/SlateCore/Private/Debugging/ConsoleSlateDebuggerPassBase.h
+++ b/Engine/Source/Runtime/SlateCore/Private/Debugging/ConsoleSlateDebuggerPassBase.h
@@ -16,22 +16,6 @@
 class FConsoleSlateDebuggerPassBase
 {
 public:
-	struct FWidgetInfo 
-	{
-		FConsoleSlateDebuggerUtility::TSWindowId Window = FConsoleSlateDebuggerUtility::InvalidWindowId;
-		TWeakPtr<const SWidget> Widget; // used to check if the widget has been destroyed
-		FVector2f PaintLocation;
-		FVector2f PaintSize;
-		FString WidgetName;
-
-		uint32 LastUpdatedFrame = 0;
-		double LastUpdatedTime = 0.0;
-		int32 UpdateCount = 0;
-		bool bInInvalidationRoot = false;
-	};
-
-	using TUpdatedWidgetMap = TMap<FConsoleSlateDebuggerUtility::TSWidgetId, FWidgetInfo>;
-
 	FConsoleSlateDebuggerPassBase();
 	virtual ~FConsoleSlateDebuggerPassBase();
 
@@ -42,8 +26,6 @@ public:
 	
 	bool IsEnabled() const { return bEnabled; }
 
-	const TUpdatedWidgetMap& GetWidgets();
-
 	virtual void LoadConfig();
 	virtual void SaveConfig();
 
@@ -66,6 +48,20 @@ protected:
 	 * @param bIncrementUpdateCount If true, the UpdateCount of the widget should be increased. This can ve false when building the initial widget list
 	 */
 	virtual void AddUpdatedWidget(const SWidget& Widget, const FConsoleSlateDebuggerUtility::TSWindowId WindowId, bool bIncrementUpdateCount = true) = 0;
+	
+	struct FWidgetInfo 
+	{
+		FConsoleSlateDebuggerUtility::TSWindowId Window = FConsoleSlateDebuggerUtility::InvalidWindowId;
+		TWeakPtr<const SWidget> Widget; // used to check if the widget has been destroyed
+		FVector2f PaintLocation;
+		FVector2f PaintSize;
+		FString WidgetName;
+
+		bool bInInvalidationRoot = false;
+		uint32 LastUpdatedFrame = 0;
+		double LastUpdatedTime = 0.0;
+		int32 UpdateCount = 0;
+	};
 
 	/**
 	 * The internal function that should be called from AddUpdatedWidget to update the internal widget list.
@@ -107,16 +103,18 @@ protected:
 	bool bDrawBorder;
 	bool bLogWidgetName;
 	bool bLogWidgetNameOnce;
-	bool bDebugProjectContentOnly;
+	bool bDebugGameWindowOnly;
 	bool bInvalidationRootOnly;
 	FLinearColor MostRecentColor;
 	FLinearColor LeastRecentColor;
 	FLinearColor DrawWidgetNameColor;
 	int32 MaxNumberOfWidgetInList;
 	float FadeDuration;
+	const FName PIEWindowTag;
 
 	FVector2f WidgetLogLocation{10.f, 10.f};
 
+	using TUpdatedWidgetMap = TMap<FConsoleSlateDebuggerUtility::TSWidgetId, FWidgetInfo>;
 	TUpdatedWidgetMap Widgets;
 };
 
diff --git a/Engine/Source/Runtime/SlateCore/Private/Debugging/ConsoleSlateDebuggerPrepass.cpp b/Engine/Source/Runtime/SlateCore/Private/Debugging/ConsoleSlateDebuggerPrepass.cpp
--- a/Engine/Source/Runtime/SlateCore/Private/Debugging/ConsoleSlateDebuggerPrepass.cpp
+++ b/Engine/Source/Runtime/SlateCore/Private/Debugging/ConsoleSlateDebuggerPrepass.cpp
@@ -46,16 +46,11 @@ FConsoleSlateDebuggerPrepass::FConsoleSlateDebuggerPrepass()
 		, MaxNumberOfWidgetInList
 		, TEXT("The max number of widgets that will be displayed when DisplayWidgetNameList is active.")
 		, FConsoleVariableDelegate::CreateRaw(this, &FConsoleSlateDebuggerPrepass::SaveConfigOnVariableChanged))
-	, OnlyProjectContentRefCVar(
-		TEXT("SlateDebugger.Prepass.OnlyProjectContent")
-		, bDebugProjectContentOnly
-		, TEXT("Only debug the project's contents (exclude editor, chrome, diagnostic output)")
+	, OnlyGameWindowRefCVar(
+		TEXT("SlateDebugger.Prepass.OnlyGameWindow")
+		, bDebugGameWindowOnly
+		, TEXT("Option to only debug the game window")
 		, FConsoleVariableDelegate::CreateRaw(this, &FConsoleSlateDebuggerPrepass::RebuildWidgetListOnVariableChanged))
-	, OnlyGameWindowDepCVar(
-		TEXT("SlateDebugger.Prepass.OnlyGameWindow"),
-		TEXT("SlateDebugger.Prepass.OnlyProjectContent"),
-		TEXT("5.8"),
-		EShadowCVarBehavior::Warn)
 	, OnlyInvalidationRootRefCVar(
 		TEXT("SlateDebugger.Prepass.OnlyInvalidationRoot")
 		, bInvalidationRootOnly
@@ -85,7 +80,7 @@ void FConsoleSlateDebuggerPrepass::LoadConfig()
 {
 	FConsoleSlateDebuggerPassBase::LoadConfig();
 	
-	OnlyProjectContentRefCVar->Set(bDebugProjectContentOnly, ECVF_SetByProjectSetting);
+	OnlyGameWindowRefCVar->Set(bDebugGameWindowOnly, ECVF_SetByProjectSetting);
 	OnlyInvalidationRootRefCVar->Set(bInvalidationRootOnly, ECVF_SetByProjectSetting);
 	EnableWidgetsNameListRefCVar->Set(bDisplayWidgetsNameList, ECVF_SetByProjectSetting);
 	DrawFillEnabledRefCVar->Set(bDrawBox, ECVF_SetByProjectSetting);
@@ -118,9 +113,9 @@ void FConsoleSlateDebuggerPrepass::AddUpdatedWidget(const SWidget& Widget, const
 
 void FConsoleSlateDebuggerPrepass::HandleEndWidgetPrepass(const SWidget* Widget)
 {
+	// Exclude all windows but the game window
 	const TSharedPtr<SWindow> WindowToDrawIn = FSlateApplicationBase::Get().FindWidgetWindow(Widget->AsShared());
-
-	if (!WindowToDrawIn || (bDebugProjectContentOnly && !Widget->IsProjectContent()))
+	if (!WindowToDrawIn || (bDebugGameWindowOnly && (WindowToDrawIn->GetType() != EWindowType::GameWindow && WindowToDrawIn->GetTag() != PIEWindowTag)))
 	{
 		return;
 	}
diff --git a/Engine/Source/Runtime/SlateCore/Private/Debugging/ConsoleSlateDebuggerPrepass.h b/Engine/Source/Runtime/SlateCore/Private/Debugging/ConsoleSlateDebuggerPrepass.h
--- a/Engine/Source/Runtime/SlateCore/Private/Debugging/ConsoleSlateDebuggerPrepass.h
+++ b/Engine/Source/Runtime/SlateCore/Private/Debugging/ConsoleSlateDebuggerPrepass.h
@@ -52,8 +52,7 @@ private:
 	FAutoConsoleVariableRef EnableWidgetsNameListRefCVar;
 	FAutoConsoleCommand ToggleWidgetsNameListCommand;
 	FAutoConsoleVariableRef MaxNumberOfWidgetInListRefCVar;
-	FAutoConsoleVariableRef OnlyProjectContentRefCVar;
-	FAutoConsoleVariableDeprecated OnlyGameWindowDepCVar;
+	FAutoConsoleVariableRef OnlyGameWindowRefCVar;
 	FAutoConsoleVariableRef OnlyInvalidationRootRefCVar;
 	FAutoConsoleVariableRef DrawBorderEnabledRefCVar;
 	FAutoConsoleVariableRef DrawFillEnabledRefCVar;
diff --git a/Engine/Source/Runtime/SlateCore/Private/Debugging/ConsoleSlateDebuggerUpdate.cpp b/Engine/Source/Runtime/SlateCore/Private/Debugging/ConsoleSlateDebuggerUpdate.cpp
--- a/Engine/Source/Runtime/SlateCore/Private/Debugging/ConsoleSlateDebuggerUpdate.cpp
+++ b/Engine/Source/Runtime/SlateCore/Private/Debugging/ConsoleSlateDebuggerUpdate.cpp
@@ -26,7 +26,7 @@ FConsoleSlateDebuggerUpdate::FConsoleSlateDebuggerUpdate()
 	, bDisplayUpdateFromPaint(false)
 	, bShowLegend(false)
 	, bShowQuad(false)
-	, bDebugProjectContentOnly(true)
+	, bDebugGameWindowOnly(true)
 	, WidgetUpdateFlagsFilter(EWidgetUpdateFlags::AnyUpdate)
 	, DrawVolatilePaintColor(FColorList::Red)
 	, DrawRepaintColor(FColorList::Yellow)
@@ -36,6 +36,7 @@ FConsoleSlateDebuggerUpdate::FConsoleSlateDebuggerUpdate()
 	, MaxNumberOfWidgetInList(20)
 	, InvalidationRootIdFilter(-1)
 	, CacheDuration(2.0f)
+	, PIEWindowTag("PIEWindow")
 	, StartCommand(
 		TEXT("SlateDebugger.Update.Start"),
 		TEXT("Start the update widget debug tool. It shows when widgets are updated."),
@@ -69,15 +70,10 @@ FConsoleSlateDebuggerUpdate::FConsoleSlateDebuggerUpdate()
 		TEXT("SlateDebugger.Update.SetInvalidationRootIdFilter"),
 		InvalidationRootIdFilter,
 		TEXT("Option to show only the widgets that are part of an invalidation root."))
-	, OnlyProjectContentRefCVar(
-		TEXT("SlateDebugger.Update.OnlyProjectContent"),
-		bDebugProjectContentOnly,
-		TEXT("Only debug content widgets (exclude editor, chrome, diagnostic output)"))
-	, OnlyGameWindowDepCVar(
+	, OnlyGameWindow(
 		TEXT("SlateDebugger.Update.OnlyGameWindow"),
-		TEXT("SlateDebugger.Update.OnlyProjectContent"),
-		TEXT("5.8"),
-		EShadowCVarBehavior::Warn)
+		bDebugGameWindowOnly,
+		TEXT("Option to only the debug the game window"))
 {
 	GConfig->GetBool(TEXT("SlateDebugger.Update"), TEXT("bDisplayWidgetsNameList"), bDisplayWidgetsNameList, *GEditorPerProjectIni);
 	GConfig->GetBool(TEXT("SlateDebugger.Update"), TEXT("bUseWidgetPathAsName"), bUseWidgetPathAsName, *GEditorPerProjectIni);
@@ -309,7 +305,6 @@ FConsoleSlateDebuggerUpdate::FWidgetInfo::FWidgetInfo(const SWidget* Widget, EWi
 	: WindowId(0)
 	, UpdateFlags(InUpdateFlags)
 	, LastInvalidationTime(FSlateApplicationBase::Get().GetCurrentTime())
-	, bIsProjectContent(Widget->IsProjectContent())
 {
 	PaintLocation = Widget->GetPersistentState().AllottedGeometry.GetAbsolutePosition();
 	PaintSize = Widget->GetPersistentState().AllottedGeometry.GetAbsoluteSize();
@@ -366,6 +361,13 @@ void FConsoleSlateDebuggerUpdate::HandlePaintDebugInfo(const FPaintArgs& InArgs,
 
 	++InOutLayerId;
 
+	// Exclude all windows but the game window
+	SWindow* WindowToDrawIn = InOutDrawElements.GetPaintWindow();
+	if (bDebugGameWindowOnly && (WindowToDrawIn->GetType() != EWindowType::GameWindow && WindowToDrawIn->GetTag() != PIEWindowTag))
+	{
+		return;
+	}
+
 	const FConsoleSlateDebuggerUtility::TSWindowId PaintWindow = FConsoleSlateDebuggerUtility::GetId(InOutDrawElements.GetPaintWindow());
 	int32 NumberOfWidget = 0;
 	TArray<const FString*, TInlineAllocator<100>> NamesToDisplay;
@@ -402,11 +404,6 @@ void FConsoleSlateDebuggerUpdate::HandlePaintDebugInfo(const FPaintArgs& InArgs,
 	{
 		if (Itt.Value.WindowId == PaintWindow)
 		{
-			if (bDebugProjectContentOnly && !Itt.Value.bIsProjectContent)
-			{
-				continue;
-			}
-
 			const float LerpValue = FMath::Clamp((float)(SlateApplicationCurrentTime - Itt.Value.LastInvalidationTime) / CacheDuration, 0.f, 1.f);
 			const FGeometry Geometry = FGeometry::MakeRoot(Itt.Value.PaintSize, FSlateLayoutTransform(1.f, Itt.Value.PaintLocation));
 			const FPaintGeometry PaintGeometry = Geometry.ToPaintGeometry();
diff --git a/Engine/Source/Runtime/SlateCore/Private/Debugging/ConsoleSlateDebuggerUpdate.h b/Engine/Source/Runtime/SlateCore/Private/Debugging/ConsoleSlateDebuggerUpdate.h
--- a/Engine/Source/Runtime/SlateCore/Private/Debugging/ConsoleSlateDebuggerUpdate.h
+++ b/Engine/Source/Runtime/SlateCore/Private/Debugging/ConsoleSlateDebuggerUpdate.h
@@ -37,7 +37,6 @@ public:
 	bool IsEnabled() const { return bEnabled; }
 
 	void ToggleDisplayLegend();
-	bool IsLegendShown() const { return bShowLegend; }
 	void ToogleDisplayWidgetNameList();
 	void ToogleDisplayUpdateFromPaint();
 
@@ -61,7 +60,7 @@ private:
 	bool bDisplayUpdateFromPaint;
 	bool bShowLegend;
 	bool bShowQuad;
-	bool bDebugProjectContentOnly;
+	bool bDebugGameWindowOnly;
 	EWidgetUpdateFlags WidgetUpdateFlagsFilter;
 	FLinearColor DrawVolatilePaintColor;
 	FLinearColor DrawRepaintColor;
@@ -71,6 +70,7 @@ private:
 	int32 MaxNumberOfWidgetInList;
 	int32 InvalidationRootIdFilter;
 	float CacheDuration;
+	FName PIEWindowTag;
 
 	//~ Console objects
 	FAutoConsoleCommand StartCommand;
@@ -81,8 +81,7 @@ private:
 	FAutoConsoleCommand ToogleDisplayUpdateFromPaintCommand;
 	FAutoConsoleCommand SetWidgetUpdateFlagsFilterCommand;
 	FAutoConsoleVariableRef InvalidationRootFilterRefCVar;
-	FAutoConsoleVariableRef OnlyProjectContentRefCVar;
-	FAutoConsoleVariableDeprecated OnlyGameWindowDepCVar;
+	FAutoConsoleVariableRef OnlyGameWindow;
 
 	struct FWidgetInfo
 	{
@@ -95,7 +94,6 @@ private:
 		FString WidgetName;
 		EWidgetUpdateFlags UpdateFlags;
 		double LastInvalidationTime;
-		bool bIsProjectContent;
 	};
 
 	using TWidgetMap = TMap<FConsoleSlateDebuggerUtility::TSWidgetId, FWidgetInfo>;
diff --git a/Engine/Source/Runtime/SlateCore/Private/FastUpdate/WidgetProxy.cpp b/Engine/Source/Runtime/SlateCore/Private/FastUpdate/WidgetProxy.cpp
--- a/Engine/Source/Runtime/SlateCore/Private/FastUpdate/WidgetProxy.cpp
+++ b/Engine/Source/Runtime/SlateCore/Private/FastUpdate/WidgetProxy.cpp
@@ -297,7 +297,7 @@ FWidgetProxy::FUpdateResult FWidgetProxy::Repaint(const FPaintArgs& PaintArgs, F
 		OutDrawElements.PopPixelSnappingMethod();
 	}
 
-	OutDrawElements.SetIsInProjectContent(MyState.IsProjectContent());
+	OutDrawElements.SetIsInGameLayer(MyState.bIsInGameLayer);
 	return FUpdateResult{ PrevLayerId, NewLayerId };
 }
 
diff --git a/Engine/Source/Runtime/SlateCore/Private/Rendering/DrawElementTypes.cpp b/Engine/Source/Runtime/SlateCore/Private/Rendering/DrawElementTypes.cpp
--- a/Engine/Source/Runtime/SlateCore/Private/Rendering/DrawElementTypes.cpp
+++ b/Engine/Source/Runtime/SlateCore/Private/Rendering/DrawElementTypes.cpp
@@ -35,7 +35,7 @@ static void HandleApplyDisabledEffectOnWidgetsToggled(IConsoleVariable* CVar)
 static FAutoConsoleVariableRef CVarApplyDisabledEffectOnWidgets(
 	TEXT("Slate.ApplyDisabledEffectOnWidgets"),
 	bApplyDisabledEffectOnWidgets,
-	TEXT("If true, disabled widgets will have alpha multiplied by 0.45 (applies to project content only, not editor)."),
+	TEXT("If true, disabled game-layer widgets will have alpha multiplied by 0.45."),
 	FConsoleVariableDelegate::CreateStatic(&HandleApplyDisabledEffectOnWidgetsToggled)
 );
 
@@ -153,7 +153,7 @@ static bool ShouldCull(const FSlateWindowElementList& ElementList, const FPaintG
 
 void FSlateDrawElement::Init(FSlateWindowElementList& ElementList, EElementType InElementType, uint32 InLayer, const FPaintGeometry& PaintGeometry, ESlateDrawEffect InDrawEffects)
 {
-	if (ElementList.IsInProjectContent() && !bApplyDisabledEffectOnWidgets)
+	if (ElementList.GetIsInGameLayer() && !bApplyDisabledEffectOnWidgets)
 	{
 		InDrawEffects &= ~ESlateDrawEffect::DisabledEffect;
 	}
diff --git a/Engine/Source/Runtime/SlateCore/Private/Rendering/DrawElements.cpp b/Engine/Source/Runtime/SlateCore/Private/Rendering/DrawElements.cpp
--- a/Engine/Source/Runtime/SlateCore/Private/Rendering/DrawElements.cpp
+++ b/Engine/Source/Runtime/SlateCore/Private/Rendering/DrawElements.cpp
@@ -49,10 +49,10 @@ FSlateWindowElementList::FSlateWindowElementList(const TSharedPtr<SWindow>& InPa
 	, MemManagerAllocatedMemory(0)
 #endif
 	, RenderTargetWindow(nullptr)
-	, bIsProjectContent(false)
 	, bNeedsDeferredResolve(false)
 	, ResolveToDeferredIndex()
 	, WindowSize(FVector2f(0.0f, 0.0f))
+	, bIsInGameLayer(false)
 {
 	if (InPaintWindow.IsValid())
 	{
@@ -67,17 +67,12 @@ FSlateWindowElementList::~FSlateWindowElementList()
 
 void FSlateWindowElementList::SetIsInGameLayer(bool bInGameLayer)
 {
-	bIsProjectContent = bInGameLayer;
+	bIsInGameLayer = bInGameLayer;
 }
 
 bool FSlateWindowElementList::GetIsInGameLayer()
 {
-	return bIsProjectContent;
-}
-
-void FSlateWindowElementList::SetIsInProjectContent(bool bIsInProjectContent)
-{
-	bIsProjectContent = bIsInProjectContent;
+	return bIsInGameLayer;
 }
 
 void FSlateWindowElementList::StartMergeRenderBatches()
diff --git a/Engine/Source/Runtime/SlateCore/Private/SlateCoreModule.cpp b/Engine/Source/Runtime/SlateCore/Private/SlateCoreModule.cpp
--- a/Engine/Source/Runtime/SlateCore/Private/SlateCoreModule.cpp
+++ b/Engine/Source/Runtime/SlateCore/Private/SlateCoreModule.cpp
@@ -1,6 +1,6 @@
 // Copyright Epic Games, Inc. All Rights Reserved.
 
-#include "SlateCoreModule.h"
+#include "CoreMinimal.h"
 #include "Debugging/ConsoleSlateDebugger.h"
 #include "Debugging/ConsoleSlateDebuggerInvalidate.h"
 #include "Debugging/ConsoleSlateDebuggerInvalidationRoot.h"
@@ -8,6 +8,7 @@
 #include "Debugging/ConsoleSlateDebuggerPrepass.h"
 #include "Debugging/ConsoleSlateDebuggerUpdate.h"
 #include "Debugging/ConsoleSlateDebuggerBreak.h"
+#include "Modules/ModuleInterface.h"
 #include "Modules/ModuleManager.h"
 #include "SlateGlobals.h"
 #include "Types/SlateStructs.h"
@@ -17,20 +18,15 @@ DEFINE_LOG_CATEGORY(LogSlateStyles);
 
 const float FOptionalSize::Unspecified = -1.0f;
 
-namespace
-{
-	UE::FSlateCoreModule* Instance;
-}
 
-namespace UE
+/**
+ * Implements the SlateCore module.
+ */
+class FSlateCoreModule
+	: public IModuleInterface
 {
-	FSlateCoreModule& FSlateCoreModule::Get()
-	{
-		check(Instance);
-		return *Instance;
-	}
-
-	FSlateCoreModule::FSlateCoreModule()
+public:
+	FSlateCoreModule()
 	{
 #if WITH_SLATE_DEBUGGING
 		SlateDebuggerEvent = MakeUnique<FConsoleSlateDebugger>();
@@ -45,33 +41,17 @@ namespace UE
 #endif
 	}
 
-	void FSlateCoreModule::StartupModule()
-	{
-		Instance = this;
-	}
-
-	void FSlateCoreModule::ShutdownModule()
-	{
-		Instance = nullptr;
-	}
-
 #if WITH_SLATE_DEBUGGING
-	FConsoleSlateDebuggerPaint& FSlateCoreModule::GetPaintDebugger() const
-	{
-		return *SlateDebuggerPaint;
-	}
-
-	FConsoleSlateDebuggerPrepass& FSlateCoreModule::GetPrepassDebugger() const
-	{
-		return *SlateDebuggerPrepass;
-	}
-
-	FConsoleSlateDebuggerUpdate& FSlateCoreModule::GetUpdateDebugger() const
-	{
-		return *SlateDebuggerUpdate;
-	}
+private:
+	TUniquePtr<FConsoleSlateDebugger> SlateDebuggerEvent;
+	TUniquePtr<FConsoleSlateDebuggerInvalidate> SlateDebuggerInvalidate;
+	TUniquePtr<FConsoleSlateDebuggerInvalidationRoot> SlateDebuggerInvalidationRoot;
+	TUniquePtr<FConsoleSlateDebuggerPaint> SlateDebuggerPaint;
+	TUniquePtr<FConsoleSlateDebuggerPrepass> SlateDebuggerPrepass;
+	TUniquePtr<FConsoleSlateDebuggerUpdate> SlateDebuggerUpdate;
+	TUniquePtr<FConsoleSlateDebuggerBreak> SlateDebuggerBreak;
 #endif
-}
+};
 
 
-IMPLEMENT_MODULE(UE::FSlateCoreModule, SlateCore);
+IMPLEMENT_MODULE(FSlateCoreModule, SlateCore);
diff --git a/Engine/Source/Runtime/SlateCore/Private/SlateCoreModule.h b/Engine/Source/Runtime/SlateCore/Private/SlateCoreModule.h
deleted file mode 100644
diff --git a/Engine/Source/Runtime/SlateCore/Private/Test/ConsoleSlateDebuggerWindowFilterTest.cpp b/Engine/Source/Runtime/SlateCore/Private/Test/ConsoleSlateDebuggerWindowFilterTest.cpp
deleted file mode 100644
diff --git a/Engine/Source/Runtime/SlateCore/Private/Widgets/SWidget.cpp b/Engine/Source/Runtime/SlateCore/Private/Widgets/SWidget.cpp
--- a/Engine/Source/Runtime/SlateCore/Private/Widgets/SWidget.cpp
+++ b/Engine/Source/Runtime/SlateCore/Private/Widgets/SWidget.cpp
@@ -228,8 +228,6 @@ SWidget::SWidget()
 	, bHasActiveTimers(false)
 	, bDesiredSizeSet(false)
 	, bPrepassLayoutScaleMultiplierSet(false)
-	, bIsProjectContent(false)
-	, bIsProjectContentParent(false)
 	, bHasCustomPrepass(false)
 	, bHasRelativeLayoutScale(false)
 	, bVolatilityAlwaysInvalidatesPrepass(false)
@@ -1521,7 +1519,7 @@ int32 SWidget::Paint(const FPaintArgs& Args, const FGeometry& AllottedGeometry,
 	PersistentState.IncomingUserIndex = (int8)IncomingUserIndex;
 
 	PersistentState.IncomingFlowDirection = GSlateFlowDirection;
-	PersistentState.SetIsProjectContent(IsProjectContent());
+	PersistentState.bIsInGameLayer = OutDrawElements.GetIsInGameLayer();
 
 	FPaintArgs UpdatedArgs = Args.WithNewParent(this);
 	UpdatedArgs.SetInheritedHittestability(bOutgoingHittestability);
@@ -1609,10 +1607,8 @@ int32 SWidget::Paint(const FPaintArgs& Args, const FGeometry& AllottedGeometry,
 		LLM_SCOPE_DYNAMIC_STAT_OBJECTPATH_FNAME(ClassName, ELLMTagSet::AssetClasses);
 		UE_TRACE_METADATA_SCOPE_ASSET_FNAME(AssetName, ClassName, PackageName);
 #endif
-		// Paint the geometry of this widget.
-		OutDrawElements.SetIsInProjectContent(IsProjectContent());
+	// Paint the geometry of this widget.
 		NewLayerId = OnPaint(UpdatedArgs, AllottedGeometry, CullingBounds, OutDrawElements, LayerId, ContentWidgetStyle, bParentEnabled);
-		OutDrawElements.SetIsInProjectContent(false);
 	}
 
 	// Just repainted
@@ -1834,9 +1830,6 @@ void SWidget::Prepass_ChildLoop(float InLayoutScaleMultiplier, FChildren* MyChil
 				? InLayoutScaleMultiplier * Self->GetRelativeLayoutScale(ChildIndex, InLayoutScaleMultiplier)
 				: InLayoutScaleMultiplier;
 
-			// Inherit project-content status by default.  Children can opt out (overwrite) in custom prepass.
-			Child.bIsProjectContent = Self->bIsProjectContent || Self->bIsProjectContentParent;
-
 			// Recur: Descend down the widget tree.
 			Child.Prepass_Internal(ChildLayoutScaleMultiplier);
 		}
diff --git a/Engine/Source/Runtime/SlateCore/Public/FastUpdate/WidgetProxy.h b/Engine/Source/Runtime/SlateCore/Public/FastUpdate/WidgetProxy.h
--- a/Engine/Source/Runtime/SlateCore/Public/FastUpdate/WidgetProxy.h
+++ b/Engine/Source/Runtime/SlateCore/Public/FastUpdate/WidgetProxy.h
@@ -249,9 +249,7 @@ struct FSlateWidgetPersistentState
 		, bParentEnabled(true)
 		, bInheritedHittestability(false)
 		, bDeferredPainting(false)
-PRAGMA_DISABLE_DEPRECATION_WARNINGS
 		, bIsInGameLayer(false)
-PRAGMA_ENABLE_DEPRECATION_WARNINGS
 	{}
 
 	TWeakPtr<SWidget> PaintParent;
@@ -270,26 +268,8 @@ PRAGMA_ENABLE_DEPRECATION_WARNINGS
 	uint8 bParentEnabled : 1;
 	uint8 bInheritedHittestability : 1;
 	uint8 bDeferredPainting : 1;
-
-	UE_DEPRECATED(5.8, "Use methods instead of direct field access")
 	uint8 bIsInGameLayer : 1;
 
-	/** The value of SWidget::IsContent() */
-	bool IsProjectContent() const
-	{
-PRAGMA_DISABLE_DEPRECATION_WARNINGS
-		return bIsInGameLayer;
-PRAGMA_ENABLE_DEPRECATION_WARNINGS
-	}
-
-	/** Only to be called by SWidget::Paint. */
-	void SetIsProjectContent(bool bInIsContent)
-	{
-PRAGMA_DISABLE_DEPRECATION_WARNINGS
-		bIsInGameLayer = bInIsContent;
-PRAGMA_ENABLE_DEPRECATION_WARNINGS
-	}
-
 	static const FSlateWidgetPersistentState NoState;
 };
 
diff --git a/Engine/Source/Runtime/SlateCore/Public/Rendering/DrawElements.h b/Engine/Source/Runtime/SlateCore/Public/Rendering/DrawElements.h
--- a/Engine/Source/Runtime/SlateCore/Public/Rendering/DrawElements.h
+++ b/Engine/Source/Runtime/SlateCore/Public/Rendering/DrawElements.h
@@ -405,20 +405,9 @@ public:
 
 	void AddReferencedObjects(FReferenceCollector& Collector);
 	
-	UE_DEPRECATED(5.8, "Use SetIsProjectContent from CustomPrepass instead")
 	SLATECORE_API void SetIsInGameLayer(bool bInGameLayer);
-
-	UE_DEPRECATED(5.8, "Renamed to IsInProjectContent")
 	SLATECORE_API bool GetIsInGameLayer();
 
-	void SetIsInProjectContent(bool bIsInProjectContent);
-
-	/** Whether the stack frame is drawing the project's content (as opposed to editor, chrome, diagnostic output). */
-	bool IsInProjectContent() const
-	{
-		return bIsProjectContent;
-	}
-
 	TArrayView<FSlateCachedElementList* const> GetCurrentCachedElementWithNewData() const
 	{ 
 		if (GetCurrentCachedElementData() != nullptr)
@@ -493,9 +482,6 @@ private:
 	 */
 	TArray< TSharedPtr<FDeferredPaint> > DeferredPaintList;
 
-	/** Caches whether the current widget constitutes "content" (as opposed to editor, chrome, diagnostic output). */
-	bool bIsProjectContent;
-
 	bool bNeedsDeferredResolve;
 	TArray<int32> ResolveToDeferredIndex;
 
@@ -526,4 +512,7 @@ private:
 
 	/** Store the size of the window being used to paint */
 	FVector2f WindowSize;
+
+	/** Store if currently drawing within Game Layer */
+	bool bIsInGameLayer;
 };
diff --git a/Engine/Source/Runtime/SlateCore/Public/Widgets/SWidget.h b/Engine/Source/Runtime/SlateCore/Public/Widgets/SWidget.h
--- a/Engine/Source/Runtime/SlateCore/Public/Widgets/SWidget.h
+++ b/Engine/Source/Runtime/SlateCore/Public/Widgets/SWidget.h
@@ -721,35 +721,9 @@ public:
 		return EnumHasAnyFlags(UpdateFlags, FlagsToCheck);
 	}
 
-	/** Whether the widget is a part of the project's content (as opposed to editor, chrome, diagnostic output).
-	 * Dynamic: updated during pre-pass, inherited from parent but may be overwritten in CustomPrepass.
-	 * Used to filter which widgets can be navigated to, which are considered by diagnostic tools.
-	 */
-	bool IsProjectContent() const
-	{
-		return bIsProjectContent;
-	}
-
 protected:
 	void SetVolatilePrepass(bool bVolatile) { bVolatile ? AddUpdateFlags(EWidgetUpdateFlags::NeedsVolatilePrepass) : RemoveUpdateFlags(EWidgetUpdateFlags::NeedsVolatilePrepass); }
 
-	/**
-	 * Call from CustomPrepass to change whether this widget and its children are a part of the project's content (as opposed to editor, chrome,
-	 * diagnostic output).
-	 */
-	void SetIsProjectContent(bool bInIsProjectContent)
-	{
-		bIsProjectContent = bInIsProjectContent;
-	}
-
-	/**
-	 * Call from CustomPrepass to change whether this widget's children are a part of the project's content.
-	 */
-	void SetIsProjectContentParent(bool bInIsProjectContent)
-	{
-		bIsProjectContentParent = bInIsProjectContent;
-	}
-
 	virtual bool CustomPrepass(float LayoutScaleMultiplier) { return false; }
 
 	/**
@@ -1895,9 +1869,6 @@ private:
 	uint8 bSupportsInvalidation : 1 = false;
 #endif
 
-	uint8 bIsProjectContent : 1;
-	uint8 bIsProjectContentParent : 1;
-
 protected:
 	uint8 bHasCustomPrepass : 1;
 
