Slate: Console Debugger: Don't warn when cached ini has deprecated bDebugGameWindowOnly

This was firing even when the version-controlled configs of a project did not have this setting. It upgrades itself anyhow.

#rnx
#Synced-CL 51488891

diff --git a/Engine/Source/Runtime/SlateCore/Private/Debugging/ConsoleSlateDebuggerPassBase.cpp b/Engine/Source/Runtime/SlateCore/Private/Debugging/ConsoleSlateDebuggerPassBase.cpp
--- a/Engine/Source/Runtime/SlateCore/Private/Debugging/ConsoleSlateDebuggerPassBase.cpp
+++ b/Engine/Source/Runtime/SlateCore/Private/Debugging/ConsoleSlateDebuggerPassBase.cpp
@@ -49,7 +49,6 @@ void FConsoleSlateDebuggerPassBase::LoadConfig()
 	FString Text;
 	if (GConfig->GetString(*Section, TEXT("bDebugGameWindowOnly"), Text, *GEditorPerProjectIni))
 	{
-		UE_LOGF(LogSlateDebugger, Warning, "[%S] bDebugGameWindowOnly is deprecated, use bDebugProjectContentOnly instead", *Section);
 		bDebugProjectContentOnly = FCString::ToBool(*Text);
 	}
 
