Moving Navigation Debugger behind it's own preprocessor macro: WITH_SLATE_NAVIGATION_DEBUGGING

#tests Disabled WITH_SLATE_DEBUGGING and Enabled WITH_SLATE_NAVIGATION_DEBUGGING to make sure things still compiled
#Synced-CL 51546663

diff --git a/Engine/Source/Runtime/Engine/Private/Slate/SGameLayerManager.cpp b/Engine/Source/Runtime/Engine/Private/Slate/SGameLayerManager.cpp
--- a/Engine/Source/Runtime/Engine/Private/Slate/SGameLayerManager.cpp
+++ b/Engine/Source/Runtime/Engine/Private/Slate/SGameLayerManager.cpp
@@ -21,9 +21,9 @@
 #include "Types/InvisibleToWidgetReflectorMetaData.h"
 #include "Framework/Application/SlateApplication.h"
 
-#if WITH_SLATE_DEBUGGING
+#if WITH_SLATE_NAVIGATION_DEBUGGING
 #include "Navigation/SlateNavigationDebugger.h"
-#endif //WITH_SLATE_DEBUGGING
+#endif //WITH_SLATE_NAVIGATION_DEBUGGING
 
 #include UE_INLINE_GENERATED_CPP_BY_NAME(SGameLayerManager)
 
@@ -439,7 +439,7 @@ int32 SGameLayerManager::OnPaint(const FPaintArgs& Args, const FGeometry& Allott
 
 	int32 ResultLayer = SCompoundWidget::OnPaint(Args, AllottedGeometry, MyCullingRect, OutDrawElements, LayerId, InWidgetStyle, bParentEnabled);
 
-#if WITH_SLATE_DEBUGGING
+#if WITH_SLATE_NAVIGATION_DEBUGGING
 	ResultLayer = FSlateNavigationDebugger::Get().Paint(OutDrawElements, ResultLayer);
 #endif
 
diff --git a/Engine/Source/Runtime/Slate/Private/Navigation/SlateNavigationDebugger.cpp b/Engine/Source/Runtime/Slate/Private/Navigation/SlateNavigationDebugger.cpp
--- a/Engine/Source/Runtime/Slate/Private/Navigation/SlateNavigationDebugger.cpp
+++ b/Engine/Source/Runtime/Slate/Private/Navigation/SlateNavigationDebugger.cpp
@@ -2,7 +2,7 @@
 
 #include "Navigation/SlateNavigationDebugger.h"
 
-#if WITH_SLATE_DEBUGGING
+#if WITH_SLATE_NAVIGATION_DEBUGGING
 
 #include "Misc/Paths.h"
 #include "Styling/SlateStyleRegistry.h"
@@ -694,4 +694,4 @@ int32 FSlateNavigationDebugger::Paint(FSlateWindowElementList& OutDrawElements,
 }
 
 
-#endif // WITH_SLATE_DEBUGGING
\ No newline at end of file
+#endif // WITH_SLATE_NAVIGATION_DEBUGGING
\ No newline at end of file
diff --git a/Engine/Source/Runtime/Slate/Public/Navigation/SlateNavigationDebugger.h b/Engine/Source/Runtime/Slate/Public/Navigation/SlateNavigationDebugger.h
--- a/Engine/Source/Runtime/Slate/Public/Navigation/SlateNavigationDebugger.h
+++ b/Engine/Source/Runtime/Slate/Public/Navigation/SlateNavigationDebugger.h
@@ -4,7 +4,7 @@
 
 #include "Debugging/SlateDebugging.h"
 
-#if WITH_SLATE_DEBUGGING
+#if WITH_SLATE_NAVIGATION_DEBUGGING
 #include "Input/HittestGrid.h"
 #include "Styling/SlateStyle.h"
 
@@ -94,4 +94,4 @@ private:
 	void SetListenToHitTestGrid(const bool bBind, TArray<FHittestGrid::FDebuggingFindNextFocusableWidgetArgs>& InContainer, FDelegateHandle& InHandle) const;
 };
 
-#endif // WITH_SLATE_DEBUGGING
\ No newline at end of file
+#endif // WITH_SLATE_NAVIGATION_DEBUGGING
\ No newline at end of file
diff --git a/Engine/Source/Runtime/SlateCore/Private/Input/HittestGrid.cpp b/Engine/Source/Runtime/SlateCore/Private/Input/HittestGrid.cpp
--- a/Engine/Source/Runtime/SlateCore/Private/Input/HittestGrid.cpp
+++ b/Engine/Source/Runtime/SlateCore/Private/Input/HittestGrid.cpp
@@ -38,9 +38,9 @@ static FAutoConsoleVariableRef CVarEnableNavigationExperimental(
 	GEnableNavigationExperimental,
 	TEXT("If true, will use experimental navigation system"));
 
-#if WITH_SLATE_DEBUGGING
+#if WITH_SLATE_DEBUGGING || WITH_SLATE_NAVIGATION_DEBUGGING
 FHittestGrid::FDebuggingFindNextFocusableWidget FHittestGrid::OnFindNextFocusableWidgetExecuted;
-#endif //WITH_SLATE_DEBUGGING
+#endif //WITH_SLATE_DEBUGGING || WITH_SLATE_NAVIGATION_DEBUGGING
 
 //
 // Helper Functions
@@ -345,7 +345,7 @@ namespace UE::Slate::Private
 	}
 }
 
-#if WITH_SLATE_DEBUGGING
+#if WITH_SLATE_DEBUGGING || WITH_SLATE_NAVIGATION_DEBUGGING
 namespace HittestGridDebuggingText
 {
 	/*extern*/ const FText Valid = LOCTEXT("StateValid", "Valid"); //~ The widget is valid will be consider as the result
@@ -567,7 +567,7 @@ TSharedPtr<SWidget> FHittestGrid::FindNextFocusableWidgetDefault(const FArranged
 
 	TSharedPtr<SWidget> Widget = TSharedPtr<SWidget>();
 
-#if WITH_SLATE_DEBUGGING
+#if WITH_SLATE_DEBUGGING || WITH_SLATE_NAVIGATION_DEBUGGING
 	TArray<FDebuggingFindNextFocusableWidgetArgs::FWidgetResult> IntermediateResults;
 	TArray<FDebuggingFindNextFocusableWidgetArgs::FWidgetResult>* IntermediateResultsPtr = OnFindNextFocusableWidgetExecuted.IsBound() ? &IntermediateResults : nullptr;
 #else
@@ -627,7 +627,7 @@ TSharedPtr<SWidget> FHittestGrid::FindNextFocusableWidgetDefault(const FArranged
 		break;
 	}
 
-#if WITH_SLATE_DEBUGGING
+#if WITH_SLATE_DEBUGGING || WITH_SLATE_NAVIGATION_DEBUGGING
 	if (IntermediateResultsPtr)
 	{
 		FDebuggingFindNextFocusableWidgetArgs Args{ StartingWidget, Direction, NavigationReply, RuleWidget, UserIndex, Widget , MoveTemp(*IntermediateResultsPtr) };
@@ -667,7 +667,7 @@ TSharedPtr<SWidget> FHittestGrid::FindNextFocusableWidgetExperimental(const FArr
 		return TSharedPtr<SWidget>();
 	}
 
-#if WITH_SLATE_DEBUGGING
+#if WITH_SLATE_DEBUGGING || WITH_SLATE_NAVIGATION_DEBUGGING
 	TArray<FDebuggingFindNextFocusableWidgetArgs::FWidgetResult> IntermediateResults;
 	TArray<FDebuggingFindNextFocusableWidgetArgs::FWidgetResult>* IntermediateResultsPtr = OnFindNextFocusableWidgetExecuted.IsBound() ? &IntermediateResults : nullptr;
 #else
@@ -692,13 +692,18 @@ TSharedPtr<SWidget> FHittestGrid::FindNextFocusableWidgetExperimental(const FArr
 	FNavigationRoutingParams RoutingParams(StartingWidget.Widget, NextFocusedWidget, this, NavigationReply, UserIndex);
 	NextFocusedWidget = FNavigationRouting::RouteNavigationFocus(RoutingParams);
 
-#if WITH_SLATE_DEBUGGING
+#if WITH_SLATE_DEBUGGING || WITH_SLATE_NAVIGATION_DEBUGGING
 	if (IntermediateResultsPtr)
 	{
 		FDebuggingFindNextFocusableWidgetArgs Args{ StartingWidget, Direction, NavigationReply, RuleWidget, UserIndex, NextFocusedWidget, MoveTemp(*IntermediateResultsPtr) };
 		OnFindNextFocusableWidgetExecuted.Broadcast(this, Args);
 	}
 
+	
+#endif
+
+#if WITH_SLATE_DEBUGGING
+	
 	if (GSlateNavigationDebugging && NavigationMethod.IsValid())
 	{
 		// Duplicate, as this pointer may outlive the navigation method associated with raw instanced struct pointer
diff --git a/Engine/Source/Runtime/SlateCore/Private/Input/NavigationMethod.cpp b/Engine/Source/Runtime/SlateCore/Private/Input/NavigationMethod.cpp
--- a/Engine/Source/Runtime/SlateCore/Private/Input/NavigationMethod.cpp
+++ b/Engine/Source/Runtime/SlateCore/Private/Input/NavigationMethod.cpp
@@ -6,18 +6,6 @@
 
 #define LOCTEXT_NAMESPACE "NavigationMethod"
 
-#if WITH_SLATE_DEBUGGING
-
-FText FNavigationMethod::FDebuggingText::Valid = LOCTEXT("StateValid", "Valid"); //~ The widget is valid will be consider as the result
-FText FNavigationMethod::FDebuggingText::NotCompatibleWithUserIndex = LOCTEXT("StateNotCompatibleWithUserIndex", "User Index not compatible"); //~ The widget is not compatible with the requested user index
-FText FNavigationMethod::FDebuggingText::DoesNotIntersect = LOCTEXT("StateDoesNotIntersect", "Does not intersect"); //~ The widget rect is not in the correct direction or is not intersecting with the "swept" rectangle
-FText FNavigationMethod::FDebuggingText::PreviousWidgetIsBetter = LOCTEXT("StatePreviousWidgetIsBetter", "Previous Widget was better"); //~ The widget would be valid but the previous valid is closer
-FText FNavigationMethod::FDebuggingText::NotADescendant = LOCTEXT("StateNotADescendant", "Not a descendant"); //~ We have a non escape boundary condition and the widget isn't a descendant of our boundary
-FText FNavigationMethod::FDebuggingText::Disabled = LOCTEXT("StateNotEnabled", "Disabled"); //~ The widget is not enabled
-FText FNavigationMethod::FDebuggingText::ParentDisabled = LOCTEXT("StateParentNotEnabled", "ParentDisabled"); //~ A parent of the widget is disabled
-FText FNavigationMethod::FDebuggingText::DoesNotSuportKeyboardFocus = LOCTEXT("StateDoesNotSuportKeyboardFocus", "Keyboard focus unsupported"); //~ THe widget does not support keyboard focus
-#endif // WITH_SLATE_DEBUGGING
-
 void FNavigationMethod::Initialize(const FHittestGrid* InHittestGrid, TArray<FDebugWidgetResult>* InIntermediateResultsPtr)
 {
 	HittestGrid = InHittestGrid;
@@ -73,8 +61,8 @@ void FNavigationMethod::ForEachFocusableWidgetsInCell(int32 X, int32 Y, const FN
 
 		if (!HittestGrid->IsCompatibleUserIndex(UserIndex, WidgetData.UserIndex))
 		{
-#if WITH_SLATE_DEBUGGING
-			AddDebugIntermediateResult(Widget, FDebuggingText::NotCompatibleWithUserIndex);
+#if WITH_SLATE_DEBUGGING || WITH_SLATE_NAVIGATION_DEBUGGING
+			AddDebugIntermediateResult(Widget, HittestGridDebuggingText::NotCompatibleWithUserIndex);
 #endif
 			continue;
 		}
@@ -84,32 +72,32 @@ void FNavigationMethod::ForEachFocusableWidgetsInCell(int32 X, int32 Y, const FN
 			&& NavigationReply.GetHandler().IsValid()
 			&& !HittestGrid->IsDescendantOf(NavigationReply.GetHandler().Get(), WidgetData))
 		{
-#if WITH_SLATE_DEBUGGING
-			AddDebugIntermediateResult(Widget, FDebuggingText::NotADescendant);
+#if WITH_SLATE_DEBUGGING || WITH_SLATE_NAVIGATION_DEBUGGING
+			AddDebugIntermediateResult(Widget, HittestGridDebuggingText::NotADescendant);
 #endif
 			continue;
 		}
 
 		if (!Widget->IsEnabled())
 		{
-#if WITH_SLATE_DEBUGGING
-			AddDebugIntermediateResult(Widget, FDebuggingText::Disabled);
+#if WITH_SLATE_DEBUGGING || WITH_SLATE_NAVIGATION_DEBUGGING
+			AddDebugIntermediateResult(Widget, HittestGridDebuggingText::Disabled);
 #endif
 			continue;
 		}
 
 		if (!Widget->SupportsKeyboardFocus())
 		{
-#if WITH_SLATE_DEBUGGING
-			AddDebugIntermediateResult(Widget, FDebuggingText::DoesNotSuportKeyboardFocus);
+#if WITH_SLATE_DEBUGGING || WITH_SLATE_NAVIGATION_DEBUGGING
+			AddDebugIntermediateResult(Widget, HittestGridDebuggingText::DoesNotSuportKeyboardFocus);
 #endif
 			continue;
 		}
 
 		if (DisabledDestinations.Contains(Widget))
 		{
-#if WITH_SLATE_DEBUGGING
-			AddDebugIntermediateResult(Widget, FDebuggingText::ParentDisabled);
+#if WITH_SLATE_DEBUGGING || WITH_SLATE_NAVIGATION_DEBUGGING
+			AddDebugIntermediateResult(Widget, HittestGridDebuggingText::ParentDisabled);
 #endif
 			continue;
 		}
@@ -117,8 +105,8 @@ void FNavigationMethod::ForEachFocusableWidgetsInCell(int32 X, int32 Y, const FN
 		bool bValid = WidgetFunc(Widget);
 		if (bValid)
 		{
-#if WITH_SLATE_DEBUGGING
-			AddDebugIntermediateResult(Widget, FDebuggingText::Valid);
+#if WITH_SLATE_DEBUGGING || WITH_SLATE_NAVIGATION_DEBUGGING
+			AddDebugIntermediateResult(Widget, HittestGridDebuggingText::Valid);
 #endif
 		}
 	}
@@ -132,7 +120,9 @@ void FNavigationMethod::DrawDebug(int32 InLayer, const FGeometry& AllottedGeomet
 	(void)AllottedGeometry;
 	(void)WindowElementList;
 }
+#endif //WITH_SLATE_DEBUGGING
 
+#if WITH_SLATE_DEBUGGING || WITH_SLATE_NAVIGATION_DEBUGGING
 void FNavigationMethod::AddDebugIntermediateResult(const TSharedPtr<const SWidget>& InWidget, const FText Result) 
 { 
 	if (IntermediateResultsPtr) 
@@ -140,7 +130,6 @@ void FNavigationMethod::AddDebugIntermediateResult(const TSharedPtr<const SWidge
 		IntermediateResultsPtr->Emplace(InWidget, Result); 
 	} 
 }
-
-#endif //WITH_SLATE_DEBUGGING
+#endif //WITH_SLATE_DEBUGGING || WITH_SLATE_NAVIGATION_DEBUGGING
 
 #undef LOCTEXT_NAMESPACE
diff --git a/Engine/Source/Runtime/SlateCore/Private/Input/NavigationMethodOrthogonal.cpp b/Engine/Source/Runtime/SlateCore/Private/Input/NavigationMethodOrthogonal.cpp
--- a/Engine/Source/Runtime/SlateCore/Private/Input/NavigationMethodOrthogonal.cpp
+++ b/Engine/Source/Runtime/SlateCore/Private/Input/NavigationMethodOrthogonal.cpp
@@ -67,8 +67,8 @@ TSharedPtr<SWidget> FNavigationMethodOrthogonal::FindFocusableWidget(const FSlat
 
 				if (!HittestGrid->IsCompatibleUserIndex(UserIndex, TestCandidate.UserIndex))
 				{
-#if WITH_SLATE_DEBUGGING
-					AddDebugIntermediateResult(TestWidget, FDebuggingText::NotCompatibleWithUserIndex);
+#if WITH_SLATE_DEBUGGING || WITH_SLATE_NAVIGATION_DEBUGGING
+					AddDebugIntermediateResult(TestWidget, HittestGridDebuggingText::NotCompatibleWithUserIndex);
 #endif
 					continue;
 				}
@@ -76,8 +76,8 @@ TSharedPtr<SWidget> FNavigationMethodOrthogonal::FindFocusableWidget(const FSlat
 				const FSlateRect TestCandidateRect = HittestGrid->GetWidgetRenderBoundingRect(TestWidget);
 				if (!(CompareFunc(DestSideFunc(TestCandidateRect), CurrentSourceSide) && FSlateRect::DoRectanglesIntersect(SweptRect, TestCandidateRect)))
 				{
-#if WITH_SLATE_DEBUGGING
-					AddDebugIntermediateResult(TestWidget, FDebuggingText::DoesNotIntersect);
+#if WITH_SLATE_DEBUGGING || WITH_SLATE_NAVIGATION_DEBUGGING
+					AddDebugIntermediateResult(TestWidget, HittestGridDebuggingText::DoesNotIntersect);
 #endif
 					continue;
 				}
@@ -85,8 +85,8 @@ TSharedPtr<SWidget> FNavigationMethodOrthogonal::FindFocusableWidget(const FSlat
 				// If this found widget isn't closer then the previously found widget then keep looking.
 				if (BestWidget.IsValid() && !CompareFunc(DestSideFunc(BestWidgetRect), DestSideFunc(TestCandidateRect)))
 				{
-#if WITH_SLATE_DEBUGGING
-					AddDebugIntermediateResult(TestWidget, FDebuggingText::PreviousWidgetIsBetter);
+#if WITH_SLATE_DEBUGGING || WITH_SLATE_NAVIGATION_DEBUGGING
+					AddDebugIntermediateResult(TestWidget, HittestGridDebuggingText::PreviousWidgetIsBetter);
 #endif
 					continue;
 				}
@@ -96,32 +96,32 @@ TSharedPtr<SWidget> FNavigationMethodOrthogonal::FindFocusableWidget(const FSlat
 					&& NavigationReply.GetHandler().IsValid()
 					&& !HittestGrid->IsDescendantOf(NavigationReply.GetHandler().Get(), TestCandidate))
 				{
-#if WITH_SLATE_DEBUGGING
-					AddDebugIntermediateResult(TestWidget, FDebuggingText::NotADescendant);
+#if WITH_SLATE_DEBUGGING || WITH_SLATE_NAVIGATION_DEBUGGING
+					AddDebugIntermediateResult(TestWidget, HittestGridDebuggingText::NotADescendant);
 #endif
 					continue;
 				}
 
 				if (!TestWidget->IsEnabled())
 				{
-#if WITH_SLATE_DEBUGGING
-					AddDebugIntermediateResult(TestWidget, FDebuggingText::Disabled);
+#if WITH_SLATE_DEBUGGING || WITH_SLATE_NAVIGATION_DEBUGGING
+					AddDebugIntermediateResult(TestWidget, HittestGridDebuggingText::Disabled);
 #endif
 					continue;
 				}
 				
 				if (!TestWidget->SupportsKeyboardFocus())
 				{
-#if WITH_SLATE_DEBUGGING
-					AddDebugIntermediateResult(TestWidget, FDebuggingText::DoesNotSuportKeyboardFocus);
+#if WITH_SLATE_DEBUGGING || WITH_SLATE_NAVIGATION_DEBUGGING
+					AddDebugIntermediateResult(TestWidget, HittestGridDebuggingText::DoesNotSuportKeyboardFocus);
 #endif
 					continue;
 				}
 
 				if (DisabledDestinations.Contains(TestWidget))
 				{
-#if WITH_SLATE_DEBUGGING
-					AddDebugIntermediateResult(TestWidget, FDebuggingText::ParentDisabled);
+#if WITH_SLATE_DEBUGGING || WITH_SLATE_NAVIGATION_DEBUGGING
+					AddDebugIntermediateResult(TestWidget, HittestGridDebuggingText::ParentDisabled);
 #endif
 					continue;
 				}
@@ -129,8 +129,8 @@ TSharedPtr<SWidget> FNavigationMethodOrthogonal::FindFocusableWidget(const FSlat
 				BestWidgetRect = TestCandidateRect;
 				BestWidget = TestWidget;
 
-#if WITH_SLATE_DEBUGGING
-				AddDebugIntermediateResult(TestWidget, FDebuggingText::Valid);
+#if WITH_SLATE_DEBUGGING || WITH_SLATE_NAVIGATION_DEBUGGING
+				AddDebugIntermediateResult(TestWidget, HittestGridDebuggingText::Valid);
 #endif
 			}
 		}
diff --git a/Engine/Source/Runtime/SlateCore/Private/Input/NavigationMethodProximity.cpp b/Engine/Source/Runtime/SlateCore/Private/Input/NavigationMethodProximity.cpp
--- a/Engine/Source/Runtime/SlateCore/Private/Input/NavigationMethodProximity.cpp
+++ b/Engine/Source/Runtime/SlateCore/Private/Input/NavigationMethodProximity.cpp
@@ -167,8 +167,8 @@ FNavigationMethodProximity::FCellSearchResult FNavigationMethodProximity::FindBe
 				
 		if (!IsOutsideNavigationEdge(Direction, SearchInfo.CurrentSearchRect, CandidateRect))
 		{
-#if WITH_SLATE_DEBUGGING
-			AddDebugIntermediateResult(Candidate, FDebuggingText::DoesNotIntersect);
+#if WITH_SLATE_DEBUGGING || WITH_SLATE_NAVIGATION_DEBUGGING
+			AddDebugIntermediateResult(Candidate, HittestGridDebuggingText::DoesNotIntersect);
 #endif
 			return false;
 		}
@@ -182,8 +182,8 @@ FNavigationMethodProximity::FCellSearchResult FNavigationMethodProximity::FindBe
 		// If this found widget isn't closer then the previously found widget then keep looking.
 		if (CellResult.Widget.IsValid() && Distance >= CellResult.Distance)
 		{
-#if WITH_SLATE_DEBUGGING
-			AddDebugIntermediateResult(Candidate, FDebuggingText::PreviousWidgetIsBetter);
+#if WITH_SLATE_DEBUGGING || WITH_SLATE_NAVIGATION_DEBUGGING
+			AddDebugIntermediateResult(Candidate, HittestGridDebuggingText::PreviousWidgetIsBetter);
 #endif
 			return false;
 		}
@@ -192,8 +192,8 @@ FNavigationMethodProximity::FCellSearchResult FNavigationMethodProximity::FindBe
 		if (!IsParentsEnabled(Candidate.Get()))
 		{
 			DisabledDestinations.Add(Candidate);
-#if WITH_SLATE_DEBUGGING
-			AddDebugIntermediateResult(Candidate, FDebuggingText::ParentDisabled);
+#if WITH_SLATE_DEBUGGING || WITH_SLATE_NAVIGATION_DEBUGGING
+			AddDebugIntermediateResult(Candidate, HittestGridDebuggingText::ParentDisabled);
 #endif
 			return false;
 		}
diff --git a/Engine/Source/Runtime/SlateCore/Private/SlateCoreClasses.cpp b/Engine/Source/Runtime/SlateCore/Private/SlateCoreClasses.cpp
--- a/Engine/Source/Runtime/SlateCore/Private/SlateCoreClasses.cpp
+++ b/Engine/Source/Runtime/SlateCore/Private/SlateCoreClasses.cpp
@@ -99,13 +99,14 @@ FAutoConsoleVariableRef CVarHitTestGridDebugging(
 	TEXT("Slate.HitTestGridDebugging"),
 	GSlateHitTestGridDebugging,
 	TEXT("Whether to show a visualization of everything in the hit teest grid"));
+#endif
 
+#if WITH_SLATE_DEBUGGING || WITH_SLATE_NAVIGATION_DEBUGGING
 bool GSlateNavigationDebugging = false;
 FAutoConsoleVariableRef CVarNavigationDebugging(
 	TEXT("Slate.NavigationDebugging"),
 	GSlateNavigationDebugging,
 	TEXT("Whether to show a visualization of navigation from current widget"));
-
 #endif
 
 bool GSlateUseSharedBreakIterator = false;
diff --git a/Engine/Source/Runtime/SlateCore/Public/Debugging/SlateDebugging.h b/Engine/Source/Runtime/SlateCore/Public/Debugging/SlateDebugging.h
--- a/Engine/Source/Runtime/SlateCore/Public/Debugging/SlateDebugging.h
+++ b/Engine/Source/Runtime/SlateCore/Public/Debugging/SlateDebugging.h
@@ -16,6 +16,10 @@
 	#define WITH_SLATE_DEBUGGING !(UE_BUILD_SHIPPING || UE_BUILD_TEST)
 #endif
 
+#ifndef WITH_SLATE_NAVIGATION_DEBUGGING
+	#define WITH_SLATE_NAVIGATION_DEBUGGING WITH_SLATE_DEBUGGING
+#endif
+
 #ifndef SLATE_CSV_TRACKER
 	#define SLATE_CSV_TRACKER CSV_PROFILER_STATS
 #endif
diff --git a/Engine/Source/Runtime/SlateCore/Public/Input/HittestGrid.h b/Engine/Source/Runtime/SlateCore/Public/Input/HittestGrid.h
--- a/Engine/Source/Runtime/SlateCore/Public/Input/HittestGrid.h
+++ b/Engine/Source/Runtime/SlateCore/Public/Input/HittestGrid.h
@@ -14,7 +14,7 @@
 class FArrangedChildren;
 struct FNavigationMethod;
 
-#if WITH_SLATE_DEBUGGING
+#if WITH_SLATE_DEBUGGING || WITH_SLATE_NAVIGATION_DEBUGGING
 namespace HittestGridDebuggingText
 {
 	extern SLATECORE_API const FText Valid;
@@ -26,7 +26,7 @@ namespace HittestGridDebuggingText
 	extern SLATECORE_API const FText ParentDisabled;
 	extern SLATECORE_API const FText DoesNotSuportKeyboardFocus;
 };
-#endif // WITH_SLATE_DEBUGGING
+#endif // WITH_SLATE_DEBUGGING || WITH_SLATE_NAVIGATION_DEBUGGING
 
 class ICustomHitTestPath
 {
@@ -152,10 +152,12 @@ public:
 		TArray<FWidgetResult> IntermediateResults;
 	};
 
-#if WITH_SLATE_DEBUGGING
+#if WITH_SLATE_DEBUGGING || WITH_SLATE_NAVIGATION_DEBUGGING
 	DECLARE_MULTICAST_DELEGATE_TwoParams(FDebuggingFindNextFocusableWidget, const FHittestGrid* /*HittestGrid*/, const FDebuggingFindNextFocusableWidgetArgs& /*Info*/);
 	static SLATECORE_API FDebuggingFindNextFocusableWidget OnFindNextFocusableWidgetExecuted;
+#endif
 
+#if WITH_SLATE_DEBUGGING
 	SLATECORE_API void LogGrid() const;
 
 	enum class EDisplayGridFlags
diff --git a/Engine/Source/Runtime/SlateCore/Public/Input/NavigationMethod.h b/Engine/Source/Runtime/SlateCore/Public/Input/NavigationMethod.h
--- a/Engine/Source/Runtime/SlateCore/Public/Input/NavigationMethod.h
+++ b/Engine/Source/Runtime/SlateCore/Public/Input/NavigationMethod.h
@@ -39,20 +39,10 @@ private:
 #if WITH_SLATE_DEBUGGING
 public:
 	SLATECORE_API virtual void DrawDebug(int32 InLayer, const FGeometry& AllottedGeometry, FSlateWindowElementList& WindowElementList);
+#endif // WITH_SLATE_DEBUGGING
 
 protected:
-	struct FDebuggingText
-	{
-		static FText Valid;
-		static FText NotCompatibleWithUserIndex;
-		static FText DoesNotIntersect;
-		static FText PreviousWidgetIsBetter;
-		static FText NotADescendant;
-		static FText Disabled;
-		static FText ParentDisabled;
-		static FText DoesNotSuportKeyboardFocus;
-	};
-
+#if WITH_SLATE_DEBUGGING || WITH_SLATE_NAVIGATION_DEBUGGING
 	SLATECORE_API void AddDebugIntermediateResult(const TSharedPtr<const SWidget>& InWidget, const FText Result);
-#endif // WITH_SLATE_DEBUGGING
+#endif
 };
diff --git a/Engine/Source/Runtime/SlateCore/Public/SlateGlobals.h b/Engine/Source/Runtime/SlateCore/Public/SlateGlobals.h
--- a/Engine/Source/Runtime/SlateCore/Public/SlateGlobals.h
+++ b/Engine/Source/Runtime/SlateCore/Public/SlateGlobals.h
@@ -93,6 +93,9 @@ extern SLATECORE_API bool GSlateCheckUObjectShapedGlyphSequence;
 
 #if WITH_SLATE_DEBUGGING
 extern SLATECORE_API bool GSlateHitTestGridDebugging;
+#endif
+
+#if WITH_SLATE_DEBUGGING || WITH_SLATE_NAVIGATION_DEBUGGING
 extern SLATECORE_API bool GSlateNavigationDebugging;
 #endif
 /* Forward declarations
