[Input] Have FPointerEvent::PressedButtons copy by value instead of store a pointer so that it will be thread safe in the future if there are input events are safe to be copied across different threads.

#rnx
#Synced-CL 52059832

diff --git a/Engine/Source/Runtime/SlateCore/Public/Input/Events.h b/Engine/Source/Runtime/SlateCore/Public/Input/Events.h
--- a/Engine/Source/Runtime/SlateCore/Public/Input/Events.h
+++ b/Engine/Source/Runtime/SlateCore/Public/Input/Events.h
@@ -703,7 +703,7 @@ public:
 		: ScreenSpacePosition(FVector2f(0.f, 0.f))
 		, LastScreenSpacePosition(FVector2f(0.f, 0.f))
 		, CursorDelta(FVector2f(0.f, 0.f))
-		, PressedButtons(&FTouchKeySet::EmptySet)
+		, PressedButtons(FTouchKeySet::EmptySet)
 		, EffectingButton()
 		, PointerIndex(0)
 		, TouchpadIndex(0)
@@ -730,7 +730,7 @@ public:
 		, ScreenSpacePosition(InScreenSpacePosition)
 		, LastScreenSpacePosition(InLastScreenSpacePosition)
 		, CursorDelta(FVector2f(InScreenSpacePosition) - FVector2f(InLastScreenSpacePosition))
-		, PressedButtons(&InPressedButtons)
+		, PressedButtons(InPressedButtons)
 		, EffectingButton(InEffectingButton)
 		, PointerIndex(InPointerIndex)
 		, TouchpadIndex(0)
@@ -757,7 +757,7 @@ public:
 		, ScreenSpacePosition(InScreenSpacePosition)
 		, LastScreenSpacePosition(InLastScreenSpacePosition)
 		, CursorDelta(FVector2f(InScreenSpacePosition) - FVector2f(InLastScreenSpacePosition))
-		, PressedButtons(&InPressedButtons)
+		, PressedButtons(InPressedButtons)
 		, EffectingButton(InEffectingButton)
 		, PointerIndex(InPointerIndex)
 		, TouchpadIndex(0)
@@ -785,7 +785,7 @@ public:
 		, ScreenSpacePosition(InScreenSpacePosition)
 		, LastScreenSpacePosition(InLastScreenSpacePosition)
 		, CursorDelta(FVector2f(InScreenSpacePosition) - FVector2f(InLastScreenSpacePosition))
-		, PressedButtons(&InPressedButtons)
+		, PressedButtons(InPressedButtons)
 		, EffectingButton(InEffectingButton)
 		, PointerIndex(InPointerIndex)
 		, TouchpadIndex(0)
@@ -816,7 +816,7 @@ public:
 		, ScreenSpacePosition(InScreenSpacePosition)
 		, LastScreenSpacePosition(InLastScreenSpacePosition)
 		, CursorDelta(InDelta)
-		, PressedButtons(&InPressedButtons)
+		, PressedButtons(InPressedButtons)
 		, PointerIndex(InPointerIndex)
 		, TouchpadIndex(0)
 		, Force(1.0f)
@@ -841,7 +841,7 @@ public:
 		, ScreenSpacePosition(InScreenSpacePosition)
 		, LastScreenSpacePosition(InLastScreenSpacePosition)
 		, CursorDelta(InDelta)
-		, PressedButtons(&InPressedButtons)
+		, PressedButtons(InPressedButtons)
 		, PointerIndex(InPointerIndex)
 		, TouchpadIndex(0)
 		, Force(1.0f)
@@ -869,7 +869,7 @@ public:
 		, ScreenSpacePosition(InScreenSpacePosition)
 		, LastScreenSpacePosition(InLastScreenSpacePosition)
 		, CursorDelta(FVector2f(InScreenSpacePosition) - FVector2f(InLastScreenSpacePosition))
-		, PressedButtons(bPressLeftMouseButton ? &FTouchKeySet::StandardSet : &FTouchKeySet::EmptySet)
+		, PressedButtons(bPressLeftMouseButton ? FTouchKeySet::StandardSet : FTouchKeySet::EmptySet)
 		, EffectingButton(EKeys::LeftMouseButton)
 		, PointerIndex(InPointerIndex)
 		, TouchpadIndex(InTouchpadIndex)
@@ -899,7 +899,7 @@ public:
 		, ScreenSpacePosition(InScreenSpacePosition)
 		, LastScreenSpacePosition(InLastScreenSpacePosition)
 		, CursorDelta(FVector2f(InScreenSpacePosition) - FVector2f(InLastScreenSpacePosition))
-		, PressedButtons(bPressLeftMouseButton ? &FTouchKeySet::StandardSet : &FTouchKeySet::EmptySet)
+		, PressedButtons(bPressLeftMouseButton ? FTouchKeySet::StandardSet : FTouchKeySet::EmptySet)
 		, EffectingButton(EKeys::LeftMouseButton)
 		, PointerIndex(InPointerIndex)
 		, TouchpadIndex(InTouchpadIndex)
@@ -931,7 +931,7 @@ public:
 		, ScreenSpacePosition(InScreenSpacePosition)
 		, LastScreenSpacePosition(InLastScreenSpacePosition)
 		, CursorDelta(FVector2f(LastScreenSpacePosition) - FVector2f(ScreenSpacePosition))
-		, PressedButtons(&InPressedButtons)
+		, PressedButtons(InPressedButtons)
 		, PointerIndex(0)
 		, TouchpadIndex(0)
 		, Force(1.0f)
@@ -968,7 +968,7 @@ public:
 	const UE::Slate::FDeprecateVector2DResult& GetCursorDelta() const { return CursorDelta; }
 
 	/** Mouse buttons that are currently pressed */
-	bool IsMouseButtonDown( FKey MouseButton ) const { return PressedButtons->Contains( MouseButton ); }
+	bool IsMouseButtonDown( FKey MouseButton ) const { return PressedButtons.Contains( MouseButton ); }
 
 	/** Mouse button that caused this event to be raised (possibly FKey::Invalid) */
 	FKey GetEffectingButton() const { return EffectingButton; }
@@ -1004,7 +1004,7 @@ public:
 	bool IsDirectionInvertedFromDevice() const { return bIsDirectionInvertedFromDevice; }
 
 	/** Returns the full set of pressed buttons */
-	const TSet<FKey>& GetPressedButtons() const { return *PressedButtons; }
+	const TSet<FKey>& GetPressedButtons() const { return PressedButtons; }
 
 	SLATECORE_API virtual FText ToText() const override;
 
@@ -1025,7 +1025,7 @@ private:
 	UE::Slate::FDeprecateVector2DResult ScreenSpacePosition;
 	UE::Slate::FDeprecateVector2DResult LastScreenSpacePosition;
 	UE::Slate::FDeprecateVector2DResult CursorDelta;
-	const TSet<FKey>* PressedButtons;
+	TSet<FKey> PressedButtons;
 	FKey EffectingButton;
 	uint32 PointerIndex;
 	uint32 TouchpadIndex;
