[Input] Have FPointerEvent::PressedButtons copy by value instead of store a pointer so that it will be thread safe in the future if there are input events are safe to be copied across different threads.

#was-backout 52323349
#rnx
#Synced-CL 52309720

diff --git a/Engine/Source/Runtime/Slate/Private/Framework/Application/SlateApplication.cpp b/Engine/Source/Runtime/Slate/Private/Framework/Application/SlateApplication.cpp
--- a/Engine/Source/Runtime/Slate/Private/Framework/Application/SlateApplication.cpp
+++ b/Engine/Source/Runtime/Slate/Private/Framework/Application/SlateApplication.cpp
@@ -5219,6 +5219,15 @@ bool FSlateApplication::OnMouseDown( const TSharedPtr< FGenericWindow >& Platfor
 
 	FKey Key = TranslateMouseButtonToKey( Button );
 
+	// Update PressedMouseButtons before constructing the event so that the value-copy
+	// snapshot in FPointerEvent reflects the post-press state. The matching Add inside
+	// ProcessMouseButtonDownEvent becomes a no-op for this path (TSet is idempotent)
+	// but is still needed for touch events that call ProcessMouseButtonDownEvent directly.
+	if (GetUserIndexForMouse() == CursorUserIndex)
+	{
+		PressedMouseButtons.Add(Key);
+	}
+
 	FPointerEvent MouseEvent(
 		GetUserIndexForMouse(),
 		CursorPointerIndex,
@@ -5943,6 +5952,12 @@ bool FSlateApplication::OnMouseDoubleClick( const TSharedPtr< FGenericWindow >&
 
 	FKey Key = TranslateMouseButtonToKey( Button );
 
+	// Update PressedMouseButtons before constructing the event — same reasoning as OnMouseDown.
+	if (GetUserIndexForMouse() == CursorUserIndex)
+	{
+		PressedMouseButtons.Add(Key);
+	}
+
 	FPointerEvent MouseEvent(
 		GetUserIndexForMouse(),
 		CursorPointerIndex,
@@ -6035,6 +6050,15 @@ bool FSlateApplication::OnMouseUp( const EMouseButtons::Type Button, const FVect
 
 	FKey Key = TranslateMouseButtonToKey( Button );
 
+	// Update PressedMouseButtons before constructing the event so the value-copy snapshot
+	// reflects the post-release state (button is no longer down). The matching Remove inside
+	// ProcessMouseButtonUpEvent becomes a no-op for this path but is still needed for touch
+	// events routed via ProcessTouchEndedEvent -> ProcessMouseButtonUpEvent directly.
+	if (GetUserIndexForMouse() == CursorUserIndex)
+	{
+		PressedMouseButtons.Remove(Key);
+	}
+
 	FPointerEvent MouseEvent(
 		GetUserIndexForMouse(),
 		CursorPointerIndex,
@@ -7465,6 +7489,13 @@ EDropEffect::Type FSlateApplication::OnDragDrop(const TSharedPtr< FGenericWindow
 
 	if (GetCursorUser()->IsDragDropping())
 	{
+		// Update PressedMouseButtons before constructing the synthesized mouse-up event so
+		// the value-copy snapshot reflects the post-release state. Same reasoning as OnMouseUp.
+		if (GetUserIndexForMouse() == CursorUserIndex)
+		{
+			PressedMouseButtons.Remove(EKeys::LeftMouseButton);
+		}
+
 		FPointerEvent MouseEvent(
 			GetUserIndexForMouse(),
 			CursorPointerIndex,
diff --git a/Engine/Source/Runtime/SlateCore/Public/Input/Events.h b/Engine/Source/Runtime/SlateCore/Public/Input/Events.h
--- a/Engine/Source/Runtime/SlateCore/Public/Input/Events.h
+++ b/Engine/Source/Runtime/SlateCore/Public/Input/Events.h
@@ -703,7 +703,7 @@ public:
 		: ScreenSpacePosition(FVector2f(0.f, 0.f))
 		, LastScreenSpacePosition(FVector2f(0.f, 0.f))
 		, CursorDelta(FVector2f(0.f, 0.f))
-		, PressedButtons(&FTouchKeySet::EmptySet)
+		, PressedButtons(FTouchKeySet::EmptySet)
 		, EffectingButton()
 		, PointerIndex(0)
 		, TouchpadIndex(0)
@@ -730,7 +730,7 @@ public:
 		, ScreenSpacePosition(InScreenSpacePosition)
 		, LastScreenSpacePosition(InLastScreenSpacePosition)
 		, CursorDelta(FVector2f(InScreenSpacePosition) - FVector2f(InLastScreenSpacePosition))
-		, PressedButtons(&InPressedButtons)
+		, PressedButtons(InPressedButtons)
 		, EffectingButton(InEffectingButton)
 		, PointerIndex(InPointerIndex)
 		, TouchpadIndex(0)
@@ -757,7 +757,7 @@ public:
 		, ScreenSpacePosition(InScreenSpacePosition)
 		, LastScreenSpacePosition(InLastScreenSpacePosition)
 		, CursorDelta(FVector2f(InScreenSpacePosition) - FVector2f(InLastScreenSpacePosition))
-		, PressedButtons(&InPressedButtons)
+		, PressedButtons(InPressedButtons)
 		, EffectingButton(InEffectingButton)
 		, PointerIndex(InPointerIndex)
 		, TouchpadIndex(0)
@@ -785,7 +785,7 @@ public:
 		, ScreenSpacePosition(InScreenSpacePosition)
 		, LastScreenSpacePosition(InLastScreenSpacePosition)
 		, CursorDelta(FVector2f(InScreenSpacePosition) - FVector2f(InLastScreenSpacePosition))
-		, PressedButtons(&InPressedButtons)
+		, PressedButtons(InPressedButtons)
 		, EffectingButton(InEffectingButton)
 		, PointerIndex(InPointerIndex)
 		, TouchpadIndex(0)
@@ -816,7 +816,7 @@ public:
 		, ScreenSpacePosition(InScreenSpacePosition)
 		, LastScreenSpacePosition(InLastScreenSpacePosition)
 		, CursorDelta(InDelta)
-		, PressedButtons(&InPressedButtons)
+		, PressedButtons(InPressedButtons)
 		, PointerIndex(InPointerIndex)
 		, TouchpadIndex(0)
 		, Force(1.0f)
@@ -841,7 +841,7 @@ public:
 		, ScreenSpacePosition(InScreenSpacePosition)
 		, LastScreenSpacePosition(InLastScreenSpacePosition)
 		, CursorDelta(InDelta)
-		, PressedButtons(&InPressedButtons)
+		, PressedButtons(InPressedButtons)
 		, PointerIndex(InPointerIndex)
 		, TouchpadIndex(0)
 		, Force(1.0f)
@@ -869,7 +869,7 @@ public:
 		, ScreenSpacePosition(InScreenSpacePosition)
 		, LastScreenSpacePosition(InLastScreenSpacePosition)
 		, CursorDelta(FVector2f(InScreenSpacePosition) - FVector2f(InLastScreenSpacePosition))
-		, PressedButtons(bPressLeftMouseButton ? &FTouchKeySet::StandardSet : &FTouchKeySet::EmptySet)
+		, PressedButtons(bPressLeftMouseButton ? FTouchKeySet::StandardSet : FTouchKeySet::EmptySet)
 		, EffectingButton(EKeys::LeftMouseButton)
 		, PointerIndex(InPointerIndex)
 		, TouchpadIndex(InTouchpadIndex)
@@ -899,7 +899,7 @@ public:
 		, ScreenSpacePosition(InScreenSpacePosition)
 		, LastScreenSpacePosition(InLastScreenSpacePosition)
 		, CursorDelta(FVector2f(InScreenSpacePosition) - FVector2f(InLastScreenSpacePosition))
-		, PressedButtons(bPressLeftMouseButton ? &FTouchKeySet::StandardSet : &FTouchKeySet::EmptySet)
+		, PressedButtons(bPressLeftMouseButton ? FTouchKeySet::StandardSet : FTouchKeySet::EmptySet)
 		, EffectingButton(EKeys::LeftMouseButton)
 		, PointerIndex(InPointerIndex)
 		, TouchpadIndex(InTouchpadIndex)
@@ -931,7 +931,7 @@ public:
 		, ScreenSpacePosition(InScreenSpacePosition)
 		, LastScreenSpacePosition(InLastScreenSpacePosition)
 		, CursorDelta(FVector2f(LastScreenSpacePosition) - FVector2f(ScreenSpacePosition))
-		, PressedButtons(&InPressedButtons)
+		, PressedButtons(InPressedButtons)
 		, PointerIndex(0)
 		, TouchpadIndex(0)
 		, Force(1.0f)
@@ -968,7 +968,7 @@ public:
 	const UE::Slate::FDeprecateVector2DResult& GetCursorDelta() const { return CursorDelta; }
 
 	/** Mouse buttons that are currently pressed */
-	bool IsMouseButtonDown( FKey MouseButton ) const { return PressedButtons->Contains( MouseButton ); }
+	bool IsMouseButtonDown( FKey MouseButton ) const { return PressedButtons.Contains( MouseButton ); }
 
 	/** Mouse button that caused this event to be raised (possibly FKey::Invalid) */
 	FKey GetEffectingButton() const { return EffectingButton; }
@@ -1004,7 +1004,7 @@ public:
 	bool IsDirectionInvertedFromDevice() const { return bIsDirectionInvertedFromDevice; }
 
 	/** Returns the full set of pressed buttons */
-	const TSet<FKey>& GetPressedButtons() const { return *PressedButtons; }
+	const TSet<FKey>& GetPressedButtons() const { return PressedButtons; }
 
 	SLATECORE_API virtual FText ToText() const override;
 
@@ -1025,7 +1025,7 @@ private:
 	UE::Slate::FDeprecateVector2DResult ScreenSpacePosition;
 	UE::Slate::FDeprecateVector2DResult LastScreenSpacePosition;
 	UE::Slate::FDeprecateVector2DResult CursorDelta;
-	const TSet<FKey>* PressedButtons;
+	TSet<FKey> PressedButtons;
 	FKey EffectingButton;
 	uint32 PointerIndex;
 	uint32 TouchpadIndex;
