[Dynamic Invalidation] Track foreground and disabled foreground colors via Slate attributes in `SCompoundWidget`

Overriding `GetForegroundColor` currently does not trigger invalidation when the value changes.
`ForegroundColorAttribute` and a new `DisabledForegroundColorAttribute` are now bound to their respective virtual `GetForegroundColor()`/`GetDisabledForegroundColor()` so the changes are now detected and trigger Paint invalidation properly.
`ExplicitForegroundColor` (a plain `TAttribute`) now holds the user-set value, avoiding a circular dependency between `ForegroundColorAttribute` and `GetForegroundColor()`.

#Synced-CL 52095318

diff --git a/Engine/Source/Runtime/SlateCore/Private/Widgets/SCompoundWidget.cpp b/Engine/Source/Runtime/SlateCore/Private/Widgets/SCompoundWidget.cpp
--- a/Engine/Source/Runtime/SlateCore/Private/Widgets/SCompoundWidget.cpp
+++ b/Engine/Source/Runtime/SlateCore/Private/Widgets/SCompoundWidget.cpp
@@ -15,6 +15,7 @@ void SCompoundWidget::PrivateRegisterAttributes(FSlateAttributeInitializer& Attr
 	SLATE_ADD_MEMBER_ATTRIBUTE_DEFINITION_WITH_NAME(AttributeInitializer, "ContentScale", ContentScaleAttribute, EInvalidateWidgetReason::Layout);
 	SLATE_ADD_MEMBER_ATTRIBUTE_DEFINITION_WITH_NAME(AttributeInitializer, "ColorAndOpacity", ColorAndOpacityAttribute, EInvalidateWidgetReason::Paint);
 	SLATE_ADD_MEMBER_ATTRIBUTE_DEFINITION_WITH_NAME(AttributeInitializer, "ForegroundColor", ForegroundColorAttribute, EInvalidateWidgetReason::Paint);
+	SLATE_ADD_MEMBER_ATTRIBUTE_DEFINITION_WITH_NAME(AttributeInitializer, "DisabledForegroundColor", DisabledForegroundColorAttribute, EInvalidateWidgetReason::Paint);
 }
 
 int32 SCompoundWidget::OnPaint( const FPaintArgs& Args, const FGeometry& AllottedGeometry, const FSlateRect& MyCullingRect, FSlateWindowElementList& OutDrawElements, int32 LayerId, const FWidgetStyle& InWidgetStyle, bool bParentEnabled ) const
@@ -36,7 +37,7 @@ int32 SCompoundWidget::OnPaint( const FPaintArgs& Args, const FGeometry& Allotte
 
 		FWidgetStyle CompoundedWidgetStyle = FWidgetStyle(InWidgetStyle)
 			.BlendColorAndOpacityTint(GetColorAndOpacity())
-			.SetForegroundColor(bShouldBeEnabled ? GetForegroundColor() : GetDisabledForegroundColor() );
+			.SetForegroundColor(bShouldBeEnabled ? ForegroundColorAttribute.Get() : DisabledForegroundColorAttribute.Get() );
 
 		int32 Layer = 0;
 		{
@@ -75,7 +76,13 @@ void SCompoundWidget::OnArrangeChildren( const FGeometry& AllottedGeometry, FArr
 
 FSlateColor SCompoundWidget::GetForegroundColor() const
 {
-	return ForegroundColorAttribute.Get();
+	// Read from the plain TAttribute, NOT from ForegroundColorAttribute, to avoid circular dependency (ForegroundColorAttribute is bound to this very function).
+	return ExplicitForegroundColor.Get(FSlateColor::UseForeground());
+}
+
+void SCompoundWidget::SetForegroundColor(TAttribute<FSlateColor> InForegroundColor)
+{
+	ExplicitForegroundColor = MoveTemp(InForegroundColor);
 }
 
 SCompoundWidget::SCompoundWidget()
@@ -83,7 +90,11 @@ SCompoundWidget::SCompoundWidget()
 	, ContentScaleAttribute( *this, FVector2D(1.0f,1.0f) )
 	, ColorAndOpacityAttribute( *this, FLinearColor::White )
 	, ForegroundColorAttribute( *this, FSlateColor::UseForeground() )
+	, DisabledForegroundColorAttribute( *this, FSlateColor::UseForeground() )
+	, ExplicitForegroundColor( FSlateColor::UseForeground() )
 {
+	ForegroundColorAttribute.Assign(*this, TAttribute<FSlateColor>::CreateRaw(this, &SCompoundWidget::GetForegroundColor));
+	DisabledForegroundColorAttribute.Assign(*this, TAttribute<FSlateColor>::CreateRaw(this, &SCompoundWidget::GetDisabledForegroundColor));
 }
 
 SCompoundWidget::~SCompoundWidget() = default;
diff --git a/Engine/Source/Runtime/SlateCore/Public/Widgets/SCompoundWidget.h b/Engine/Source/Runtime/SlateCore/Public/Widgets/SCompoundWidget.h
--- a/Engine/Source/Runtime/SlateCore/Public/Widgets/SCompoundWidget.h
+++ b/Engine/Source/Runtime/SlateCore/Public/Widgets/SCompoundWidget.h
@@ -29,7 +29,7 @@ public:
 	 *
 	 * @return Size scaling factor.
 	 */
-	const FVector2D GetContentScale() const
+	FVector2D GetContentScale() const
 	{
 		return ContentScaleAttribute.Get();
 	}
@@ -65,12 +65,9 @@ public:
 	/**
 	 * Sets the widget's foreground color.
 	 *
-	 * @param InColor The color to set.
+	 * @param InForegroundColor The color to set.
 	 */
-	void SetForegroundColor( TAttribute<FSlateColor> InForegroundColor )
-	{
-		ForegroundColorAttribute.Assign(*this, MoveTemp(InForegroundColor));
-	}
+	SLATECORE_API void SetForegroundColor( TAttribute<FSlateColor> InForegroundColor );
 
 public:
 
@@ -95,7 +92,7 @@ protected:
 	TSlateAttributeRef<FLinearColor> GetColorAndOpacityAttribute() const { return TSlateAttributeRef<FLinearColor>{SharedThis(this), ColorAndOpacityAttribute}; }
 
 	/** @return an attribute reference of ForegroundColor */
-	TSlateAttributeRef<FSlateColor> GetForegroundColorAttribute() const {	return TSlateAttributeRef<FSlateColor>{SharedThis(this), ForegroundColorAttribute}; }
+	TSlateAttributeRef<FSlateColor> GetForegroundColorAttribute() const { return TSlateAttributeRef<FSlateColor>{SharedThis(this), ForegroundColorAttribute}; }
 
 protected:
 
@@ -119,6 +116,12 @@ private:
 	/** The color and opacity to apply to this widget and all its descendants. */
 	TSlateAttribute<FLinearColor> ColorAndOpacityAttribute;
 
-	/** Optional foreground color that will be inherited by all of this widget's contents */
+	/** Computed foreground color attribute, always bound to the virtual GetForegroundColor(). */
 	TSlateAttribute<FSlateColor> ForegroundColorAttribute;
+
+	/** Computed disabled foreground color attribute, always bound to the virtual GetDisabledForegroundColor(). */
+	TSlateAttribute<FSlateColor> DisabledForegroundColorAttribute;
+
+	/** Optional foreground color that will be inherited by all of this widget's contents. */
+	TAttribute<FSlateColor> ExplicitForegroundColor;
 };
