Add error/fatal logs to catch invalid parent / child widget relationnship before it ends up crashing.

#tests Ran the editor and the game to validate that the fatal log wasn't triggered in obvious cases. Used the debugger to force the new code path to be executed and validated that the message was displayed correctly (both editor and ingame)
#rnx
#Synced-CL 52573034

diff --git a/Engine/Source/Runtime/SlateCore/Private/FastUpdate/SlateInvalidationWidgetList.cpp b/Engine/Source/Runtime/SlateCore/Private/FastUpdate/SlateInvalidationWidgetList.cpp
--- a/Engine/Source/Runtime/SlateCore/Private/FastUpdate/SlateInvalidationWidgetList.cpp
+++ b/Engine/Source/Runtime/SlateCore/Private/FastUpdate/SlateInvalidationWidgetList.cpp
@@ -482,7 +482,18 @@ FSlateInvalidationWidgetIndex FSlateInvalidationWidgetList::Internal_BuildWidget
 				const bool bShouldAdd = ShouldBeAdded(NewWidget);
 				if (bShouldAdd)
 				{
-					check(NewWidget.GetParentWidget().Get() == &Widget);
+					if (NewWidget.GetParentWidget().Get() != &Widget)
+					{
+						TStringBuilder<512> Builder;
+						Builder << TEXT("The widget '") << FReflectionMetaData::GetWidgetDebugInfo(NewWidget) << TEXT("' has the widget '")
+						<< FReflectionMetaData::GetWidgetPath(NewWidget.GetParentWidget().Get()) << TEXT("' as parent, but is also a child of widget '")
+						<< FReflectionMetaData::GetWidgetPath(Widget) << ("' This is not supported and may eventually crash. Please fix.");
+#if UE_BUILD_SHIPPING
+						UE_LOG(LogSlate, Error, TEXT("%s"), Builder.ToString());
+#else
+						UE_LOG(LogSlate, Fatal, TEXT("%s"), Builder.ToString());
+#endif
+					}
 					LeafMostChildIndex = Internal_BuildWidgetList_Recursive(NewWidget, NewIndex, LastestIndex, NewVisibility, bParentOrSelfVolatile);
 				}
 			});
