[Dynamic Invalidation] Fix `SCompoundWidget` foreground color performance regression from CL 52461870

CL 52461870 attempted to fix missing Paint invalidation when `GetForegroundColor()` is overridden in `SCompoundWidget` subclasses. It did so by binding `ForegroundColorAttribute` and `DisabledForegroundColorAttribute` to their virtual getters in the constructor, so the invalidation system could poll them for changes. This caused a significant performance regression by polling two virtual functions on every widget, every frame.

**This CL reverts `SCompoundWidget` to its original design** and fixes the affected widgets individually using a pattern compatible with dynamic invalidation. Note that only widgets already supporting invalidation were adjusted; other widgets are unaffected and continue to work as before.

- If the foreground color is a known value, call `SetForegroundColor(value)` directly whenever it changes.
- Otherwise, do **not** override `GetForegroundColor()`. Instead, implement a separate helper (by convention, `GetForegroundColorInternal()`) and bind it as a delegate in `Construct`:
```cpp
SetForegroundColor(TAttribute<FSlateColor>::CreateRaw(this, &SMyWidget::GetForegroundColorInternal));
```
This lets `SCompoundWidget::GetForegroundColor()` remain as-is (returning the tracked `ForegroundColorAttribute`) so the invalidation system can detect changes correctly.

`GetDisabledForegroundColor()` follows the same principle but is not backed by a `TSlateAttribute` (as only `SButton` overrides it, and it handles its own invalidation). Widgets that override `GetDisabledForegroundColor()` are responsible for triggering Paint invalidation themselves when the value changes.

---

### Rule for future invalidation-compatible widgets with custom foreground color logic

1. Do **not** override `SCompoundWidget::GetForegroundColor()`.
2. Implement `GetForegroundColorInternal()` with the custom logic.
3. In `Construct`, call `SetForegroundColor(TAttribute<FSlateColor>::CreateRaw(this, &SMyWidget::GetForegroundColorInternal))`.

#Synced-CL 52471066

diff --git a/Engine/Source/Editor/GraphEditor/Private/SCommentBubble.cpp b/Engine/Source/Editor/GraphEditor/Private/SCommentBubble.cpp
--- a/Engine/Source/Editor/GraphEditor/Private/SCommentBubble.cpp
+++ b/Engine/Source/Editor/GraphEditor/Private/SCommentBubble.cpp
@@ -79,7 +79,7 @@ void SCommentBubble::Construct( const FArguments& InArgs )
 	ToggleButtonCheck		= InArgs._ToggleButtonCheck.IsBound() ?	InArgs._ToggleButtonCheck : 
 																	TAttribute<ECheckBoxState>( this, &SCommentBubble::GetToggleButtonCheck );
 	// Ensue this value is set to something sensible
-	CalculatedForegroundColor = FStyleColors::Foreground;
+	SetForegroundColor(FStyleColors::Foreground);
 	BubbleLuminance = 0.0f;
 
 	// Cache the comment
@@ -115,7 +115,7 @@ void SCommentBubble::Tick( const FGeometry& AllottedGeometry, const double InCur
 	{
 		const FLinearColor BubbleColor = GetBubbleColor().GetSpecifiedColor() * SCommentBubbleDefs::LuminanceCoEff;
 		BubbleLuminance = BubbleColor.R + BubbleColor.G + BubbleColor.B;
-		CalculatedForegroundColor = BubbleLuminance < 0.5f ? FStyleColors::Foreground : FStyleColors::Background;
+		SetForegroundColor(BubbleLuminance < 0.5f ? FStyleColors::Foreground : FStyleColors::Background);
 	}
 
 	TickVisibility(InCurrentTime, InDeltaTime);
diff --git a/Engine/Source/Editor/GraphEditor/Public/SCommentBubble.h b/Engine/Source/Editor/GraphEditor/Public/SCommentBubble.h
--- a/Engine/Source/Editor/GraphEditor/Public/SCommentBubble.h
+++ b/Engine/Source/Editor/GraphEditor/Public/SCommentBubble.h
@@ -158,9 +158,6 @@ protected:
 	/** Returns the color of the main bubble */
 	UE_API FSlateColor GetBubbleColor() const;
 
-	/** Returns the foreground color for the text and buttons, taking into account the bubble color */
-	FSlateColor GetForegroundColor() const override { return CalculatedForegroundColor; }
-
 	/** Called when the comment text is committed */
 	UE_API void OnCommentTextCommitted( const FText& NewText, ETextCommit::Type CommitInfo );
 
@@ -201,8 +198,6 @@ protected:
 	/** Delegate to determine if the graph node is currently hovered */
 	FIsGraphNodeHovered IsGraphNodeHovered;
 
-	/** Current Foreground Color */
-	FSlateColor CalculatedForegroundColor;
 	/** The luminance (R + G + B) of the bubble's color, used to control text foreground color */
 	float BubbleLuminance;
 	/** Allow pin behaviour */
diff --git a/Engine/Source/Editor/SceneOutliner/Private/ActorDescTreeItem.cpp b/Engine/Source/Editor/SceneOutliner/Private/ActorDescTreeItem.cpp
--- a/Engine/Source/Editor/SceneOutliner/Private/ActorDescTreeItem.cpp
+++ b/Engine/Source/Editor/SceneOutliner/Private/ActorDescTreeItem.cpp
@@ -61,7 +61,7 @@ struct SActorDescTreeLabel : FSceneOutlinerCommonLabelData, public SCompoundWidg
 				.Text(this, &SActorDescTreeLabel::GetDisplayText)
 				.ToolTipText(this, &SActorDescTreeLabel::GetTooltipText)
 				.HighlightText(HighlightText)
-				.ColorAndOpacity(this, &SActorDescTreeLabel::GetForegroundColor)
+				.ColorAndOpacity(this, &SActorDescTreeLabel::GetForegroundColorInternal)
 				.OnTextCommitted(this, &SActorDescTreeLabel::OnLabelCommitted)
 				.OnVerifyTextChanged(this, &SActorDescTreeLabel::OnVerifyItemLabelChanged)
 				.IsSelected(FIsSelected::CreateSP(&InRow, &STableRow<FSceneOutlinerTreeItemPtr>::IsSelectedExclusively))
@@ -115,6 +115,7 @@ struct SActorDescTreeLabel : FSceneOutlinerCommonLabelData, public SCompoundWidg
 					MainContent
 				]
 			];
+		SetForegroundColor(TAttribute<FSlateColor>::CreateRaw(this, &SActorDescTreeLabel::GetForegroundColorInternal));
 	}
 
 private:
@@ -197,7 +198,11 @@ private:
 		return FText();
 	}
 
-	virtual FSlateColor GetForegroundColor() const override
+	/**
+	 * Returns the Foreground Color of this widget.
+	 * Used to avoid overriding SCompoundWidget::GetForegroundColor for invalidation to work properly.
+	 */
+	FSlateColor GetForegroundColorInternal() const
 	{
 		if (const auto TreeItem = TreeItemPtr.Pin())
 		{
diff --git a/Engine/Source/Editor/SceneOutliner/Private/ActorFolderTreeItem.cpp b/Engine/Source/Editor/SceneOutliner/Private/ActorFolderTreeItem.cpp
--- a/Engine/Source/Editor/SceneOutliner/Private/ActorFolderTreeItem.cpp
+++ b/Engine/Source/Editor/SceneOutliner/Private/ActorFolderTreeItem.cpp
@@ -44,7 +44,7 @@ struct SActorFolderTreeLabel : FSceneOutlinerCommonLabelData, public SCompoundWi
 			.Text(this, &SActorFolderTreeLabel::GetDisplayText)
 			.ToolTipText(this, &SActorFolderTreeLabel::GetTooltipText)
 			.HighlightText(SceneOutliner.GetFilterHighlightText())
-			.ColorAndOpacity(this, &SActorFolderTreeLabel::GetForegroundColor)
+			.ColorAndOpacity(this, &SActorFolderTreeLabel::GetForegroundColorInternal)
 			.OnTextCommitted(this, &SActorFolderTreeLabel::OnLabelCommitted)
 			.OnVerifyTextChanged(this, &SActorFolderTreeLabel::OnVerifyItemLabelChanged)
 			.OnEnterEditingMode(this, &SActorFolderTreeLabel::OnEnterEditingMode)
@@ -84,6 +84,7 @@ struct SActorFolderTreeLabel : FSceneOutlinerCommonLabelData, public SCompoundWi
 					InlineTextBlock.ToSharedRef()
 				]
 			];
+		SetForegroundColor(TAttribute<FSlateColor>::CreateRaw(this, &SActorFolderTreeLabel::GetForegroundColorInternal));
 	}
 
 private:
@@ -145,7 +146,11 @@ private:
 		}
 	}
 
-	FSlateColor GetForegroundColor() const
+	/**
+	 * Returns the Foreground Color of this widget.
+	 * Used to avoid overriding SCompoundWidget::GetForegroundColor for invalidation to work properly.
+	 */
+	FSlateColor GetForegroundColorInternal() const
 	{
 		if (auto BaseColor = FSceneOutlinerCommonLabelData::GetForegroundColor(*TreeItemPtr.Pin()))
 		{
diff --git a/Engine/Source/Editor/SceneOutliner/Private/ActorTreeItem.cpp b/Engine/Source/Editor/SceneOutliner/Private/ActorTreeItem.cpp
--- a/Engine/Source/Editor/SceneOutliner/Private/ActorTreeItem.cpp
+++ b/Engine/Source/Editor/SceneOutliner/Private/ActorTreeItem.cpp
@@ -62,7 +62,7 @@ struct SActorTreeLabel : FSceneOutlinerCommonLabelData, public SCompoundWidget
 				.Text(this, &SActorTreeLabel::GetDisplayText)
 				.ToolTipText(this, &SActorTreeLabel::GetTooltipText)
 				.HighlightText(HighlightText)
-				.ColorAndOpacity(this, &SActorTreeLabel::GetForegroundColor)
+				.ColorAndOpacity(this, &SActorTreeLabel::GetForegroundColorInternal)
 				.OnTextCommitted(this, &SActorTreeLabel::OnLabelCommitted)
 				.OnVerifyTextChanged(this, &SActorTreeLabel::OnVerifyItemLabelChanged)
 				.OnEnterEditingMode(this, &SActorTreeLabel::OnEnterEditingMode)
@@ -105,6 +105,7 @@ struct SActorTreeLabel : FSceneOutlinerCommonLabelData, public SCompoundWidget
 					MainContent
 				]
 			];
+		SetForegroundColor(TAttribute<FSlateColor>::CreateRaw(this, &SActorTreeLabel::GetForegroundColorInternal));
 	}
 
 private:
@@ -235,7 +236,11 @@ private:
 		return ToolTipText;
 	}
 
-	FSlateColor GetForegroundColor() const
+	/**
+	 * Returns the Foreground Color of this widget.
+	 * Used to avoid overriding SCompoundWidget::GetForegroundColor for invalidation to work properly.
+	 */
+	FSlateColor GetForegroundColorInternal() const
 	{
 		AActor* Actor = ActorPtr.Get();
 
diff --git a/Engine/Source/Editor/SceneOutliner/Private/LevelTreeItem.cpp b/Engine/Source/Editor/SceneOutliner/Private/LevelTreeItem.cpp
--- a/Engine/Source/Editor/SceneOutliner/Private/LevelTreeItem.cpp
+++ b/Engine/Source/Editor/SceneOutliner/Private/LevelTreeItem.cpp
@@ -60,9 +60,10 @@ struct SLevelTreeLabel : FSceneOutlinerCommonLabelData, public SCompoundWidget
 				SNew(STextBlock)
 				.Text(this, &SLevelTreeLabel::GetDisplayText)
 				.HighlightText(SceneOutliner.GetFilterHighlightText())
-				.ColorAndOpacity(this, &SLevelTreeLabel::GetForegroundColor)
+				.ColorAndOpacity(this, &SLevelTreeLabel::GetForegroundColorInternal)
 			]
 		];
+		SetForegroundColor(TAttribute<FSlateColor>::CreateRaw(this, &SLevelTreeLabel::GetForegroundColorInternal));
 	}
 
 private:
@@ -74,7 +75,11 @@ private:
 		return Item.IsValid() ? FText::FromString(Item->GetDisplayString()) : FText();
 	}
 
-	FSlateColor GetForegroundColor() const
+	/**
+	 * Returns the Foreground Color of this widget.
+	 * Used to avoid overriding SCompoundWidget::GetForegroundColor for invalidation to work properly.
+	 */
+	FSlateColor GetForegroundColorInternal() const
 	{
 		if (auto BaseColor = FSceneOutlinerCommonLabelData::GetForegroundColor(*TreeItemPtr.Pin()))
 		{
diff --git a/Engine/Source/Editor/SceneOutliner/Private/WorldTreeItem.cpp b/Engine/Source/Editor/SceneOutliner/Private/WorldTreeItem.cpp
--- a/Engine/Source/Editor/SceneOutliner/Private/WorldTreeItem.cpp
+++ b/Engine/Source/Editor/SceneOutliner/Private/WorldTreeItem.cpp
@@ -121,7 +121,7 @@ struct SWorldTreeLabel : FSceneOutlinerCommonLabelData, public SCompoundWidget
 				SNew(STextBlock)
 				.Text(this, &SWorldTreeLabel::GetDisplayText)
 				.HighlightText(SceneOutliner.GetFilterHighlightText())
-				.ColorAndOpacity(this, &SWorldTreeLabel::GetForegroundColor)
+				.ColorAndOpacity(this, &SWorldTreeLabel::GetForegroundColorInternal)
 				.OnDoubleClicked(this, &SWorldTreeLabel::HandleDoubleClick)
 				.ToolTip(IDocumentation::Get()->CreateToolTip(
 					TAttribute<FText>(this, &SWorldTreeLabel::GetTooltipText),
@@ -131,6 +131,7 @@ struct SWorldTreeLabel : FSceneOutlinerCommonLabelData, public SCompoundWidget
 				))
 			]
 		];
+		SetForegroundColor(TAttribute<FSlateColor>::CreateRaw(this, &SWorldTreeLabel::GetForegroundColorInternal));
 	}
 
 private:
@@ -156,7 +157,11 @@ private:
 		}
 	}
 
-	FSlateColor GetForegroundColor() const
+	/**
+	 * Returns the Foreground Color of this widget.
+	 * Used to avoid overriding SCompoundWidget::GetForegroundColor for invalidation to work properly.
+	 */
+	FSlateColor GetForegroundColorInternal() const
 	{
 		if (auto BaseColor = FSceneOutlinerCommonLabelData::GetForegroundColor(*TreeItemPtr.Pin()))
 		{
diff --git a/Engine/Source/Runtime/Slate/Private/Widgets/Input/SCheckBox.cpp b/Engine/Source/Runtime/Slate/Private/Widgets/Input/SCheckBox.cpp
--- a/Engine/Source/Runtime/Slate/Private/Widgets/Input/SCheckBox.cpp
+++ b/Engine/Source/Runtime/Slate/Private/Widgets/Input/SCheckBox.cpp
@@ -85,6 +85,7 @@ void SCheckBox::Construct( const SCheckBox::FArguments& InArgs )
 	SetHoveredSound(InArgs._HoveredSoundOverride);
 	SetCheckedSound(InArgs._CheckedSoundOverride);
 	SetUncheckedSound(InArgs._UncheckedSoundOverride);
+	SetForegroundColor(TAttribute<FSlateColor>::CreateRaw(this, &SCheckBox::GetForegroundColorInternal));
 }
 
 /**
@@ -304,7 +305,7 @@ bool SCheckBox::IsInteractable() const
 	return IsEnabled();
 }
 
-FSlateColor SCheckBox::GetForegroundColor() const
+FSlateColor SCheckBox::GetForegroundColorInternal() const
 {
 	FSlateColor UserColor = ForegroundColorOverrideAttribute.Get();
 
diff --git a/Engine/Source/Runtime/Slate/Public/Widgets/Input/SCheckBox.h b/Engine/Source/Runtime/Slate/Public/Widgets/Input/SCheckBox.h
--- a/Engine/Source/Runtime/Slate/Public/Widgets/Input/SCheckBox.h
+++ b/Engine/Source/Runtime/Slate/Public/Widgets/Input/SCheckBox.h
@@ -183,7 +183,6 @@ public:
 	SLATE_API virtual void OnMouseEnter( const FGeometry& MyGeometry, const FPointerEvent& MouseEvent ) override;
 	SLATE_API virtual void OnMouseLeave( const FPointerEvent& MouseEvent ) override;
 	SLATE_API virtual bool IsInteractable() const override;
-	SLATE_API virtual FSlateColor GetForegroundColor() const override;
 #if WITH_ACCESSIBILITY
 	SLATE_API virtual TSharedRef<FSlateAccessibleWidget> CreateAccessibleWidget() override;
 #endif
@@ -319,6 +318,13 @@ protected:
 	void SetUncheckedSound(TOptional<FSlateSound> InUncheckedSound);
 	FSlateSound GetUncheckedSound() const;
 
+private:
+	/**
+	 * Returns the Foreground Color of this widget.
+	 * Used to avoid overriding SCompoundWidget::GetForegroundColor for invalidation to work properly.
+	 */
+	FSlateColor GetForegroundColorInternal() const;
+
 protected:
 	
 	const FCheckBoxStyle* Style;
diff --git a/Engine/Source/Runtime/SlateCore/Private/Widgets/SCompoundWidget.cpp b/Engine/Source/Runtime/SlateCore/Private/Widgets/SCompoundWidget.cpp
--- a/Engine/Source/Runtime/SlateCore/Private/Widgets/SCompoundWidget.cpp
+++ b/Engine/Source/Runtime/SlateCore/Private/Widgets/SCompoundWidget.cpp
@@ -15,7 +15,6 @@ void SCompoundWidget::PrivateRegisterAttributes(FSlateAttributeInitializer& Attr
 	SLATE_ADD_MEMBER_ATTRIBUTE_DEFINITION_WITH_NAME(AttributeInitializer, "ContentScale", ContentScaleAttribute, EInvalidateWidgetReason::Layout);
 	SLATE_ADD_MEMBER_ATTRIBUTE_DEFINITION_WITH_NAME(AttributeInitializer, "ColorAndOpacity", ColorAndOpacityAttribute, EInvalidateWidgetReason::Paint);
 	SLATE_ADD_MEMBER_ATTRIBUTE_DEFINITION_WITH_NAME(AttributeInitializer, "ForegroundColor", ForegroundColorAttribute, EInvalidateWidgetReason::Paint);
-	SLATE_ADD_MEMBER_ATTRIBUTE_DEFINITION_WITH_NAME(AttributeInitializer, "DisabledForegroundColor", DisabledForegroundColorAttribute, EInvalidateWidgetReason::Paint);
 }
 
 int32 SCompoundWidget::OnPaint( const FPaintArgs& Args, const FGeometry& AllottedGeometry, const FSlateRect& MyCullingRect, FSlateWindowElementList& OutDrawElements, int32 LayerId, const FWidgetStyle& InWidgetStyle, bool bParentEnabled ) const
@@ -37,7 +36,7 @@ int32 SCompoundWidget::OnPaint( const FPaintArgs& Args, const FGeometry& Allotte
 
 		FWidgetStyle CompoundedWidgetStyle = FWidgetStyle(InWidgetStyle)
 			.BlendColorAndOpacityTint(GetColorAndOpacity())
-			.SetForegroundColor(bShouldBeEnabled ? ForegroundColorAttribute.Get() : DisabledForegroundColorAttribute.Get() );
+			.SetForegroundColor(bShouldBeEnabled ? GetForegroundColor() : GetDisabledForegroundColor() );
 
 		int32 Layer = 0;
 		{
@@ -76,13 +75,7 @@ void SCompoundWidget::OnArrangeChildren( const FGeometry& AllottedGeometry, FArr
 
 FSlateColor SCompoundWidget::GetForegroundColor() const
 {
-	// Read from the plain TAttribute, NOT from ForegroundColorAttribute, to avoid circular dependency (ForegroundColorAttribute is bound to this very function).
-	return ExplicitForegroundColor.Get(FSlateColor::UseForeground());
-}
-
-void SCompoundWidget::SetForegroundColor(TAttribute<FSlateColor> InForegroundColor)
-{
-	ExplicitForegroundColor = MoveTemp(InForegroundColor);
+	return ForegroundColorAttribute.Get();
 }
 
 SCompoundWidget::SCompoundWidget()
@@ -90,11 +83,7 @@ SCompoundWidget::SCompoundWidget()
 	, ContentScaleAttribute( *this, FVector2D(1.0f,1.0f) )
 	, ColorAndOpacityAttribute( *this, FLinearColor::White )
 	, ForegroundColorAttribute( *this, FSlateColor::UseForeground() )
-	, DisabledForegroundColorAttribute( *this, FSlateColor::UseForeground() )
-	, ExplicitForegroundColor( FSlateColor::UseForeground() )
 {
-	ForegroundColorAttribute.Assign(*this, TAttribute<FSlateColor>::CreateRaw(this, &SCompoundWidget::GetForegroundColor));
-	DisabledForegroundColorAttribute.Assign(*this, TAttribute<FSlateColor>::CreateRaw(this, &SCompoundWidget::GetDisabledForegroundColor));
 }
 
 SCompoundWidget::~SCompoundWidget() = default;
diff --git a/Engine/Source/Runtime/SlateCore/Public/Widgets/SCompoundWidget.h b/Engine/Source/Runtime/SlateCore/Public/Widgets/SCompoundWidget.h
--- a/Engine/Source/Runtime/SlateCore/Public/Widgets/SCompoundWidget.h
+++ b/Engine/Source/Runtime/SlateCore/Public/Widgets/SCompoundWidget.h
@@ -67,7 +67,10 @@ public:
 	 *
 	 * @param InForegroundColor The color to set.
 	 */
-	SLATECORE_API void SetForegroundColor( TAttribute<FSlateColor> InForegroundColor );
+	void SetForegroundColor( TAttribute<FSlateColor> InForegroundColor )
+	{
+		ForegroundColorAttribute.Assign(*this, MoveTemp(InForegroundColor));
+	}
 
 public:
 
@@ -116,12 +119,6 @@ private:
 	/** The color and opacity to apply to this widget and all its descendants. */
 	TSlateAttribute<FLinearColor> ColorAndOpacityAttribute;
 
-	/** Computed foreground color attribute, always bound to the virtual GetForegroundColor(). */
+	/** Optional foreground color that will be inherited by all of this widget's contents */
 	TSlateAttribute<FSlateColor> ForegroundColorAttribute;
-
-	/** Computed disabled foreground color attribute, always bound to the virtual GetDisabledForegroundColor(). */
-	TSlateAttribute<FSlateColor> DisabledForegroundColorAttribute;
-
-	/** Optional foreground color that will be inherited by all of this widget's contents. */
-	TAttribute<FSlateColor> ExplicitForegroundColor;
 };
