PR #13899: Optimized Slate Timer update by using a single loop

#tests Ran the editor, played with a few tools, stepped in the code to validate timer were called properly
#Synced-CL 52573034

diff --git a/Engine/Source/Runtime/SlateCore/Private/Application/SlateApplicationBase.cpp b/Engine/Source/Runtime/SlateCore/Private/Application/SlateApplicationBase.cpp
--- a/Engine/Source/Runtime/SlateCore/Private/Application/SlateApplicationBase.cpp
+++ b/Engine/Source/Runtime/SlateCore/Private/Application/SlateApplicationBase.cpp
@@ -121,33 +121,31 @@ void FSlateApplicationBase::UpdateAnyActiveTimersArePending()
 {
 	FScopeLock ActiveTimerLock(&ActiveTimerCS);
 
-	// first remove any tick handles that may have become invalid.
-	// If we didn't remove invalid handles here, they would never get removed because
-	// we don't force widgets to UnRegister before they are destroyed.
+	const double CurrentTime = GetCurrentTime();
+	bool bAnyTickReady = false;
 	for (auto It = ActiveTimerHandles.CreateIterator(); It; ++It)
 	{
-		if(!(*It).IsValid())
+		auto ActiveTimerInfoPinned = (*It).Pin();
+
+		// remove any tick handles that may have become invalid.
+		// If we didn't remove invalid handles here, they would never get removed because
+		// we don't force widgets to UnRegister before they are destroyed.
+		if(!ActiveTimerInfoPinned.IsValid())
 		{
 			It.RemoveCurrent();
+			continue;
 		}
-	}
 
-	// The rest are valid. Update their pending status and see if any are ready.
-	const double CurrentTime = GetCurrentTime();
-	bool bAnyTickReady = false;
-	for ( auto& ActiveTimerInfo : ActiveTimerHandles )
-	{
-		auto ActiveTimerInfoPinned = ActiveTimerInfo.Pin();
-		check( ActiveTimerInfoPinned.IsValid() );
+		check(ActiveTimerInfoPinned.IsValid());
 
 		// If an active timer is still pending execution from last frame, it is collapsed 
 		// or otherwise blocked from ticking. Disregard until it executes.
-		if ( ActiveTimerInfoPinned->IsPendingExecution() )
+		if (ActiveTimerInfoPinned->IsPendingExecution())
 		{
 			continue;
 		}
 
-		if ( ActiveTimerInfoPinned->UpdateExecutionPendingState( CurrentTime ) )
+		if (ActiveTimerInfoPinned->UpdateExecutionPendingState(CurrentTime))
 		{
 			bAnyTickReady = true;
 		}
