PR #14225: [Font][Tracing] Add cpu profile event to show how long invalidation takes after async loading a font

#tests Ran with it for a while

#Synced-CL 52573034
#Synced-CL 52573034

diff --git a/Engine/Source/Runtime/SlateCore/Private/Fonts/FontCacheCompositeFont.cpp b/Engine/Source/Runtime/SlateCore/Private/Fonts/FontCacheCompositeFont.cpp
--- a/Engine/Source/Runtime/SlateCore/Private/Fonts/FontCacheCompositeFont.cpp
+++ b/Engine/Source/Runtime/SlateCore/Private/Fonts/FontCacheCompositeFont.cpp
@@ -873,6 +873,7 @@ void FCompositeFontCache::Update()
 			AsyncTask(ENamedThreads::GameThread, [NumLoadFontFaceTasks, NumCompletedFontFaceTasks, CurrentThreadString]() {
 				if (FSlateApplicationBase::IsInitialized())
 				{
+					TRACE_CPUPROFILER_EVENT_SCOPE(FCompositFontCache_InvalidateAfterAsyncLoad);
 					UE_LOGF(LogSlate, Display, "Invalidating all widgets from %d/%d async loading font faces completing loading on %ls thread.", NumCompletedFontFaceTasks, NumLoadFontFaceTasks, CurrentThreadString);
 					FSlateApplicationBase::Get().InvalidateAllWidgets(false);
 					GSlateLayoutGeneration++;
