Editor: Promote Slate.bAllowThrottling to Editor Preferences

* Expose Slate.bAllowThrottling as a checkbox in Editor Preferences > Performance > Editor Performance so artists can discover and toggle it without using the console.
* Moved console variable outside of singleton (now static) to prevent construction dependency issues with the cvar not being found.

#was-backout 52691016

#Synced-CL 52423541

diff --git a/Engine/Source/Editor/UnrealEd/Classes/Editor/EditorPerformanceSettings.h b/Engine/Source/Editor/UnrealEd/Classes/Editor/EditorPerformanceSettings.h
--- a/Engine/Source/Editor/UnrealEd/Classes/Editor/EditorPerformanceSettings.h
+++ b/Engine/Source/Editor/UnrealEd/Classes/Editor/EditorPerformanceSettings.h
@@ -73,6 +73,10 @@ class UEditorPerformanceSettings : public UDeveloperSettings
 	UPROPERTY(EditAnywhere, config, Category=EditorPerformance, meta=(DisplayName="Use Less CPU when in Background") )
 	uint32 bThrottleCPUWhenNotForeground:1;
 
+	/** When enabled, the editor throttles background processing (world ticking, viewport updates) during UI interactions to keep menus and panels responsive. Disable this to keep viewports and simulations running at full speed while interacting with the editor. */
+	UPROPERTY(EditAnywhere, config, Category=EditorPerformance, meta=(DisplayName="Allow Editor (Slate) Throttling", ConsoleVariable="Slate.bAllowThrottling"))
+	uint32 bAllowSlateThrottling:1;
+
 	/** Disables realtime viewports by default when connected via a remote session */
 	UPROPERTY(EditAnywhere, config, Category = EditorPerformance, meta = (DisplayName = "Disable realtime viewports by default in Remote Sessions"))
 	uint32 bDisableRealtimeViewportsInRemoteSessions : 1;
diff --git a/Engine/Source/Editor/UnrealEd/Private/EditorPerformanceSettings.cpp b/Engine/Source/Editor/UnrealEd/Private/EditorPerformanceSettings.cpp
--- a/Engine/Source/Editor/UnrealEd/Private/EditorPerformanceSettings.cpp
+++ b/Engine/Source/Editor/UnrealEd/Private/EditorPerformanceSettings.cpp
@@ -19,6 +19,7 @@ UEditorPerformanceSettings::UEditorPerformanceSettings(const FObjectInitializer&
 	: Super(ObjectInitializer)
 	, bShowFrameRateAndMemory(false)
 	, bThrottleCPUWhenNotForeground(true)
+	, bAllowSlateThrottling(true)
 	, bDisableRealtimeViewportsInRemoteSessions(true)
 	, bMonitorEditorPerformance(false)
 	, bEnableScalabilityWarningIndicator(true)
@@ -51,6 +52,11 @@ void UEditorPerformanceSettings::PostInitProperties()
 		ExportValuesToConsoleVariables(EnableVSyncProperty);
 	}
 
+	if (FProperty* AllowSlateThrottlingProperty = GetClass()->FindPropertyByName(GET_MEMBER_NAME_CHECKED(UEditorPerformanceSettings, bAllowSlateThrottling)))
+	{
+		ExportValuesToConsoleVariables(AllowSlateThrottlingProperty);
+	}
+
 	UEditorPerformanceProjectSettings::ExportResolutionValuesToConsoleVariables();
 }
 
diff --git a/Engine/Source/Runtime/SlateCore/Private/Application/ThrottleManager.cpp b/Engine/Source/Runtime/SlateCore/Private/Application/ThrottleManager.cpp
--- a/Engine/Source/Runtime/SlateCore/Private/Application/ThrottleManager.cpp
+++ b/Engine/Source/Runtime/SlateCore/Private/Application/ThrottleManager.cpp
@@ -3,13 +3,18 @@
 #include "Application/ThrottleManager.h"
 
 
+static int32 GSlateAllowThrottling = 1;
+static FAutoConsoleVariableRef CVarSlateAllowThrottling(
+	TEXT("Slate.bAllowThrottling"),
+	GSlateAllowThrottling,
+	TEXT("Allow Slate to throttle parts of the engine to ensure the UI is responsive"));
+
+
 /* FSlateThrottleManager structors
  *****************************************************************************/
 
 FSlateThrottleManager::FSlateThrottleManager( )
-	: bShouldThrottle(1)
-	, CVarAllowThrottle(TEXT("Slate.bAllowThrottling"), bShouldThrottle, TEXT("Allow Slate to throttle parts of the engine to ensure the UI is responsive") )
-	, ThrottleCount(0)
+	: ThrottleCount(0)
 	, DisableThrottleCount(0)
 { }
 
@@ -34,7 +39,7 @@ bool FSlateThrottleManager::IsAllowingExpensiveTasks( ) const
 {
 	// Expensive tasks are allowed if the number of active throttle requests is zero
 	// or we always always allow it due to a CVar
-	return (ThrottleCount == 0) || !bShouldThrottle || DisableThrottleCount > 0;
+	return (ThrottleCount == 0) || !GSlateAllowThrottling || DisableThrottleCount > 0;
 }
 
 
diff --git a/Engine/Source/Runtime/SlateCore/Public/Application/ThrottleManager.h b/Engine/Source/Runtime/SlateCore/Public/Application/ThrottleManager.h
--- a/Engine/Source/Runtime/SlateCore/Public/Application/ThrottleManager.h
+++ b/Engine/Source/Runtime/SlateCore/Public/Application/ThrottleManager.h
@@ -92,12 +92,6 @@ public:
 	static SLATECORE_API FSlateThrottleManager& Get( );
 
 private:
-	/** CVar variable to check if we allow throttling (int32 for compatibility) */
-	int32 bShouldThrottle;
-
-	/** CVar allowing us to toggle throttling ability */
-	FAutoConsoleVariableRef CVarAllowThrottle;
-
 	/** Number of active throttle requests */
 	uint32 ThrottleCount;
 
